Fix up compile issue in GeometryCollection plugin related to skin weight profile system changes

diff --git a/Engine/Plugins/Experimental/GeometryCollectionPlugin/Source/GeometryCollectionEditor/Private/GeometryCollection/GeometryCollectionConversion.cpp b/Engine/Plugins/Experimental/GeometryCollectionPlugin/Source/GeometryCollectionEditor/Private/GeometryCollection/GeometryCollectionConversion.cpp
--- a/Engine/Plugins/Experimental/GeometryCollectionPlugin/Source/GeometryCollectionEditor/Private/GeometryCollection/GeometryCollectionConversion.cpp
+++ b/Engine/Plugins/Experimental/GeometryCollectionPlugin/Source/GeometryCollectionEditor/Private/GeometryCollection/GeometryCollectionConversion.cpp
@@ -238,7 +238,7 @@ void FGeometryCollectionConversion::AppendSkeletalMesh(const USkeletalMesh* Skel
 					if (SkelMeshRenderData->LODRenderData.Num())
 					{
 						const FSkeletalMeshLODRenderData & SkeletalMeshLODRenderData = SkelMeshRenderData->LODRenderData[0];
-						const FSkinWeightVertexBuffer & SkinWeightVertexBuffer = *SkeletalMeshLODRenderData.GetSkinWeightBuffer();
+						const FSkinWeightVertexBuffer & SkinWeightVertexBuffer = *SkeletalMeshLODRenderData.GetSkinWeightVertexBuffer();
 
 						const FSkelMeshRenderSection & RenderSection = SkeletalMeshLODRenderData.RenderSections[0];
 						const TArray<FBoneIndexType> & SkeletalBoneMap = RenderSection.BoneMap;
