Fix monolithic builds.

diff --git a/Engine/Plugins/Experimental/ChaosSolverPlugin/Source/ChaosSolverEditor/Private/Chaos/ChaosSolverEditorStyle.h b/Engine/Plugins/Experimental/ChaosSolverPlugin/Source/ChaosSolverEditor/Private/Chaos/ChaosSolverEditorStyle.h
--- a/Engine/Plugins/Experimental/ChaosSolverPlugin/Source/ChaosSolverEditor/Private/Chaos/ChaosSolverEditorStyle.h
+++ b/Engine/Plugins/Experimental/ChaosSolverPlugin/Source/ChaosSolverEditor/Private/Chaos/ChaosSolverEditorStyle.h
@@ -12,8 +12,10 @@ public:
 		const FVector2D Icon16x16(16.f, 16.f);
 		const FVector2D Icon64x64(64.f, 64.f);
 
+#if !IS_MONOLITHIC
 		FString PluginBasePath = FPaths::GetPath(FModuleManager::Get().GetModuleFilename("ChaosSolverEditor"));
 		SetContentRoot(PluginBasePath / TEXT("../../Resources"));
+#endif
 
 		Set("ClassIcon.ChaosSolver", new FSlateImageBrush(RootToContentDir(TEXT("ChaosSolver_16x.png")), Icon16x16));
 		Set("ClassThumbnail.ChaosSolver", new FSlateImageBrush(RootToContentDir(TEXT("ChaosSolver_64x.png")), Icon64x64));
diff --git a/Engine/Plugins/Experimental/FieldSystemPlugin/Source/FieldSyStemEditor/Private/Field/FieldSystemEditorStyle.h b/Engine/Plugins/Experimental/FieldSystemPlugin/Source/FieldSyStemEditor/Private/Field/FieldSystemEditorStyle.h
--- a/Engine/Plugins/Experimental/FieldSystemPlugin/Source/FieldSyStemEditor/Private/Field/FieldSystemEditorStyle.h
+++ b/Engine/Plugins/Experimental/FieldSystemPlugin/Source/FieldSyStemEditor/Private/Field/FieldSystemEditorStyle.h
@@ -12,8 +12,10 @@ public:
 		const FVector2D Icon16x16(16.f, 16.f);
 		const FVector2D Icon64x64(64.f, 64.f);
 
+#if !IS_MONOLITHIC
 		FString PluginBasePath = FPaths::GetPath(FModuleManager::Get().GetModuleFilename("FieldSystemEditor"));
 		SetContentRoot(PluginBasePath / TEXT("../../Resources"));
+#endif
 
 		Set("ClassIcon.FieldSystem", new FSlateImageBrush(RootToContentDir(TEXT("FieldSystem_16x.png")), Icon16x16));
 		Set("ClassThumbnail.FieldSystem", new FSlateImageBrush(RootToContentDir(TEXT("FieldSystem_64x.png")), Icon64x64));
diff --git a/Engine/Plugins/Experimental/GeometryCollectionPlugin/Source/GeometryCollectionEditor/Private/GeometryCollection/GeometryCollectionEditorStyle.h b/Engine/Plugins/Experimental/GeometryCollectionPlugin/Source/GeometryCollectionEditor/Private/GeometryCollection/GeometryCollectionEditorStyle.h
--- a/Engine/Plugins/Experimental/GeometryCollectionPlugin/Source/GeometryCollectionEditor/Private/GeometryCollection/GeometryCollectionEditorStyle.h
+++ b/Engine/Plugins/Experimental/GeometryCollectionPlugin/Source/GeometryCollectionEditor/Private/GeometryCollection/GeometryCollectionEditorStyle.h
@@ -12,8 +12,10 @@ public:
 		const FVector2D Icon16x16(16.f, 16.f);
 		const FVector2D Icon64x64(64.f, 64.f);
 
+#if !IS_MONOLITHIC
 		FString PluginBasePath = FPaths::GetPath(FModuleManager::Get().GetModuleFilename("GeometryCollectionEditor"));
 		SetContentRoot(PluginBasePath / TEXT("../../Resources"));
+#endif
 
 		Set("ClassIcon.GeometryCollection", new FSlateImageBrush(RootToContentDir(TEXT("GeometryCollection_16x.png")), Icon16x16));
 		Set("ClassThumbnail.GeometryCollection", new FSlateImageBrush(RootToContentDir(TEXT("GeometryCollection_64x.png")), Icon64x64));
