Fix Geometry Collection debug draw picking tool.
Merge from Dev-Physics:
CL 6910611
CL 6910614
CL 6911200

diff --git a/Engine/Plugins/Experimental/GeometryCollectionPlugin/Source/GeometryCollectionEditor/Private/GeometryCollection/DetailCustomizations/SelectedRigidBodyCustomization.cpp b/Engine/Plugins/Experimental/GeometryCollectionPlugin/Source/GeometryCollectionEditor/Private/GeometryCollection/DetailCustomizations/SelectedRigidBodyCustomization.cpp
--- a/Engine/Plugins/Experimental/GeometryCollectionPlugin/Source/GeometryCollectionEditor/Private/GeometryCollection/DetailCustomizations/SelectedRigidBodyCustomization.cpp
+++ b/Engine/Plugins/Experimental/GeometryCollectionPlugin/Source/GeometryCollectionEditor/Private/GeometryCollection/DetailCustomizations/SelectedRigidBodyCustomization.cpp
@@ -12,6 +12,8 @@
 #include "EditorFontGlyphs.h"
 #include "GeometryCollection/GeometryCollectionDebugDrawActor.h"
 #include "GeometryCollection/GeometryCollectionSelectRigidBodyEdMode.h"
+#include "GeometryCollection/GeometryCollectionComponent.h"
+#include "GeometryCollection/GeometryCollectionActor.h"
 
 TSharedRef<IPropertyTypeCustomization> FSelectedRigidBodyCustomization::MakeInstance()
 {
@@ -20,6 +22,7 @@ TSharedRef<IPropertyTypeCustomization> FSelectedRigidBodyCustomization::MakeInst
 
 void FSelectedRigidBodyCustomization::CustomizeHeader(TSharedRef<IPropertyHandle> StructPropertyHandle, FDetailWidgetRow& HeaderRow, IPropertyTypeCustomizationUtils& /*StructCustomizationUtils*/)
 {
+	const TSharedRef<IPropertyHandle> PropertyHandleActor = StructPropertyHandle->GetChildHandle(GET_MEMBER_NAME_CHECKED(FGeometryCollectionDebugDrawActorSelectedRigidBody, GeometryCollectionActor)).ToSharedRef();
 	const TSharedRef<IPropertyHandle> PropertyHandleId = StructPropertyHandle->GetChildHandle(GET_MEMBER_NAME_CHECKED(FGeometryCollectionDebugDrawActorSelectedRigidBody, Id)).ToSharedRef();
 
 	// Add buttons
@@ -30,23 +33,21 @@ void FSelectedRigidBodyCustomization::CustomizeHeader(TSharedRef<IPropertyHandle
 				NSLOCTEXT("ChaosSelectedRigidBody", "ChaosSelectedRigidBody_ToolTip", "Select a Rigid Body here by either entering its Id, or clicking on the Pick button."))
 		]
 	.ValueContent()
-		.MinDesiredWidth(125.f)
+		.MinDesiredWidth(360.f)
 		[
 			SNew(SHorizontalBox)
 			+ SHorizontalBox::Slot()
-			.FillWidth(80.0f)
 			[
 				PropertyHandleId->CreatePropertyValueWidget()
 			]
 			+ SHorizontalBox::Slot()
-			.Padding(4.0f)
-			.FillWidth(20.0f)
+			.AutoWidth()
 			[
 				SNew(SButton)
 				.ButtonStyle(FEditorStyle::Get(), "HoverHintOnly")
 				.ToolTipText(NSLOCTEXT("ChaosSelectedRigidBody", "Pick_ToolTip", "Pick a Rigid Body."))
 				.ForegroundColor(FSlateColor::UseForeground())
-				.ContentPadding(4.0f)
+				.ContentPadding(6.0f)
 				.OnClicked(this, &FSelectedRigidBodyCustomization::OnPick, PropertyHandleId)
 				[ 
 					SNew(SImage)
@@ -54,6 +55,116 @@ void FSelectedRigidBodyCustomization::CustomizeHeader(TSharedRef<IPropertyHandle
 					.ColorAndOpacity(FSlateColor::UseForeground())
 				]
 			]
+			+ SHorizontalBox::Slot()
+			.AutoWidth()
+			.Padding(8.f, 0.f, 0.f, 0.f)
+			[
+				SNew(SButton)
+				.ButtonStyle(FEditorStyle::Get(), "HoverHintOnly")
+				.ToolTipText(NSLOCTEXT("ChaosSelectedRigidBody", "Pick_UpClusterLevel", "Go to parent cluster."))
+				.ForegroundColor(FSlateColor::UseForeground())
+				.IsEnabled_Lambda([PropertyHandleActor, PropertyHandleId]()
+				{
+					return GetParentClusterRigidBodyId(PropertyHandleActor, PropertyHandleId) != INDEX_NONE;
+				})
+				.OnClicked_Lambda([PropertyHandleActor, PropertyHandleId]()
+				{
+					const int32 RigidBodyId = GetParentClusterRigidBodyId(PropertyHandleActor, PropertyHandleId);
+					if (RigidBodyId != INDEX_NONE)
+					{
+						PropertyHandleId->SetValue(RigidBodyId);
+					}
+					return FReply::Handled();
+				})
+				[ 
+					SNew(STextBlock)
+					.Font(FEditorStyle::Get().GetFontStyle("FontAwesome.16"))
+					.Text(FEditorFontGlyphs::Caret_Square_O_Up)
+					.ColorAndOpacity(FSlateColor::UseForeground())
+				]
+			]
+			+ SHorizontalBox::Slot()
+			.AutoWidth()
+			[
+				SNew(SButton)
+				.ButtonStyle(FEditorStyle::Get(), "HoverHintOnly")
+				.ToolTipText(NSLOCTEXT("ChaosSelectedRigidBody", "Pick_DownClusterLevel", "Go to child cluster."))
+				.ForegroundColor(FSlateColor::UseForeground())
+				.IsEnabled_Lambda([PropertyHandleActor, PropertyHandleId]()
+				{
+					return GetChildClusterRigidBodyId(PropertyHandleActor, PropertyHandleId) != INDEX_NONE;
+				})
+				.OnClicked_Lambda([PropertyHandleActor, PropertyHandleId]()
+				{
+					const int32 RigidBodyId = GetChildClusterRigidBodyId(PropertyHandleActor, PropertyHandleId);
+					if (RigidBodyId != INDEX_NONE)
+					{
+						PropertyHandleId->SetValue(RigidBodyId);
+					}
+					return FReply::Handled();
+				})
+				[ 
+					SNew(STextBlock)
+					.Font(FEditorStyle::Get().GetFontStyle("FontAwesome.16"))
+					.Text(FEditorFontGlyphs::Caret_Square_O_Down)
+					.ColorAndOpacity(FSlateColor::UseForeground())
+				]
+			]
+			+ SHorizontalBox::Slot()
+			.AutoWidth()
+			.Padding(8.0f, 0.f, 0.f, 0.f)
+			[
+				SNew(SButton)
+				.ButtonStyle(FEditorStyle::Get(), "HoverHintOnly")
+				.ToolTipText(NSLOCTEXT("ChaosSelectedRigidBody", "Pick_PrevClusterSibling", "Go to previous clustered sibling."))
+				.ForegroundColor(FSlateColor::UseForeground())
+				.IsEnabled_Lambda([PropertyHandleActor, PropertyHandleId]()
+				{
+					return GetPreviousClusteredSiblingRigidBodyId(PropertyHandleActor, PropertyHandleId) != INDEX_NONE;
+				})
+				.OnClicked_Lambda([PropertyHandleActor, PropertyHandleId]()
+				{
+					const int32 RigidBodyId = GetPreviousClusteredSiblingRigidBodyId(PropertyHandleActor, PropertyHandleId);
+					if (RigidBodyId != INDEX_NONE)
+					{
+						PropertyHandleId->SetValue(RigidBodyId);
+					}
+					return FReply::Handled();
+				})
+				[ 
+					SNew(STextBlock)
+					.Font(FEditorStyle::Get().GetFontStyle("FontAwesome.16"))
+					.Text(FEditorFontGlyphs::Caret_Square_O_Left)
+					.ColorAndOpacity(FSlateColor::UseForeground())
+				]
+			]
+			+ SHorizontalBox::Slot()
+			.AutoWidth()
+			[
+				SNew(SButton)
+				.ButtonStyle(FEditorStyle::Get(), "HoverHintOnly")
+				.ToolTipText(NSLOCTEXT("ChaosSelectedRigidBody", "Pick_NextClusterSibling", "Go to next clustered sibling."))
+				.ForegroundColor(FSlateColor::UseForeground())
+				.IsEnabled_Lambda([PropertyHandleActor, PropertyHandleId]()
+				{
+					return GetNextClusteredSiblingRigidBodyId(PropertyHandleActor, PropertyHandleId) != INDEX_NONE;
+				})
+				.OnClicked_Lambda([PropertyHandleActor, PropertyHandleId]()
+				{
+					const int32 RigidBodyId = GetNextClusteredSiblingRigidBodyId(PropertyHandleActor, PropertyHandleId);
+					if (RigidBodyId != INDEX_NONE)
+					{
+						PropertyHandleId->SetValue(RigidBodyId);
+					}
+					return FReply::Handled();
+				})
+				[ 
+					SNew(STextBlock)
+					.Font(FEditorStyle::Get().GetFontStyle("FontAwesome.16"))
+					.Text(FEditorFontGlyphs::Caret_Square_O_Right)
+					.ColorAndOpacity(FSlateColor::UseForeground())
+				]
+			]
 		];
 }
 
@@ -61,6 +172,123 @@ void FSelectedRigidBodyCustomization::CustomizeChildren(TSharedRef<IPropertyHand
 {
 }
 
+void FSelectedRigidBodyCustomization::GetSelectedGeometryCollectionCluster(TSharedRef<IPropertyHandle> PropertyHandleActor, TSharedRef<IPropertyHandle> PropertyHandleId, const UGeometryCollectionComponent*& OutGeometryCollectionComponent, int32& OutTransformIndex)
+{
+	OutGeometryCollectionComponent = nullptr;
+	OutTransformIndex = INDEX_NONE;
+
+	// Retrieve selected rigid body
+	int32 SelectedRigidBodyId;
+	if (PropertyHandleId->GetValue(SelectedRigidBodyId) == FPropertyAccess::Success && SelectedRigidBodyId != INDEX_NONE)
+	{
+		// Retrieve selected Geometry Collection Actor
+		UObject* Object;
+		if (PropertyHandleActor->GetValue(Object) == FPropertyAccess::Success && Object != nullptr)
+		{
+			// Retrieve selected tranform index
+			const UGeometryCollectionComponent* const GeometryCollectionComponent = CastChecked<AGeometryCollectionActor>(Object)->GetGeometryCollectionComponent();
+			const TManagedArray<int32>& RigidBodyIds = GeometryCollectionComponent->GetRigidBodyIdArray();
+
+			for (int32 TransformIndex = 0; TransformIndex < RigidBodyIds.Num(); ++TransformIndex)
+			{
+				const int32 RigidBodyId = RigidBodyIds[TransformIndex];
+				if (RigidBodyId == SelectedRigidBodyId)
+				{
+					OutGeometryCollectionComponent = GeometryCollectionComponent;
+					OutTransformIndex = TransformIndex;
+					break;
+				}
+			}
+		}
+	}
+}
+
+int32 FSelectedRigidBodyCustomization::GetParentClusterRigidBodyId(TSharedRef<IPropertyHandle> PropertyHandleActor, TSharedRef<IPropertyHandle> PropertyHandleId)
+{
+	const UGeometryCollectionComponent* GeometryCollectionComponent;
+	int32 TransformIndex;
+	GetSelectedGeometryCollectionCluster(PropertyHandleActor, PropertyHandleId, GeometryCollectionComponent, TransformIndex);
+	if (GeometryCollectionComponent)
+	{
+		const int32 ParentTransformIndex = GeometryCollectionComponent->GetParentArray()[TransformIndex];
+		if (ParentTransformIndex != INDEX_NONE)
+		{
+			return GeometryCollectionComponent->GetRigidBodyIdArray()[ParentTransformIndex];
+		}
+	}
+	return INDEX_NONE;
+}
+
+int32 FSelectedRigidBodyCustomization::GetChildClusterRigidBodyId(TSharedRef<IPropertyHandle> PropertyHandleActor, TSharedRef<IPropertyHandle> PropertyHandleId)
+{
+	const UGeometryCollectionComponent* GeometryCollectionComponent;
+	int32 TransformIndex;
+	GetSelectedGeometryCollectionCluster(PropertyHandleActor, PropertyHandleId, GeometryCollectionComponent, TransformIndex);
+	if (GeometryCollectionComponent)
+	{
+		const TSet<int32>& Children = GeometryCollectionComponent->GetChildrenArray()[TransformIndex];
+		if (TSet<int32>::TConstIterator ChildTransformIndexIt = Children.CreateConstIterator())
+		{
+			return GeometryCollectionComponent->GetRigidBodyIdArray()[*ChildTransformIndexIt];
+		}
+	}
+	return INDEX_NONE;
+}
+
+int32 FSelectedRigidBodyCustomization::GetPreviousClusteredSiblingRigidBodyId(TSharedRef<IPropertyHandle> PropertyHandleActor, TSharedRef<IPropertyHandle> PropertyHandleId)
+{
+	const UGeometryCollectionComponent* GeometryCollectionComponent;
+	int32 TransformIndex;
+	GetSelectedGeometryCollectionCluster(PropertyHandleActor, PropertyHandleId, GeometryCollectionComponent, TransformIndex);
+	if (GeometryCollectionComponent)
+	{
+		const int32 ParentTransformIndex = GeometryCollectionComponent->GetParentArray()[TransformIndex];
+		if (ParentTransformIndex != INDEX_NONE)
+		{
+			const TSet<int32>& Siblings = GeometryCollectionComponent->GetChildrenArray()[ParentTransformIndex];
+			int32 PrevSibling = INDEX_NONE;
+			for (int32 Sibling : Siblings)
+			{
+				if (Sibling == TransformIndex)
+				{
+					if (PrevSibling != INDEX_NONE)
+					{
+						return GeometryCollectionComponent->GetRigidBodyIdArray()[PrevSibling];
+					}
+					break;
+				}
+				PrevSibling = Sibling;
+			}
+		}
+	}
+	return INDEX_NONE;
+}
+
+int32 FSelectedRigidBodyCustomization::GetNextClusteredSiblingRigidBodyId(TSharedRef<IPropertyHandle> PropertyHandleActor, TSharedRef<IPropertyHandle> PropertyHandleId)
+{
+	const UGeometryCollectionComponent* GeometryCollectionComponent;
+	int32 TransformIndex;
+	GetSelectedGeometryCollectionCluster(PropertyHandleActor, PropertyHandleId, GeometryCollectionComponent, TransformIndex);
+	if (GeometryCollectionComponent)
+	{
+		const int32 ParentTransformIndex = GeometryCollectionComponent->GetParentArray()[TransformIndex];
+		if (ParentTransformIndex != INDEX_NONE)
+		{
+			const TSet<int32>& Siblings = GeometryCollectionComponent->GetChildrenArray()[ParentTransformIndex];
+			int32 Sibling = INDEX_NONE;
+			for (int32 NextSibling : Siblings)
+			{
+				if (Sibling == TransformIndex)
+				{
+					return GeometryCollectionComponent->GetRigidBodyIdArray()[NextSibling];
+				}
+				Sibling = NextSibling;
+			}
+		}
+	}
+	return INDEX_NONE;
+}
+
 FReply FSelectedRigidBodyCustomization::OnPick(TSharedRef<IPropertyHandle> PropertyHandleId) const
 {
 	if (FGeometryCollectionSelectRigidBodyEdMode::IsModeActive())
diff --git a/Engine/Plugins/Experimental/GeometryCollectionPlugin/Source/GeometryCollectionEditor/Private/GeometryCollection/DetailCustomizations/SelectedRigidBodyCustomization.h b/Engine/Plugins/Experimental/GeometryCollectionPlugin/Source/GeometryCollectionEditor/Private/GeometryCollection/DetailCustomizations/SelectedRigidBodyCustomization.h
--- a/Engine/Plugins/Experimental/GeometryCollectionPlugin/Source/GeometryCollectionEditor/Private/GeometryCollection/DetailCustomizations/SelectedRigidBodyCustomization.h
+++ b/Engine/Plugins/Experimental/GeometryCollectionPlugin/Source/GeometryCollectionEditor/Private/GeometryCollection/DetailCustomizations/SelectedRigidBodyCustomization.h
@@ -5,6 +5,7 @@
 #include "Input/Reply.h"
 
 class IDetailLayoutBuilder;
+class UGeometryCollectionComponent;
 
 /** Details customization for rigid body selection. */
 class FSelectedRigidBodyCustomization : public IPropertyTypeCustomization
@@ -16,5 +17,12 @@ public:
 	virtual void CustomizeChildren(TSharedRef<IPropertyHandle> StructPropertyHandle, class IDetailChildrenBuilder& ChildBuilder, IPropertyTypeCustomizationUtils& StructCustomizationUtils) override;
 
 private:
+	static void GetSelectedGeometryCollectionCluster(TSharedRef<IPropertyHandle> PropertyHandleActor, TSharedRef<IPropertyHandle> PropertyHandleId, const UGeometryCollectionComponent*& OutGeometryCollectionComponent, int32& OutTransformIndex);
+
+	static int32 GetParentClusterRigidBodyId(TSharedRef<IPropertyHandle> PropertyHandleActor, TSharedRef<IPropertyHandle> PropertyHandleId);
+	static int32 GetChildClusterRigidBodyId(TSharedRef<IPropertyHandle> PropertyHandleActor, TSharedRef<IPropertyHandle> PropertyHandleId);
+	static int32 GetPreviousClusteredSiblingRigidBodyId(TSharedRef<IPropertyHandle> PropertyHandleActor, TSharedRef<IPropertyHandle> PropertyHandleId);
+	static int32 GetNextClusteredSiblingRigidBodyId(TSharedRef<IPropertyHandle> PropertyHandleActor, TSharedRef<IPropertyHandle> PropertyHandleId);
+
 	FReply OnPick(TSharedRef<IPropertyHandle> PropertyHandleId) const;
 };
diff --git a/Engine/Plugins/Experimental/GeometryCollectionPlugin/Source/GeometryCollectionEditor/Private/GeometryCollection/GeometryCollectionSelectRigidBodyEdMode.cpp b/Engine/Plugins/Experimental/GeometryCollectionPlugin/Source/GeometryCollectionEditor/Private/GeometryCollection/GeometryCollectionSelectRigidBodyEdMode.cpp
--- a/Engine/Plugins/Experimental/GeometryCollectionPlugin/Source/GeometryCollectionEditor/Private/GeometryCollection/GeometryCollectionSelectRigidBodyEdMode.cpp
+++ b/Engine/Plugins/Experimental/GeometryCollectionPlugin/Source/GeometryCollectionEditor/Private/GeometryCollection/GeometryCollectionSelectRigidBodyEdMode.cpp
@@ -91,21 +91,23 @@ void FGeometryCollectionSelectRigidBodyEdMode::DeactivateMode()
 	}
 }
 
-bool FGeometryCollectionSelectRigidBodyEdMode::MouseMove(FEditorViewportClient* ViewportClient, FViewport* Viewport, int32 x, int32 y)
+void FGeometryCollectionSelectRigidBodyEdMode::Tick(FEditorViewportClient* /*ViewportClient*/, float /*DeltaTime*/)
 {
-	if (!GEditor || !GEditor->PlayWorld)
+	if (!GEditor || !GEditor->PlayWorld || !GEditor->bIsSimulatingInEditor)
 	{
-		// Left play world
+		// Left simulating in editor
 		bIsHoveringGeometryCollection = false;
 		DeactivateMode();
 	}
+}
+
+bool FGeometryCollectionSelectRigidBodyEdMode::MouseMove(FEditorViewportClient* /*ViewportClient*/, FViewport* Viewport, int32 x, int32 y)
+{
 #if GEOMETRYCOLLECTION_EDITOR_SELECTION
-	else
-	{
-		// Hover status
-		const HHitProxy* const HitProxy = Viewport->GetHitProxy(x, y);
-		bIsHoveringGeometryCollection = HitProxy && HitProxy->IsA(HGeometryCollection::StaticGetType());
-	}
+	// Hover status
+	const HHitProxy* const HitProxy = Viewport->GetHitProxy(x, y);
+	bIsHoveringGeometryCollection = HitProxy && HitProxy->IsA(HGeometryCollection::StaticGetType());
+	UE_CLOG(bIsHoveringGeometryCollection, LogGeometryCollectionSelectRigidBodyEdMode, VeryVerbose, TEXT("Hovering GeometryCollectionActor %s, Transform Index %d."), *static_cast<const HGeometryCollection*>(HitProxy)->Actor->GetName(), static_cast<const HGeometryCollection*>(HitProxy)->TransformIndex);
 #endif // #if GEOMETRYCOLLECTION_EDITOR_SELECTION
 	return true;
 }
diff --git a/Engine/Plugins/Experimental/GeometryCollectionPlugin/Source/GeometryCollectionEditor/Private/GeometryCollection/GeometryCollectionSelectRigidBodyEdMode.h b/Engine/Plugins/Experimental/GeometryCollectionPlugin/Source/GeometryCollectionEditor/Private/GeometryCollection/GeometryCollectionSelectRigidBodyEdMode.h
--- a/Engine/Plugins/Experimental/GeometryCollectionPlugin/Source/GeometryCollectionEditor/Private/GeometryCollection/GeometryCollectionSelectRigidBodyEdMode.h
+++ b/Engine/Plugins/Experimental/GeometryCollectionPlugin/Source/GeometryCollectionEditor/Private/GeometryCollection/GeometryCollectionSelectRigidBodyEdMode.h
@@ -35,6 +35,8 @@ public:
 	virtual bool UsesTransformWidget() const override { return false; }
 	virtual bool UsesTransformWidget(FWidget::EWidgetMode /*CheckMode*/) const override { return false; }
 
+	virtual void Tick(FEditorViewportClient* ViewportClient, float DeltaTime) override;
+
 	virtual bool HandleClick(FEditorViewportClient* InViewportClient, HHitProxy* HitProxy, const FViewportClick& Click) override;
 	virtual bool MouseMove(FEditorViewportClient* ViewportClient, FViewport* Viewport, int32 x, int32 y) override;
 
diff --git a/Engine/Source/Runtime/Experimental/GeometryCollectionEngine/Private/GeometryCollection/GeometryCollectionDebugDrawActor.cpp b/Engine/Source/Runtime/Experimental/GeometryCollectionEngine/Private/GeometryCollection/GeometryCollectionDebugDrawActor.cpp
--- a/Engine/Source/Runtime/Experimental/GeometryCollectionEngine/Private/GeometryCollection/GeometryCollectionDebugDrawActor.cpp
+++ b/Engine/Source/Runtime/Experimental/GeometryCollectionEngine/Private/GeometryCollection/GeometryCollectionDebugDrawActor.cpp
@@ -362,7 +362,11 @@ void AGeometryCollectionDebugDrawActor::OnPropertyChanged(bool bForceVisibilityU
 		if (GeometryCollectionDebugDrawComponent)
 		{
 			check(GeometryCollectionDebugDrawComponent->GeometryCollectionDebugDrawActor == this);
-			GeometryCollectionDebugDrawComponent->OnDebugDrawPropertiesChanged(bForceVisibilityUpdate);
+			const bool bIsSelected = GeometryCollectionDebugDrawComponent->OnDebugDrawPropertiesChanged(bForceVisibilityUpdate);
+			if (bIsSelected)
+			{
+				SelectedRigidBody.GeometryCollectionActor = *ActorIterator;
+			}
 		}
 	}
 #endif
diff --git a/Engine/Source/Runtime/Experimental/GeometryCollectionEngine/Private/GeometryCollection/GeometryCollectionDebugDrawComponent.cpp b/Engine/Source/Runtime/Experimental/GeometryCollectionEngine/Private/GeometryCollection/GeometryCollectionDebugDrawComponent.cpp
--- a/Engine/Source/Runtime/Experimental/GeometryCollectionEngine/Private/GeometryCollection/GeometryCollectionDebugDrawComponent.cpp
+++ b/Engine/Source/Runtime/Experimental/GeometryCollectionEngine/Private/GeometryCollection/GeometryCollectionDebugDrawComponent.cpp
@@ -258,7 +258,7 @@ bool UGeometryCollectionDebugDrawComponent::CanEditChange(const UProperty* InPro
 #endif  // #if WITH_EDITOR
 
 #if GEOMETRYCOLLECTION_DEBUG_DRAW
-void UGeometryCollectionDebugDrawComponent::OnDebugDrawPropertiesChanged(bool bForceVisibilityUpdate)
+bool UGeometryCollectionDebugDrawComponent::OnDebugDrawPropertiesChanged(bool bForceVisibilityUpdate)
 {
 	if (HasBegunPlay() && GeometryCollectionComponent)
 	{
@@ -278,6 +278,7 @@ void UGeometryCollectionDebugDrawComponent::OnDebugDrawPropertiesChanged(bool bF
 		// Update tick function
 		UpdateTickStatus();
 	}
+	return SelectedTransformIndex != INDEX_NONE;
 }
 
 void UGeometryCollectionDebugDrawComponent::OnClusterChanged()
diff --git a/Engine/Source/Runtime/Experimental/GeometryCollectionEngine/Private/GeometryCollection/GeometryCollectionSceneProxy.cpp b/Engine/Source/Runtime/Experimental/GeometryCollectionEngine/Private/GeometryCollection/GeometryCollectionSceneProxy.cpp
--- a/Engine/Source/Runtime/Experimental/GeometryCollectionEngine/Private/GeometryCollection/GeometryCollectionSceneProxy.cpp
+++ b/Engine/Source/Runtime/Experimental/GeometryCollectionEngine/Private/GeometryCollection/GeometryCollectionSceneProxy.cpp
@@ -613,18 +613,26 @@ void FGeometryCollectionSceneProxy::GetDynamicMeshElements(const TArray<const FS
 			// #todo(dmp): refactor this to share more code later
 			if (!DynamicData->IsDynamic)			
 			{
+#if GEOMETRYCOLLECTION_EDITOR_SELECTION
+				const TArray<FGeometryCollectionSection>& SectionArray = bUsesSubSections && SubSections.Num() ? SubSections: ConstantData->OriginalMeshSections;
+				UE_LOG(FGeometryCollectionSceneProxyLogging, VeryVerbose, TEXT("GetDynamicMeshElements, bUseSubSections=%d, NumSections=%d for %p."), bUsesSubSections, SectionArray.Num(), this);
+
+#else  // #if GEOMETRYCOLLECTION_EDITOR_SELECTION
+				const TArray<FGeometryCollectionSection>& SectionArray = ConstantData->OriginalMeshSections;
+#endif  // #if GEOMETRYCOLLECTION_EDITOR_SELECTION #else
+
 				//Grab the material proxies we'll be using for each section
 				TArray<FMaterialRenderProxy*, TInlineAllocator<32>> MaterialProxies;
-				for (int SectionIndex = 0; SectionIndex < ConstantData->OriginalMeshSections.Num(); SectionIndex++)
+				for (int SectionIndex = 0; SectionIndex < SectionArray.Num(); SectionIndex++)
 				{
-					const FGeometryCollectionSection& Section = ConstantData->OriginalMeshSections[SectionIndex];
+					const FGeometryCollectionSection& Section = SectionArray[SectionIndex];
 					FMaterialRenderProxy* MaterialProxy = GetMaterial(Collector, Section.MaterialID);
 					MaterialProxies.Add(MaterialProxy);
 				}
 
-				for (int SectionIndex = 0; SectionIndex < ConstantData->OriginalMeshSections.Num(); SectionIndex++)
+				for (int SectionIndex = 0; SectionIndex < SectionArray.Num(); SectionIndex++)
 				{
-					const FGeometryCollectionSection& Section = ConstantData->OriginalMeshSections[SectionIndex];
+					const FGeometryCollectionSection& Section = SectionArray[SectionIndex];
 
 					// Draw the mesh.
 					FMeshBatch& Mesh = Collector.AllocateMesh();
@@ -667,7 +675,7 @@ void FGeometryCollectionSceneProxy::GetDynamicMeshElements(const TArray<const FS
 			else
 			{
 #if GEOMETRYCOLLECTION_EDITOR_SELECTION
-				const TArray<FGeometryCollectionSection>& SectionArray = bUsesSubSections ? SubSections: Sections;
+				const TArray<FGeometryCollectionSection>& SectionArray = bUsesSubSections && SubSections.Num() ? SubSections: Sections;
 				UE_LOG(FGeometryCollectionSceneProxyLogging, VeryVerbose, TEXT("GetDynamicMeshElements, bUseSubSections=%d, NumSections=%d for %p."), bUsesSubSections, SectionArray.Num(), this);
 
 #else  // #if GEOMETRYCOLLECTION_EDITOR_SELECTION
@@ -874,18 +882,26 @@ void FGeometryCollectionSceneProxy::UseSubSections(bool bInUsesSubSections, bool
 
 void FGeometryCollectionSceneProxy::InitializeSubSections_RenderThread()
 {
-	// Reserve sub sections array with a minimum of one transform per section
-	SubSections.Empty(Sections.Num());
-	SubSectionHitProxyIndexMap.Empty(Sections.Num());
-
 	// Exit now if there isn't any data
-	if (!ConstantData) { return; }
+	if (!ConstantData)
+	{
+		SubSections.Empty();
+		SubSectionHitProxyIndexMap.Empty();
+		return;
+	}
 
+	// Retrieve the correct arrays depending on the dynamic state
+	const bool bIsDynamic = DynamicData && DynamicData->IsDynamic;
+	const TArray<FGeometryCollectionSection>& SectionArray = bIsDynamic ? Sections: ConstantData->OriginalMeshSections;
+	const TArray<FIntVector>& IndexArray = bIsDynamic ? ConstantData->Indices: ConstantData->OriginalMeshIndices;
 	const TArray<int32>& BoneMap = ConstantData->BoneMap;
-	const TArray<FIntVector>& Indices = ConstantData->Indices;
+
+	// Reserve sub sections array with a minimum of one transform per section
+	SubSections.Empty(SectionArray.Num());
+	SubSectionHitProxyIndexMap.Empty(SectionArray.Num());
 
 	// Lambda that adds a new subsection and update the HitProxy section index
-	auto AddSubSection = [this, Indices](int32 HitProxyIndex, const FGeometryCollectionSection& Section, int32 FirstFaceIndex, int32 EndFaceIndex)
+	auto AddSubSection = [this, IndexArray](int32 HitProxyIndex, const FGeometryCollectionSection& Section, int32 FirstFaceIndex, int32 EndFaceIndex)
 	{
 		// Find the matching HitProxy for this transform/section
 		HGeometryCollection* const SubSectionHitProxy = SubSectionHitProxies[HitProxyIndex];
@@ -902,8 +918,8 @@ void FGeometryCollectionSceneProxy::InitializeSubSections_RenderThread()
 			SubSection.MaxVertexIndex = TNumericLimits<int32>::Min();
 			for (int32 FaceIndex = FirstFaceIndex; FaceIndex < EndFaceIndex; ++FaceIndex)
 			{
-				SubSection.MinVertexIndex = FMath::Min(SubSection.MinVertexIndex, Indices[FaceIndex].GetMin());
-				SubSection.MaxVertexIndex = FMath::Max(SubSection.MaxVertexIndex, Indices[FaceIndex].GetMax());
+				SubSection.MinVertexIndex = FMath::Min(SubSection.MinVertexIndex, IndexArray[FaceIndex].GetMin());
+				SubSection.MaxVertexIndex = FMath::Max(SubSection.MaxVertexIndex, IndexArray[FaceIndex].GetMax());
 			}
 			check(SubSection.MinVertexIndex >= Section.MinVertexIndex && SubSection.MinVertexIndex <= Section.MaxVertexIndex)
 			check(SubSection.MaxVertexIndex >= Section.MinVertexIndex && SubSection.MaxVertexIndex <= Section.MaxVertexIndex)
@@ -922,24 +938,24 @@ void FGeometryCollectionSceneProxy::InitializeSubSections_RenderThread()
 	};
 
 	// Create subsections per transform
-	const int32 NumTransforms = (Sections.Num() > 0) ? SubSectionHitProxies.Num() / Sections.Num(): 0;
+	const int32 NumTransforms = (SectionArray.Num() > 0) ? SubSectionHitProxies.Num() / SectionArray.Num(): 0;
 
-	for (int32 SectionIndex = 0; SectionIndex < Sections.Num(); ++SectionIndex)
+	for (int32 SectionIndex = 0; SectionIndex < SectionArray.Num(); ++SectionIndex)
 	{
 		const int32 SectionOffset = SectionIndex * NumTransforms;
 
-		const FGeometryCollectionSection& Section = Sections[SectionIndex];
+		const FGeometryCollectionSection& Section = SectionArray[SectionIndex];
 		check(Section.NumTriangles > 0);  // Sections are not created with zero triangles
 
 		const int32 FirstFaceIndex = Section.FirstIndex / 3;
 		const int32 EndFaceIndex = FirstFaceIndex + Section.NumTriangles;
 
-		int32 TransformIndex = BoneMap[Indices[FirstFaceIndex][0]];  // Assumes one transform per triangle
+		int32 TransformIndex = BoneMap[IndexArray[FirstFaceIndex][0]];  // Assumes one transform per triangle
 		int32 FaceIndex = FirstFaceIndex;
 
 		for (int32 NextFaceIndex = FaceIndex + 1; NextFaceIndex < EndFaceIndex; ++NextFaceIndex)
 		{
-			const int32 NextTransformIndex = BoneMap[Indices[NextFaceIndex][0]];  // Assumes one transform per triangle
+			const int32 NextTransformIndex = BoneMap[IndexArray[NextFaceIndex][0]];  // Assumes one transform per triangle
 			if (TransformIndex != NextTransformIndex)
 			{
 				// Add the current subsection
diff --git a/Engine/Source/Runtime/Experimental/GeometryCollectionEngine/Public/GeometryCollection/GeometryCollectionDebugDrawActor.h b/Engine/Source/Runtime/Experimental/GeometryCollectionEngine/Public/GeometryCollection/GeometryCollectionDebugDrawActor.h
--- a/Engine/Source/Runtime/Experimental/GeometryCollectionEngine/Public/GeometryCollection/GeometryCollectionDebugDrawActor.h
+++ b/Engine/Source/Runtime/Experimental/GeometryCollectionEngine/Public/GeometryCollection/GeometryCollectionDebugDrawActor.h
@@ -41,11 +41,14 @@ struct FGeometryCollectionDebugDrawActorSelectedRigidBody
 {
 	GENERATED_USTRUCT_BODY()
 
-	explicit FGeometryCollectionDebugDrawActorSelectedRigidBody(int32 InId = -1) : Id(InId) {}
+	explicit FGeometryCollectionDebugDrawActorSelectedRigidBody(int32 InId = -1) : Id(InId), GeometryCollectionActor(nullptr) {}
 	
 	/** Id of the selected rigid body whose to visualize debug informations. Use -1 to visualize all Geometry Collections. */
 	UPROPERTY(EditAnywhere, Category = "Debug Draw", meta = (ClampMin="-1"))
 	int32 Id;
+
+	UPROPERTY(EditAnywhere, Category = "Debug Draw")
+	AActor* GeometryCollectionActor;
 };
 
 /**
diff --git a/Engine/Source/Runtime/Experimental/GeometryCollectionEngine/Public/GeometryCollection/GeometryCollectionDebugDrawComponent.h b/Engine/Source/Runtime/Experimental/GeometryCollectionEngine/Public/GeometryCollection/GeometryCollectionDebugDrawComponent.h
--- a/Engine/Source/Runtime/Experimental/GeometryCollectionEngine/Public/GeometryCollection/GeometryCollectionDebugDrawComponent.h
+++ b/Engine/Source/Runtime/Experimental/GeometryCollectionEngine/Public/GeometryCollection/GeometryCollectionDebugDrawComponent.h
@@ -235,8 +235,8 @@ public:
 #endif
 
 #if GEOMETRYCOLLECTION_DEBUG_DRAW
-	/** Update selection and visibility after any change in properties. Also enable/disable this component tick update. */
-	void OnDebugDrawPropertiesChanged(bool bForceVisibilityUpdate);
+	/** Update selection and visibility after any change in properties. Also enable/disable this component tick update. Return true if this geometry collection is selected. */
+	bool OnDebugDrawPropertiesChanged(bool bForceVisibilityUpdate);
 
 	/** Update selection and visibility after a change in cluster. Only handled when the debug drawing is active (the component is ticking). */
 	void OnClusterChanged();
