Improves the feedback on the geometry collection actor debug draw picking tool icon, so that it shows as pushed when selecting, and disabled when the rigid body selection cannot be triggered.
Merge CL710678 from Dev-Physics to Release-4.23.

diff --git a/Engine/Plugins/Experimental/GeometryCollectionPlugin/Source/GeometryCollectionEditor/Private/GeometryCollection/DetailCustomizations/SelectedRigidBodyCustomization.cpp b/Engine/Plugins/Experimental/GeometryCollectionPlugin/Source/GeometryCollectionEditor/Private/GeometryCollection/DetailCustomizations/SelectedRigidBodyCustomization.cpp
--- a/Engine/Plugins/Experimental/GeometryCollectionPlugin/Source/GeometryCollectionEditor/Private/GeometryCollection/DetailCustomizations/SelectedRigidBodyCustomization.cpp
+++ b/Engine/Plugins/Experimental/GeometryCollectionPlugin/Source/GeometryCollectionEditor/Private/GeometryCollection/DetailCustomizations/SelectedRigidBodyCustomization.cpp
@@ -42,13 +42,22 @@ void FSelectedRigidBodyCustomization::CustomizeHeader(TSharedRef<IPropertyHandle
 			]
 			+ SHorizontalBox::Slot()
 			.AutoWidth()
+			.Padding(2.f, 0.f, 0.f, 0.f)
 			[
-				SNew(SButton)
-				.ButtonStyle(FEditorStyle::Get(), "HoverHintOnly")
+				SAssignNew(CheckBoxPick, SCheckBox)
+				.Style(FEditorStyle::Get(), "ToggleButtonCheckbox")
 				.ToolTipText(NSLOCTEXT("ChaosSelectedRigidBody", "Pick_ToolTip", "Pick a Rigid Body."))
 				.ForegroundColor(FSlateColor::UseForeground())
-				.ContentPadding(6.0f)
-				.OnClicked(this, &FSelectedRigidBodyCustomization::OnPick, PropertyHandleId)
+				.Padding(6.0f)
+				.IsChecked_Lambda([]()
+				{
+					return FGeometryCollectionSelectRigidBodyEdMode::IsModeActive() ? ECheckBoxState::Checked : ECheckBoxState::Unchecked;
+				})
+				.IsEnabled_Lambda([]()
+				{
+					return FGeometryCollectionSelectRigidBodyEdMode::CanActivateMode();
+				})
+				.OnCheckStateChanged(this, &FSelectedRigidBodyCustomization::OnPick, PropertyHandleId)
 				[ 
 					SNew(SImage)
 					.Image(FEditorStyle::GetBrush("PropertyWindow.Button_PickActorInteractive"))
@@ -289,15 +298,40 @@ int32 FSelectedRigidBodyCustomization::GetNextClusteredSiblingRigidBodyId(TShare
 	return INDEX_NONE;
 }
 
-FReply FSelectedRigidBodyCustomization::OnPick(TSharedRef<IPropertyHandle> PropertyHandleId) const
+void FSelectedRigidBodyCustomization::OnPick(ECheckBoxState InCheckState, TSharedRef<IPropertyHandle> PropertyHandleId) const
 {
-	if (FGeometryCollectionSelectRigidBodyEdMode::IsModeActive())
+	// Set unchecked by default (enter mode is required to set checked)
+	if (CheckBoxPick)
 	{
-		FGeometryCollectionSelectRigidBodyEdMode::DeactivateMode();
+		CheckBoxPick->SetIsChecked(ECheckBoxState::Unchecked);
+	}
+
+	if (InCheckState == ECheckBoxState::Unchecked)
+	{
+		if (FGeometryCollectionSelectRigidBodyEdMode::IsModeActive())
+		{
+			FGeometryCollectionSelectRigidBodyEdMode::DeactivateMode();
+		}
 	}
 	else
 	{
-		FGeometryCollectionSelectRigidBodyEdMode::ActivateMode(PropertyHandleId);
+		const TWeakPtr<SCheckBox> CheckBoxPickWeakPtr = CheckBoxPick;
+		const TFunction<void()> OnEnterMode = [CheckBoxPickWeakPtr]()
+		{
+			if (const TSharedPtr<SCheckBox> CheckBoxPickPin = CheckBoxPickWeakPtr.Pin())
+			{
+				CheckBoxPickPin->SetToolTipText(NSLOCTEXT("ChaosSelectedRigidBody", "CancelPick_ToolTip", "Cancel Picking."));
+				CheckBoxPickPin->SetIsChecked(ECheckBoxState::Checked);
+			}
+		};
+		const TFunction<void()> OnExitMode = [CheckBoxPickWeakPtr]()
+		{
+			if (const TSharedPtr<SCheckBox> CheckBoxPickPin = CheckBoxPickWeakPtr.Pin())
+			{
+				CheckBoxPickPin->SetToolTipText(NSLOCTEXT("ChaosSelectedRigidBody", "Pick_ToolTip", "Pick a Rigid Body."));
+				CheckBoxPickPin->SetIsChecked(ECheckBoxState::Unchecked);
+			}
+		};
+		FGeometryCollectionSelectRigidBodyEdMode::ActivateMode(PropertyHandleId, OnEnterMode, OnExitMode);
 	}
-	return FReply::Handled();
 }
diff --git a/Engine/Plugins/Experimental/GeometryCollectionPlugin/Source/GeometryCollectionEditor/Private/GeometryCollection/DetailCustomizations/SelectedRigidBodyCustomization.h b/Engine/Plugins/Experimental/GeometryCollectionPlugin/Source/GeometryCollectionEditor/Private/GeometryCollection/DetailCustomizations/SelectedRigidBodyCustomization.h
--- a/Engine/Plugins/Experimental/GeometryCollectionPlugin/Source/GeometryCollectionEditor/Private/GeometryCollection/DetailCustomizations/SelectedRigidBodyCustomization.h
+++ b/Engine/Plugins/Experimental/GeometryCollectionPlugin/Source/GeometryCollectionEditor/Private/GeometryCollection/DetailCustomizations/SelectedRigidBodyCustomization.h
@@ -2,7 +2,7 @@
 
 #pragma once
 #include "IPropertyTypeCustomization.h"
-#include "Input/Reply.h"
+#include "Widgets/Input/SCheckBox.h"
 
 class IDetailLayoutBuilder;
 class UGeometryCollectionComponent;
@@ -17,6 +17,8 @@ public:
 	virtual void CustomizeChildren(TSharedRef<IPropertyHandle> StructPropertyHandle, class IDetailChildrenBuilder& ChildBuilder, IPropertyTypeCustomizationUtils& StructCustomizationUtils) override;
 
 private:
+	TSharedPtr<SCheckBox> CheckBoxPick;
+
 	static void GetSelectedGeometryCollectionCluster(TSharedRef<IPropertyHandle> PropertyHandleActor, TSharedRef<IPropertyHandle> PropertyHandleId, const UGeometryCollectionComponent*& OutGeometryCollectionComponent, int32& OutTransformIndex);
 
 	static int32 GetParentClusterRigidBodyId(TSharedRef<IPropertyHandle> PropertyHandleActor, TSharedRef<IPropertyHandle> PropertyHandleId);
@@ -24,5 +26,5 @@ private:
 	static int32 GetPreviousClusteredSiblingRigidBodyId(TSharedRef<IPropertyHandle> PropertyHandleActor, TSharedRef<IPropertyHandle> PropertyHandleId);
 	static int32 GetNextClusteredSiblingRigidBodyId(TSharedRef<IPropertyHandle> PropertyHandleActor, TSharedRef<IPropertyHandle> PropertyHandleId);
 
-	FReply OnPick(TSharedRef<IPropertyHandle> PropertyHandleId) const;
+	void OnPick(ECheckBoxState InCheckState, TSharedRef<IPropertyHandle> PropertyHandleId) const;
 };
diff --git a/Engine/Plugins/Experimental/GeometryCollectionPlugin/Source/GeometryCollectionEditor/Private/GeometryCollection/GeometryCollectionSelectRigidBodyEdMode.cpp b/Engine/Plugins/Experimental/GeometryCollectionPlugin/Source/GeometryCollectionEditor/Private/GeometryCollection/GeometryCollectionSelectRigidBodyEdMode.cpp
--- a/Engine/Plugins/Experimental/GeometryCollectionPlugin/Source/GeometryCollectionEditor/Private/GeometryCollection/GeometryCollectionSelectRigidBodyEdMode.cpp
+++ b/Engine/Plugins/Experimental/GeometryCollectionPlugin/Source/GeometryCollectionEditor/Private/GeometryCollection/GeometryCollectionSelectRigidBodyEdMode.cpp
@@ -18,7 +18,7 @@ DEFINE_LOG_CATEGORY_STATIC(LogGeometryCollectionSelectRigidBodyEdMode, Log, All)
 const FEditorModeID FGeometryCollectionSelectRigidBodyEdMode::EditorModeID(TEXT("Chaos.Select"));
 const int32 FGeometryCollectionSelectRigidBodyEdMode::MessageKey = GetTypeHash(EditorModeID);
 
-void FGeometryCollectionSelectRigidBodyEdMode::ActivateMode(TSharedRef<IPropertyHandle> PropertyHandleId)
+bool FGeometryCollectionSelectRigidBodyEdMode::CanActivateMode()
 {
 	auto HasValidEditorLevelViewport = [](const TIndirectArray<FWorldContext>& WorldContexts, TMap<FName, FSlatePlayInEditorInfo>& SlatePlayInEditorMap)
 	{
@@ -34,27 +34,46 @@ void FGeometryCollectionSelectRigidBodyEdMode::ActivateMode(TSharedRef<IProperty
 		return false;
 	};
 
+	return GEditor && GEditor->PlayWorld                                                             // Playing in editor
+		&& !IsRunningGame()                                                                          // Not running in -game mode
+		&& HasValidEditorLevelViewport(GEditor->GetWorldContexts(), GEditor->SlatePlayInEditorMap);  // Not running in new PIE window
+}
+
+void FGeometryCollectionSelectRigidBodyEdMode::ActivateMode(TSharedRef<IPropertyHandle> PropertyHandleId, TFunction<void()> OnEnterMode, TFunction<void()> OnExitMode)
+{
 	// Make sure we are playing in an editor window
-	if (GEditor && GEditor->PlayWorld &&  // Playing in editor
-		!IsRunningGame() &&               // Not running in -game mode
-		HasValidEditorLevelViewport(GEditor->GetWorldContexts(), GEditor->SlatePlayInEditorMap))  // Not running in new PIE window
+	if (CanActivateMode())
 	{
-		// Eject PIE
-		if (!GEditor->bIsSimulatingInEditor)  // Running in PIE
+		// Running in PIE?
+		if (!GEditor->bIsSimulatingInEditor)
 		{
+			// Eject PIE
 			GEditor->RequestToggleBetweenPIEandSIE();
+
+			// Log/display message
 			UE_LOG(LogGeometryCollectionSelectRigidBodyEdMode, Warning, TEXT("Player possession ejected by rigid body picker."));
 			if (GEngine)
 			{
 				GEngine->AddOnScreenDebugMessage(-1, 5.f, FColor::Yellow, TEXT("Player possession ejected by rigid body picker."));
 			}
 		}
+
 		// Activate editor mode
 		GLevelEditorModeTools().ActivateMode(EditorModeID);
-		FEdMode* const EdMode = GLevelEditorModeTools().FindMode(EditorModeID);
-		if (EdMode)
+		if (FEdMode* const EdMode = GLevelEditorModeTools().FindMode(EditorModeID))
 		{
-			static_cast<FGeometryCollectionSelectRigidBodyEdMode*>(EdMode)->PropertyHandleId = PropertyHandleId;
+			// Set pointers
+			FGeometryCollectionSelectRigidBodyEdMode* const GeometryCollectionSelectRigidBodyEdMode = static_cast<FGeometryCollectionSelectRigidBodyEdMode*>(EdMode);
+			GeometryCollectionSelectRigidBodyEdMode->PropertyHandleId = PropertyHandleId;
+			GeometryCollectionSelectRigidBodyEdMode->OnExitMode = OnExitMode;
+
+			// Execute enter mode callback
+			if (OnEnterMode)
+			{
+				OnEnterMode();
+			}
+
+			// Log/display message
 			UE_LOG(LogGeometryCollectionSelectRigidBodyEdMode, Log, TEXT("Click on the rigid body to select..."));
 			if (GEngine)
 			{
@@ -64,6 +83,7 @@ void FGeometryCollectionSelectRigidBodyEdMode::ActivateMode(TSharedRef<IProperty
 	}
 	else
 	{
+		// Can't pick in non playing editor
 		UE_LOG(LogGeometryCollectionSelectRigidBodyEdMode, Error, TEXT("The rigid body picker is only available in play in editor."));
 		if (GEngine)
 		{
diff --git a/Engine/Plugins/Experimental/GeometryCollectionPlugin/Source/GeometryCollectionEditor/Private/GeometryCollection/GeometryCollectionSelectRigidBodyEdMode.h b/Engine/Plugins/Experimental/GeometryCollectionPlugin/Source/GeometryCollectionEditor/Private/GeometryCollection/GeometryCollectionSelectRigidBodyEdMode.h
--- a/Engine/Plugins/Experimental/GeometryCollectionPlugin/Source/GeometryCollectionEditor/Private/GeometryCollection/GeometryCollectionSelectRigidBodyEdMode.h
+++ b/Engine/Plugins/Experimental/GeometryCollectionPlugin/Source/GeometryCollectionEditor/Private/GeometryCollection/GeometryCollectionSelectRigidBodyEdMode.h
@@ -13,13 +13,16 @@ class FGeometryCollectionSelectRigidBodyEdMode : public FEdMode
 public:
 	static const FEditorModeID EditorModeID;
 
-	/* Activate this editor mode */
-	static void ActivateMode(TSharedRef<IPropertyHandle> PropertyHandleId);
+	/** Check whether this mode can be activated (needs to run in SIE/PIE). */
+	static bool CanActivateMode();
 
-	/* Activate this editor mode */
+	/** Activate this editor mode. */
+	static void ActivateMode(TSharedRef<IPropertyHandle> PropertyHandleId, TFunction<void()> OnEnterMode, TFunction<void()> OnExitMode);
+
+	/** Deactivate this editor mode. */
 	static void DeactivateMode();
 
-	/* Return whether Activate this editor mode */
+	/** Return whether this editor mode is active. */
 	static bool IsModeActive();
 
 	FGeometryCollectionSelectRigidBodyEdMode() : bIsHoveringGeometryCollection(false) {}
@@ -27,7 +30,7 @@ public:
 
 	/* FEdMode interface */
 	virtual void Enter() override { EnableTransformSelectionMode(true); }
-	virtual void Exit() override { EnableTransformSelectionMode(false); PropertyHandleId.Reset(); }
+	virtual void Exit() override { EnableTransformSelectionMode(false); PropertyHandleId.Reset(); if (OnExitMode) { OnExitMode(); } }
 	virtual bool IsCompatibleWith(FEditorModeID /*OtherModeID*/) const override { return false; }
 
 	virtual bool GetCursor(EMouseCursor::Type& OutCursor) const { OutCursor = bIsHoveringGeometryCollection ? EMouseCursor::EyeDropper: EMouseCursor::SlashedCircle; return true; }
@@ -48,6 +51,7 @@ private:
 private:
 	static const int32 MessageKey;
 	TWeakPtr<IPropertyHandle> PropertyHandleId;  // Handle of the property that will get updated with the selected rigid body id
+	TFunction<void()> OnExitMode;  // Callback function called when the edit mode is deactivated
 	bool bIsHoveringGeometryCollection;
 };
 
