Fix geometry collection debug draw navigation to skip children/siblings rigid body ids when they are invalid.

diff --git a/Engine/Plugins/Experimental/GeometryCollectionPlugin/Source/GeometryCollectionEditor/Private/GeometryCollection/DetailCustomizations/SelectedRigidBodyCustomization.cpp b/Engine/Plugins/Experimental/GeometryCollectionPlugin/Source/GeometryCollectionEditor/Private/GeometryCollection/DetailCustomizations/SelectedRigidBodyCustomization.cpp
--- a/Engine/Plugins/Experimental/GeometryCollectionPlugin/Source/GeometryCollectionEditor/Private/GeometryCollection/DetailCustomizations/SelectedRigidBodyCustomization.cpp
+++ b/Engine/Plugins/Experimental/GeometryCollectionPlugin/Source/GeometryCollectionEditor/Private/GeometryCollection/DetailCustomizations/SelectedRigidBodyCustomization.cpp
@@ -17,6 +17,8 @@
 #include "GeometryCollection/GeometryCollectionComponent.h"
 #include "GeometryCollection/GeometryCollectionActor.h"
 
+DEFINE_LOG_CATEGORY_STATIC(LogSelectedRigidBodyCustomization, Log, All);
+
 TSharedRef<IPropertyTypeCustomization> FSelectedRigidBodyCustomization::MakeInstance()
 {
 	return MakeShareable(new FSelectedRigidBodyCustomization);
@@ -216,6 +218,7 @@ void FSelectedRigidBodyCustomization::GetSelectedGeometryCollectionCluster(TShar
 			}
 		}
 	}
+	UE_CLOG(OutGeometryCollectionComponent && OutGeometryCollectionComponent->GetOwner(), LogSelectedRigidBodyCustomization, VeryVerbose, TEXT("Component actor %s, TransformIndex %d."), *OutGeometryCollectionComponent->GetOwner()->GetName(), OutTransformIndex);
 }
 
 int32 FSelectedRigidBodyCustomization::GetParentClusterRigidBodyId(TSharedRef<IPropertyHandle> PropertyHandleActor, TSharedRef<IPropertyHandle> PropertyHandleId)
@@ -242,9 +245,14 @@ int32 FSelectedRigidBodyCustomization::GetChildClusterRigidBodyId(TSharedRef<IPr
 	if (GeometryCollectionComponent)
 	{
 		const TSet<int32>& Children = GeometryCollectionComponent->GetChildrenArray()[TransformIndex];
-		if (TSet<int32>::TConstIterator ChildTransformIndexIt = Children.CreateConstIterator())
+		// Iterate through all children until the first valid rigid body id is found
+		for (TSet<int32>::TConstIterator ChildTransformIndexIt = Children.CreateConstIterator(); ChildTransformIndexIt; ++ChildTransformIndexIt)
 		{
-			return GeometryCollectionComponent->GetRigidBodyIdArray()[*ChildTransformIndexIt];
+			const int32 RigidBodyId = GeometryCollectionComponent->GetRigidBodyIdArray()[*ChildTransformIndexIt];
+			if (RigidBodyId != INDEX_NONE)
+			{
+				return RigidBodyId;
+			}
 		}
 	}
 	return INDEX_NONE;
@@ -261,18 +269,19 @@ int32 FSelectedRigidBodyCustomization::GetPreviousClusteredSiblingRigidBodyId(TS
 		if (ParentTransformIndex != INDEX_NONE)
 		{
 			const TSet<int32>& Siblings = GeometryCollectionComponent->GetChildrenArray()[ParentTransformIndex];
-			int32 PrevSibling = INDEX_NONE;
+			int32 PrevRigidBodyId = INDEX_NONE;
 			for (int32 Sibling : Siblings)
 			{
 				if (Sibling == TransformIndex)
 				{
-					if (PrevSibling != INDEX_NONE)
-					{
-						return GeometryCollectionComponent->GetRigidBodyIdArray()[PrevSibling];
-					}
-					break;
+					return PrevRigidBodyId;
+				}
+				// Only use a previous sibling with a valid rigid body id
+				const int32 RigidBodyId = GeometryCollectionComponent->GetRigidBodyIdArray()[Sibling];
+				if (RigidBodyId != INDEX_NONE)
+				{
+					PrevRigidBodyId = RigidBodyId;
 				}
-				PrevSibling = Sibling;
 			}
 		}
 	}
@@ -295,9 +304,17 @@ int32 FSelectedRigidBodyCustomization::GetNextClusteredSiblingRigidBodyId(TShare
 			{
 				if (Sibling == TransformIndex)
 				{
-					return GeometryCollectionComponent->GetRigidBodyIdArray()[NextSibling];
+					// Return the first valid rigid body id
+					const int32 NextRigidBodyId = GeometryCollectionComponent->GetRigidBodyIdArray()[NextSibling];
+					if (NextRigidBodyId != INDEX_NONE)
+					{
+						return NextRigidBodyId;
+					}
+				}
+				else
+				{
+					Sibling = NextSibling;
 				}
-				Sibling = NextSibling;
 			}
 		}
 	}
