Fracture: Removes Geometry Collection Bone Hierarchy from the World Outliner

diff --git a/Engine/Plugins/Experimental/GeometryCollectionPlugin/Source/GeometryCollectionEditor/Private/GeometryCollection/DetailCustomizations/GeomComponentCacheCustomization.cpp b/Engine/Plugins/Experimental/GeometryCollectionPlugin/Source/GeometryCollectionEditor/Private/GeometryCollection/DetailCustomizations/GeomComponentCacheCustomization.cpp
--- a/Engine/Plugins/Experimental/GeometryCollectionPlugin/Source/GeometryCollectionEditor/Private/GeometryCollection/DetailCustomizations/GeomComponentCacheCustomization.cpp
+++ b/Engine/Plugins/Experimental/GeometryCollectionPlugin/Source/GeometryCollectionEditor/Private/GeometryCollection/DetailCustomizations/GeomComponentCacheCustomization.cpp
@@ -13,6 +13,8 @@
 #include "GeometryCollection/GeometryCollectionObject.h"
 #include "GeometryCollection/GeometryCollectionEditorPlugin.h"
 #include "EditorFontGlyphs.h"
+#include "Widgets/Layout/SBorder.h"
+#include "Widgets/Text/STextBlock.h"
 
 #define LOCTEXT_NAMESPACE "GeomCollectionCacheParamsCustomization"
 
diff --git a/Engine/Plugins/Experimental/GeometryCollectionPlugin/Source/GeometryCollectionEditor/Private/GeometryCollection/GeometryCollectionBoneTreeItem.cpp b/Engine/Plugins/Experimental/GeometryCollectionPlugin/Source/GeometryCollectionEditor/Private/GeometryCollection/GeometryCollectionBoneTreeItem.cpp
deleted file mode 100644
diff --git a/Engine/Plugins/Experimental/GeometryCollectionPlugin/Source/GeometryCollectionEditor/Private/GeometryCollection/GeometryCollectionEditorPlugin.cpp b/Engine/Plugins/Experimental/GeometryCollectionPlugin/Source/GeometryCollectionEditor/Private/GeometryCollection/GeometryCollectionEditorPlugin.cpp
--- a/Engine/Plugins/Experimental/GeometryCollectionPlugin/Source/GeometryCollectionEditor/Private/GeometryCollection/GeometryCollectionEditorPlugin.cpp
+++ b/Engine/Plugins/Experimental/GeometryCollectionPlugin/Source/GeometryCollectionEditor/Private/GeometryCollection/GeometryCollectionEditorPlugin.cpp
@@ -169,7 +169,6 @@ void IGeometryCollectionEditorPlugin::StartupModule()
 
 	// Bind our scene outliner provider to the editor
 	IModularFeatures& ModularFeatures = IModularFeatures::Get();
-	ModularFeatures.RegisterModularFeature(TEXT("SceneOutlinerTraversal"), &OutlinerTraversalAccessor);
 	ModularFeatures.RegisterModularFeature(ITargetCacheProvider::GetFeatureName(), &TargetCacheProvider);
 
 	// Register type customizations
@@ -224,7 +223,6 @@ void IGeometryCollectionEditorPlugin::ShutdownModule()
 
 		// Unbind provider from editor
 		IModularFeatures& ModularFeatures = IModularFeatures::Get();
-		ModularFeatures.UnregisterModularFeature(TEXT("SceneOutlinerTraversal"), &OutlinerTraversalAccessor);
 		ModularFeatures.UnregisterModularFeature(ITargetCacheProvider::GetFeatureName(), &TargetCacheProvider);
 
 		// Unregister type customizations
diff --git a/Engine/Plugins/Experimental/GeometryCollectionPlugin/Source/GeometryCollectionEditor/Private/GeometryCollection/GeometryCollectionTreeItem.cpp b/Engine/Plugins/Experimental/GeometryCollectionPlugin/Source/GeometryCollectionEditor/Private/GeometryCollection/GeometryCollectionTreeItem.cpp
deleted file mode 100644
diff --git a/Engine/Plugins/Experimental/GeometryCollectionPlugin/Source/GeometryCollectionEditor/Private/GeometryCollection/OutlinerTraversalAccessor.cpp b/Engine/Plugins/Experimental/GeometryCollectionPlugin/Source/GeometryCollectionEditor/Private/GeometryCollection/OutlinerTraversalAccessor.cpp
deleted file mode 100644
diff --git a/Engine/Plugins/Experimental/GeometryCollectionPlugin/Source/GeometryCollectionEditor/Private/GeometryCollection/OutlinerTraversalAccessor.h b/Engine/Plugins/Experimental/GeometryCollectionPlugin/Source/GeometryCollectionEditor/Private/GeometryCollection/OutlinerTraversalAccessor.h
deleted file mode 100644
diff --git a/Engine/Plugins/Experimental/GeometryCollectionPlugin/Source/GeometryCollectionEditor/Public/GeometryCollection/GeometryCollectionBoneTreeItem.h b/Engine/Plugins/Experimental/GeometryCollectionPlugin/Source/GeometryCollectionEditor/Public/GeometryCollection/GeometryCollectionBoneTreeItem.h
deleted file mode 100644
diff --git a/Engine/Plugins/Experimental/GeometryCollectionPlugin/Source/GeometryCollectionEditor/Public/GeometryCollection/GeometryCollectionEditorPlugin.h b/Engine/Plugins/Experimental/GeometryCollectionPlugin/Source/GeometryCollectionEditor/Public/GeometryCollection/GeometryCollectionEditorPlugin.h
--- a/Engine/Plugins/Experimental/GeometryCollectionPlugin/Source/GeometryCollectionEditor/Public/GeometryCollection/GeometryCollectionEditorPlugin.h
+++ b/Engine/Plugins/Experimental/GeometryCollectionPlugin/Source/GeometryCollectionEditor/Public/GeometryCollection/GeometryCollectionEditorPlugin.h
@@ -5,9 +5,10 @@
 #include "CoreMinimal.h"
 #include "Modules/ModuleInterface.h"
 #include "Modules/ModuleManager.h"
-#include "GeometryCollection/OutlinerTraversalAccessor.h"
 #include "GeometryCollection/GeometryCollectionProviderEditor.h"
 
+
+class ISlateStyle;
 class FSlateStyleSet;
 class FAssetTypeActions_GeometryCollection;
 class FAssetTypeActions_GeometryCollectionCache;
@@ -44,7 +45,6 @@ public:
 		return FModuleManager::Get().IsModuleLoaded( "GeometryCollectionEditorPlugin" );
 	}
 
-	FOutlinerTraversalAccessor& GetAccessor() { return OutlinerTraversalAccessor; }
 	TSharedPtr<FSlateStyleSet> GetStyleSet() { return StyleSet; }
 
 	static FName GetEditorStyleName();
@@ -59,7 +59,6 @@ private:
 	FGeometryCollectionAssetBroker* AssetBroker;
 
 	// Modular features
-	FOutlinerTraversalAccessor OutlinerTraversalAccessor;
 	// Provider for new caches requested from other modules
 	FTargetCacheProviderEditor TargetCacheProvider;
 	//////////////////////////////////////////////////////////////////////////
diff --git a/Engine/Plugins/Experimental/GeometryCollectionPlugin/Source/GeometryCollectionEditor/Public/GeometryCollection/GeometryCollectionTreeItem.h b/Engine/Plugins/Experimental/GeometryCollectionPlugin/Source/GeometryCollectionEditor/Public/GeometryCollection/GeometryCollectionTreeItem.h
deleted file mode 100644
