Add a warning to the geometry collection debug draw UI for when the editor is not in PIE.

diff --git a/Engine/Plugins/Experimental/GeometryCollectionPlugin/Source/GeometryCollectionEditor/Private/GeometryCollection/DetailCustomizations/WarningMessageCustomization.cpp b/Engine/Plugins/Experimental/GeometryCollectionPlugin/Source/GeometryCollectionEditor/Private/GeometryCollection/DetailCustomizations/WarningMessageCustomization.cpp
new file mode 100644
--- /dev/null
+++ b/Engine/Plugins/Experimental/GeometryCollectionPlugin/Source/GeometryCollectionEditor/Private/GeometryCollection/DetailCustomizations/WarningMessageCustomization.cpp
@@ -0,0 +1,95 @@
+// Copyright 1998-2019 Epic Games, Inc. All Rights Reserved.
+
+#include "GeometryCollection/DetailCustomizations/WarningMessageCustomization.h"
+
+#include "PropertyHandle.h"
+#include "IPropertyTypeCustomization.h"
+#include "IPropertyUtilities.h"
+#include "DetailWidgetRow.h"
+#include "Widgets/Layout/SBorder.h"
+#include "Widgets/Text/STextBlock.h"
+#include "EditorFontGlyphs.h"
+#include "GeometryCollection/GeometryCollectionActor.h"
+#include "GeometryCollection/GeometryCollectionDebugDrawActor.h"
+
+TSharedRef<IPropertyTypeCustomization> FWarningMessageCustomization::MakeInstance()
+{
+	return MakeShareable(new FWarningMessageCustomization);
+}
+
+void FWarningMessageCustomization::CustomizeHeader(TSharedRef<IPropertyHandle> StructPropertyHandle, FDetailWidgetRow& HeaderRow, IPropertyTypeCustomizationUtils& StructCustomizationUtils)
+{
+	// Hide reset-to-default
+	StructPropertyHandle->MarkResetToDefaultCustomized();
+
+	// Check whether the selected actor has begun play
+	bool bHasBegunPlay = false;
+	for (const TWeakObjectPtr<UObject> SelectedObject: StructCustomizationUtils.GetPropertyUtilities()->GetSelectedObjects())
+	{
+		if (AGeometryCollectionDebugDrawActor* const Actor = Cast<AGeometryCollectionDebugDrawActor>(SelectedObject.Get()))
+		{
+			bHasBegunPlay = Actor->HasActorBegunPlay();
+			break;
+		}
+		if (AGeometryCollectionActor* const Actor = Cast<AGeometryCollectionActor>(SelectedObject.Get()))
+		{
+			bHasBegunPlay = Actor->HasActorBegunPlay();
+			break;
+		}
+		if (UGeometryCollectionDebugDrawComponent* const Component = Cast<UGeometryCollectionDebugDrawComponent>(SelectedObject.Get()))
+		{
+			bHasBegunPlay = Component->HasBegunPlay();
+			break;
+		}
+	}
+
+	if (!bHasBegunPlay)
+	{
+		static const FText ToolTipMessage = NSLOCTEXT("GeometryCollectionDebugDrawWarningMessage", "WarningMessage_ToolTip", "The actor must be playing for these debug draw properties to have any effect.");
+		static const FText Message = NSLOCTEXT("GeometryCollectionDebugDrawWarningMessage", "WarningMessage_MessageWaiting", "Waiting for Play...");
+		static const FSlateColor Color = FEditorStyle::GetSlateColor("SelectionColor");
+
+		// Add buttons
+		HeaderRow.ValueContent()
+		.MinDesiredWidth(140.f)
+		[
+			SNew(SHorizontalBox)
+			+ SHorizontalBox::Slot()
+			.AutoWidth()
+			.Padding(0.f)
+			[
+				SNew(SBorder)
+				.BorderImage(FEditorStyle::GetBrush("NoBorder"))
+				.ToolTipText(ToolTipMessage)
+				.VAlign(VAlign_Center)
+				.HAlign(HAlign_Center)
+				[
+					SNew(STextBlock)
+					.Font(FEditorStyle::Get().GetFontStyle("FontAwesome.11"))
+					.Text(FEditorFontGlyphs::Exclamation_Triangle)
+					.ColorAndOpacity(Color)
+				]
+			]
+			+ SHorizontalBox::Slot()
+			.AutoWidth()
+			.Padding(4.f, 0.f)
+			[
+				SNew(SBorder)
+				.BorderImage(FEditorStyle::GetBrush("NoBorder"))
+				.ToolTipText(ToolTipMessage)
+				.VAlign(VAlign_Center)
+				.HAlign(HAlign_Center)
+				[
+					SNew(STextBlock)
+					.Font(FCoreStyle::GetDefaultFontStyle("Italic", 9) )
+					.Text(Message)
+					.ColorAndOpacity(Color)
+				]
+			]
+		];
+	}
+}
+
+void FWarningMessageCustomization::CustomizeChildren(TSharedRef<IPropertyHandle> /*StructPropertyHandle*/, IDetailChildrenBuilder& /*ChildBuilder*/, IPropertyTypeCustomizationUtils& /*StructCustomizationUtils*/)
+{
+}
diff --git a/Engine/Plugins/Experimental/GeometryCollectionPlugin/Source/GeometryCollectionEditor/Private/GeometryCollection/DetailCustomizations/WarningMessageCustomization.h b/Engine/Plugins/Experimental/GeometryCollectionPlugin/Source/GeometryCollectionEditor/Private/GeometryCollection/DetailCustomizations/WarningMessageCustomization.h
new file mode 100644
--- /dev/null
+++ b/Engine/Plugins/Experimental/GeometryCollectionPlugin/Source/GeometryCollectionEditor/Private/GeometryCollection/DetailCustomizations/WarningMessageCustomization.h
@@ -0,0 +1,16 @@
+// Copyright 1998-2019 Epic Games, Inc. All Rights Reserved.
+
+#pragma once
+#include "IPropertyTypeCustomization.h"
+
+class IDetailLayoutBuilder;
+
+/** Details customization for warning message in the UI. */
+class FWarningMessageCustomization : public IPropertyTypeCustomization
+{
+public:
+	static TSharedRef<IPropertyTypeCustomization> MakeInstance();
+
+	virtual void CustomizeHeader(TSharedRef<IPropertyHandle> StructPropertyHandle, class FDetailWidgetRow& HeaderRow, IPropertyTypeCustomizationUtils& StructCustomizationUtils) override;
+	virtual void CustomizeChildren(TSharedRef<IPropertyHandle> StructPropertyHandle, class IDetailChildrenBuilder& ChildBuilder, IPropertyTypeCustomizationUtils& StructCustomizationUtils) override;
+};
diff --git a/Engine/Plugins/Experimental/GeometryCollectionPlugin/Source/GeometryCollectionEditor/Private/GeometryCollection/GeometryCollectionEditorPlugin.cpp b/Engine/Plugins/Experimental/GeometryCollectionPlugin/Source/GeometryCollectionEditor/Private/GeometryCollection/GeometryCollectionEditorPlugin.cpp
--- a/Engine/Plugins/Experimental/GeometryCollectionPlugin/Source/GeometryCollectionEditor/Private/GeometryCollection/GeometryCollectionEditorPlugin.cpp
+++ b/Engine/Plugins/Experimental/GeometryCollectionPlugin/Source/GeometryCollectionEditor/Private/GeometryCollection/GeometryCollectionEditorPlugin.cpp
@@ -20,6 +20,7 @@
 #include "Features/IModularFeatures.h"
 #include "GeometryCollection/DetailCustomizations/GeomComponentCacheCustomization.h"
 #include "GeometryCollection/DetailCustomizations/SelectedRigidBodyCustomization.h"
+#include "GeometryCollection/DetailCustomizations/WarningMessageCustomization.h"
 #include "Styling/SlateStyle.h"
 #include "Styling/SlateStyleRegistry.h"
 #include "Styling/CoreStyle.h"
@@ -177,6 +178,7 @@ void IGeometryCollectionEditorPlugin::StartupModule()
 	{
 		PropertyModule->RegisterCustomPropertyTypeLayout("GeomComponentCacheParameters", FOnGetPropertyTypeCustomizationInstance::CreateStatic(&FGeomComponentCacheParametersCustomization::MakeInstance));
 		PropertyModule->RegisterCustomPropertyTypeLayout("GeometryCollectionDebugDrawActorSelectedRigidBody", FOnGetPropertyTypeCustomizationInstance::CreateStatic(&FSelectedRigidBodyCustomization::MakeInstance));
+		PropertyModule->RegisterCustomPropertyTypeLayout("GeometryCollectionDebugDrawWarningMessage", FOnGetPropertyTypeCustomizationInstance::CreateStatic(&FWarningMessageCustomization::MakeInstance));
 	}
 
 	// Register rigid body selection editor mode
@@ -231,6 +233,7 @@ void IGeometryCollectionEditorPlugin::ShutdownModule()
 		{
 			PropertyModule->UnregisterCustomPropertyTypeLayout("GeomCollectionCacheParameters");
 			PropertyModule->UnregisterCustomPropertyTypeLayout("GeometryCollectionDebugDrawActorSelectedRigidBody");
+			PropertyModule->UnregisterCustomPropertyTypeLayout("GeometryCollectionDebugDrawWarningMessage");
 		}
 
 		// Unregister rigid body selection editor mode
diff --git a/Engine/Source/Runtime/Experimental/GeometryCollectionEngine/Private/GeometryCollection/GeometryCollectionDebugDrawActor.cpp b/Engine/Source/Runtime/Experimental/GeometryCollectionEngine/Private/GeometryCollection/GeometryCollectionDebugDrawActor.cpp
--- a/Engine/Source/Runtime/Experimental/GeometryCollectionEngine/Private/GeometryCollection/GeometryCollectionDebugDrawActor.cpp
+++ b/Engine/Source/Runtime/Experimental/GeometryCollectionEngine/Private/GeometryCollection/GeometryCollectionDebugDrawActor.cpp
@@ -5,7 +5,6 @@
 #include "GeometryCollection/GeometryCollectionParticlesData.h"
 #include "GeometryCollection/GeometryCollectionComponent.h"
 #include "GeometryCollection/GeometryCollectionActor.h"
-#include "GeometryCollection/GeometryCollectionDebugDrawComponent.h"
 
 #include "DrawDebugHelpers.h"
 #include "Debug/DebugDrawService.h"
diff --git a/Engine/Source/Runtime/Experimental/GeometryCollectionEngine/Public/GeometryCollection/GeometryCollectionDebugDrawActor.h b/Engine/Source/Runtime/Experimental/GeometryCollectionEngine/Public/GeometryCollection/GeometryCollectionDebugDrawActor.h
--- a/Engine/Source/Runtime/Experimental/GeometryCollectionEngine/Public/GeometryCollection/GeometryCollectionDebugDrawActor.h
+++ b/Engine/Source/Runtime/Experimental/GeometryCollectionEngine/Public/GeometryCollection/GeometryCollectionDebugDrawActor.h
@@ -4,6 +4,7 @@
 
 #include "GameFramework/Actor.h"
 #include "Components/BillboardComponent.h"
+#include "GeometryCollection/GeometryCollectionDebugDrawComponent.h"
 
 #include "GeometryCollectionDebugDrawActor.generated.h"
 
@@ -44,7 +45,7 @@ struct FGeometryCollectionDebugDrawActorSelectedRigidBody
 	GENERATED_USTRUCT_BODY()
 
 	explicit FGeometryCollectionDebugDrawActorSelectedRigidBody(int32 InId = -1) : Id(InId), Solver(nullptr), GeometryCollection(nullptr) {}
-	
+
 	/** Id of the selected rigid body whose to visualize debug informations. Use -1 to visualize all Geometry Collections. */
 	UPROPERTY(EditAnywhere, Category = "Selected Rigid Body", meta = (ClampMin="-1"))
 	int32 Id;
@@ -74,6 +75,10 @@ class GEOMETRYCOLLECTIONENGINE_API AGeometryCollectionDebugDrawActor : public AA
 	GENERATED_UCLASS_BODY()
 public:
 
+	/** Warning message to explain that the debug draw properties have no effect until starting playing/simulating. */
+	UPROPERTY(EditAnywhere, Category = "Debug Draw")
+	FGeometryCollectionDebugDrawWarningMessage WarningMessage;
+
 	/** Picking tool used to select a rigid body id. */
 	UPROPERTY(EditAnywhere, Category = "Debug Draw")
 	FGeometryCollectionDebugDrawActorSelectedRigidBody SelectedRigidBody;
diff --git a/Engine/Source/Runtime/Experimental/GeometryCollectionEngine/Public/GeometryCollection/GeometryCollectionDebugDrawComponent.h b/Engine/Source/Runtime/Experimental/GeometryCollectionEngine/Public/GeometryCollection/GeometryCollectionDebugDrawComponent.h
--- a/Engine/Source/Runtime/Experimental/GeometryCollectionEngine/Public/GeometryCollection/GeometryCollectionDebugDrawComponent.h
+++ b/Engine/Source/Runtime/Experimental/GeometryCollectionEngine/Public/GeometryCollection/GeometryCollectionDebugDrawComponent.h
@@ -24,8 +24,21 @@ namespace Chaos { template<class T, int d> class TImplicitObject; }
 namespace Chaos { template<class T, int d> class TPBDRigidParticles; }
 #endif  // #if INCLUDE_CHAOS && GEOMETRYCOLLECTION_DEBUG_DRAW
 
-// Component adding debug drawing functionality to a GeometryCollectionActor.
-// This component is automatically added to every GeometryCollectionActor.
+/**
+* FGeometryCollectionDebugDrawWarningMessage
+*   Empty structure used to embed a warning message in the UI through a detail customization.
+*/
+USTRUCT()
+struct FGeometryCollectionDebugDrawWarningMessage
+{
+	GENERATED_USTRUCT_BODY()
+};
+
+/**
+* UGeometryCollectionDebugDrawComponent
+*   Component adding debug drawing functionality to a GeometryCollectionActor.
+*   This component is automatically added to every GeometryCollectionActor.
+*/
 UCLASS(meta = (BlueprintSpawnableComponent), HideCategories = ("Tags", "Activation", "Cooking", "AssetUserData", "Collision"))
 class GEOMETRYCOLLECTIONENGINE_API UGeometryCollectionDebugDrawComponent : public UActorComponent
 {
@@ -33,6 +46,10 @@ class GEOMETRYCOLLECTIONENGINE_API UGeometryCollectionDebugDrawComponent : publi
 
 public:
 
+	/** Warning message to explain that the debug draw properties have no effect until starting playing/simulating. */
+	UPROPERTY(EditAnywhere, Category = "Debug Draw")
+	FGeometryCollectionDebugDrawWarningMessage WarningMessage;
+
 	/** Enable the visualization of the rigid body ids. */
 	UPROPERTY(EditAnywhere, Category = "Debug Draw|Rigid Body")
 	bool bShowRigidBodyIds;
