Prevent double click from accidentaly collapsing the geometry collection debug draw selection navigation.

diff --git a/Engine/Plugins/Experimental/GeometryCollectionPlugin/Source/GeometryCollectionEditor/Private/GeometryCollection/DetailCustomizations/SelectedRigidBodyCustomization.cpp b/Engine/Plugins/Experimental/GeometryCollectionPlugin/Source/GeometryCollectionEditor/Private/GeometryCollection/DetailCustomizations/SelectedRigidBodyCustomization.cpp
--- a/Engine/Plugins/Experimental/GeometryCollectionPlugin/Source/GeometryCollectionEditor/Private/GeometryCollection/DetailCustomizations/SelectedRigidBodyCustomization.cpp
+++ b/Engine/Plugins/Experimental/GeometryCollectionPlugin/Source/GeometryCollectionEditor/Private/GeometryCollection/DetailCustomizations/SelectedRigidBodyCustomization.cpp
@@ -40,139 +40,145 @@ void FSelectedRigidBodyCustomization::CustomizeHeader(TSharedRef<IPropertyHandle
 	.ValueContent()
 	.MinDesiredWidth(140.f)
 	[
-		SNew(SHorizontalBox)
-		+ SHorizontalBox::Slot()
-		.AutoWidth()
+		SNew(SBorder)
+		.BorderImage(FEditorStyle::GetBrush("NoBorder"))
 		.Padding(0.f)
+		.OnMouseDoubleClick_Lambda([](const FGeometry&, const FPointerEvent&) { return FReply::Handled(); })  // Disable the double click on all arrows to avoid closing/opening the child properties by mistake
 		[
-			SAssignNew(CheckBoxPick, SCheckBox)
-			.Style(FEditorStyle::Get(), "ToggleButtonCheckbox")
-			.ToolTipText(NSLOCTEXT("ChaosSelectedRigidBody", "Pick_ToolTip", "Pick a Rigid Body."))
-			.ForegroundColor(FSlateColor::UseForeground())
-			.Padding(6.0f)
-			.IsChecked_Lambda([]()
-			{
-				return FGeometryCollectionSelectRigidBodyEdMode::IsModeActive() ? ECheckBoxState::Checked : ECheckBoxState::Unchecked;
-			})
-			.IsEnabled_Lambda([]()
-			{
-				return FGeometryCollectionSelectRigidBodyEdMode::CanActivateMode();
-			})
-			.OnCheckStateChanged(this, &FSelectedRigidBodyCustomization::OnPick, PropertyHandleId, PropertyHandleSolver)
-			[ 
-				SNew(SImage)
-				.Image(FEditorStyle::GetBrush("PropertyWindow.Button_PickActorInteractive"))
-				.ColorAndOpacity(FSlateColor::UseForeground())
+			SNew(SHorizontalBox)
+			+ SHorizontalBox::Slot()
+			.AutoWidth()
+			.Padding(0.f)
+			[
+				SAssignNew(CheckBoxPick, SCheckBox)
+				.Style(FEditorStyle::Get(), "ToggleButtonCheckbox")
+				.ToolTipText(NSLOCTEXT("ChaosSelectedRigidBody", "Pick_ToolTip", "Pick a Rigid Body."))
+				.ForegroundColor(FSlateColor::UseForeground())
+				.Padding(6.0f)
+				.IsChecked_Lambda([]()
+				{
+					return FGeometryCollectionSelectRigidBodyEdMode::IsModeActive() ? ECheckBoxState::Checked : ECheckBoxState::Unchecked;
+				})
+				.IsEnabled_Lambda([]()
+				{
+					return FGeometryCollectionSelectRigidBodyEdMode::CanActivateMode();
+				})
+				.OnCheckStateChanged(this, &FSelectedRigidBodyCustomization::OnPick, PropertyHandleId, PropertyHandleSolver)
+				[ 
+					SNew(SImage)
+					.Image(FEditorStyle::GetBrush("PropertyWindow.Button_PickActorInteractive"))
+					.ColorAndOpacity(FSlateColor::UseForeground())
+				]
 			]
-		]
-		+ SHorizontalBox::Slot()
-		.AutoWidth()
-		.Padding(8.f, 0.f, 0.f, 0.f)
-		[
-			SNew(SButton)
-			.ButtonStyle(FEditorStyle::Get(), "HoverHintOnly")
-			.ToolTipText(NSLOCTEXT("ChaosSelectedRigidBody", "Pick_UpClusterLevel", "Go to parent cluster."))
-			.ForegroundColor(FSlateColor::UseForeground())
-			.IsEnabled_Lambda([PropertyHandleActor, PropertyHandleId]()
-			{
-				return GetParentClusterRigidBodyId(PropertyHandleActor, PropertyHandleId) != INDEX_NONE;
-			})
-			.OnClicked_Lambda([PropertyHandleActor, PropertyHandleId]()
-			{
-				const int32 RigidBodyId = GetParentClusterRigidBodyId(PropertyHandleActor, PropertyHandleId);
-				if (RigidBodyId != INDEX_NONE)
+			+ SHorizontalBox::Slot()
+			.AutoWidth()
+			.Padding(8.f, 0.f, 0.f, 0.f)
+			[
+				SNew(SButton)
+				.ButtonStyle(FEditorStyle::Get(), "HoverHintOnly")
+				.ToolTipText(NSLOCTEXT("ChaosSelectedRigidBody", "Pick_UpClusterLevel", "Go to parent cluster."))
+				.ForegroundColor(FSlateColor::UseForeground())
+				.IsEnabled_Lambda([PropertyHandleActor, PropertyHandleId]()
 				{
-					PropertyHandleId->SetValue(RigidBodyId);
-				}
-				return FReply::Handled();
-			})
-			[ 
-				SNew(STextBlock)
-				.Font(FEditorStyle::Get().GetFontStyle("FontAwesome.16"))
-				.Text(FEditorFontGlyphs::Caret_Square_O_Up)
-				.ColorAndOpacity(FSlateColor::UseForeground())
+					return GetParentClusterRigidBodyId(PropertyHandleActor, PropertyHandleId) != INDEX_NONE;
+				})
+				.OnClicked_Lambda([PropertyHandleActor, PropertyHandleId]()
+				{
+					const int32 RigidBodyId = GetParentClusterRigidBodyId(PropertyHandleActor, PropertyHandleId);
+					if (RigidBodyId != INDEX_NONE)
+					{
+						PropertyHandleId->SetValue(RigidBodyId);
+					}
+					return FReply::Handled();
+				})
+				[ 
+					SNew(STextBlock)
+					.Font(FEditorStyle::Get().GetFontStyle("FontAwesome.16"))
+					.Text(FEditorFontGlyphs::Caret_Square_O_Up)
+					.ColorAndOpacity(FSlateColor::UseForeground())
+				]
 			]
-		]
-		+ SHorizontalBox::Slot()
-		.AutoWidth()
-		[
-			SNew(SButton)
-			.ButtonStyle(FEditorStyle::Get(), "HoverHintOnly")
-			.ToolTipText(NSLOCTEXT("ChaosSelectedRigidBody", "Pick_DownClusterLevel", "Go to child cluster."))
-			.ForegroundColor(FSlateColor::UseForeground())
-			.IsEnabled_Lambda([PropertyHandleActor, PropertyHandleId]()
-			{
-				return GetChildClusterRigidBodyId(PropertyHandleActor, PropertyHandleId) != INDEX_NONE;
-			})
-			.OnClicked_Lambda([PropertyHandleActor, PropertyHandleId]()
-			{
-				const int32 RigidBodyId = GetChildClusterRigidBodyId(PropertyHandleActor, PropertyHandleId);
-				if (RigidBodyId != INDEX_NONE)
+			+ SHorizontalBox::Slot()
+			.AutoWidth()
+			[
+				SNew(SButton)
+				.ButtonStyle(FEditorStyle::Get(), "HoverHintOnly")
+				.ToolTipText(NSLOCTEXT("ChaosSelectedRigidBody", "Pick_DownClusterLevel", "Go to child cluster."))
+				.ForegroundColor(FSlateColor::UseForeground())
+				.IsEnabled_Lambda([PropertyHandleActor, PropertyHandleId]()
 				{
-					PropertyHandleId->SetValue(RigidBodyId);
-				}
-				return FReply::Handled();
-			})
-			[ 
-				SNew(STextBlock)
-				.Font(FEditorStyle::Get().GetFontStyle("FontAwesome.16"))
-				.Text(FEditorFontGlyphs::Caret_Square_O_Down)
-				.ColorAndOpacity(FSlateColor::UseForeground())
+					return GetChildClusterRigidBodyId(PropertyHandleActor, PropertyHandleId) != INDEX_NONE;
+				})
+				.OnClicked_Lambda([PropertyHandleActor, PropertyHandleId]()
+				{
+					const int32 RigidBodyId = GetChildClusterRigidBodyId(PropertyHandleActor, PropertyHandleId);
+					if (RigidBodyId != INDEX_NONE)
+					{
+						PropertyHandleId->SetValue(RigidBodyId);
+					}
+					return FReply::Handled();
+				})
+				[ 
+					SNew(STextBlock)
+					.Font(FEditorStyle::Get().GetFontStyle("FontAwesome.16"))
+					.Text(FEditorFontGlyphs::Caret_Square_O_Down)
+					.ColorAndOpacity(FSlateColor::UseForeground())
+				]
 			]
-		]
-		+ SHorizontalBox::Slot()
-		.AutoWidth()
-		.Padding(8.0f, 0.f, 0.f, 0.f)
-		[
-			SNew(SButton)
-			.ButtonStyle(FEditorStyle::Get(), "HoverHintOnly")
-			.ToolTipText(NSLOCTEXT("ChaosSelectedRigidBody", "Pick_PrevClusterSibling", "Go to previous clustered sibling."))
-			.ForegroundColor(FSlateColor::UseForeground())
-			.IsEnabled_Lambda([PropertyHandleActor, PropertyHandleId]()
-			{
-				return GetPreviousClusteredSiblingRigidBodyId(PropertyHandleActor, PropertyHandleId) != INDEX_NONE;
-			})
-			.OnClicked_Lambda([PropertyHandleActor, PropertyHandleId]()
-			{
-				const int32 RigidBodyId = GetPreviousClusteredSiblingRigidBodyId(PropertyHandleActor, PropertyHandleId);
-				if (RigidBodyId != INDEX_NONE)
+			+ SHorizontalBox::Slot()
+			.AutoWidth()
+			.Padding(8.f, 0.f, 0.f, 0.f)
+			[
+				SNew(SButton)
+				.ButtonStyle(FEditorStyle::Get(), "HoverHintOnly")
+				.ToolTipText(NSLOCTEXT("ChaosSelectedRigidBody", "Pick_PrevClusterSibling", "Go to previous clustered sibling."))
+				.ForegroundColor(FSlateColor::UseForeground())
+				.IsEnabled_Lambda([PropertyHandleActor, PropertyHandleId]()
 				{
-					PropertyHandleId->SetValue(RigidBodyId);
-				}
-				return FReply::Handled();
-			})
-			[ 
-				SNew(STextBlock)
-				.Font(FEditorStyle::Get().GetFontStyle("FontAwesome.16"))
-				.Text(FEditorFontGlyphs::Caret_Square_O_Left)
-				.ColorAndOpacity(FSlateColor::UseForeground())
+					return GetPreviousClusteredSiblingRigidBodyId(PropertyHandleActor, PropertyHandleId) != INDEX_NONE;
+				})
+				.OnClicked_Lambda([PropertyHandleActor, PropertyHandleId]()
+				{
+					const int32 RigidBodyId = GetPreviousClusteredSiblingRigidBodyId(PropertyHandleActor, PropertyHandleId);
+					if (RigidBodyId != INDEX_NONE)
+					{
+						PropertyHandleId->SetValue(RigidBodyId);
+					}
+					return FReply::Handled();
+				})
+				[ 
+					SNew(STextBlock)
+					.Font(FEditorStyle::Get().GetFontStyle("FontAwesome.16"))
+					.Text(FEditorFontGlyphs::Caret_Square_O_Left)
+					.ColorAndOpacity(FSlateColor::UseForeground())
+				]
 			]
-		]
-		+ SHorizontalBox::Slot()
-		.AutoWidth()
-		[
-			SNew(SButton)
-			.ButtonStyle(FEditorStyle::Get(), "HoverHintOnly")
-			.ToolTipText(NSLOCTEXT("ChaosSelectedRigidBody", "Pick_NextClusterSibling", "Go to next clustered sibling."))
-			.ForegroundColor(FSlateColor::UseForeground())
-			.IsEnabled_Lambda([PropertyHandleActor, PropertyHandleId]()
-			{
-				return GetNextClusteredSiblingRigidBodyId(PropertyHandleActor, PropertyHandleId) != INDEX_NONE;
-			})
-			.OnClicked_Lambda([PropertyHandleActor, PropertyHandleId]()
-			{
-				const int32 RigidBodyId = GetNextClusteredSiblingRigidBodyId(PropertyHandleActor, PropertyHandleId);
-				if (RigidBodyId != INDEX_NONE)
+			+ SHorizontalBox::Slot()
+			.AutoWidth()
+			[
+				SNew(SButton)
+				.ButtonStyle(FEditorStyle::Get(), "HoverHintOnly")
+				.ToolTipText(NSLOCTEXT("ChaosSelectedRigidBody", "Pick_NextClusterSibling", "Go to next clustered sibling."))
+				.ForegroundColor(FSlateColor::UseForeground())
+				.IsEnabled_Lambda([PropertyHandleActor, PropertyHandleId]()
 				{
-					PropertyHandleId->SetValue(RigidBodyId);
-				}
-				return FReply::Handled();
-			})
-			[ 
-				SNew(STextBlock)
-				.Font(FEditorStyle::Get().GetFontStyle("FontAwesome.16"))
-				.Text(FEditorFontGlyphs::Caret_Square_O_Right)
-				.ColorAndOpacity(FSlateColor::UseForeground())
+					return GetNextClusteredSiblingRigidBodyId(PropertyHandleActor, PropertyHandleId) != INDEX_NONE;
+				})
+				.OnClicked_Lambda([PropertyHandleActor, PropertyHandleId]()
+				{
+					const int32 RigidBodyId = GetNextClusteredSiblingRigidBodyId(PropertyHandleActor, PropertyHandleId);
+					if (RigidBodyId != INDEX_NONE)
+					{
+						PropertyHandleId->SetValue(RigidBodyId);
+					}
+					return FReply::Handled();
+				})
+				[ 
+					SNew(STextBlock)
+					.Font(FEditorStyle::Get().GetFontStyle("FontAwesome.16"))
+					.Text(FEditorFontGlyphs::Caret_Square_O_Right)
+					.ColorAndOpacity(FSlateColor::UseForeground())
+				]
 			]
 		]
 	];
