Remove extraneous punctuation after the geometry collection debug draw warning message to make it more obvious that it's not waiting for something that is happening in the background.

diff --git a/Engine/Plugins/Experimental/GeometryCollectionPlugin/Source/GeometryCollectionEditor/Private/GeometryCollection/DetailCustomizations/WarningMessageCustomization.cpp b/Engine/Plugins/Experimental/GeometryCollectionPlugin/Source/GeometryCollectionEditor/Private/GeometryCollection/DetailCustomizations/WarningMessageCustomization.cpp
--- a/Engine/Plugins/Experimental/GeometryCollectionPlugin/Source/GeometryCollectionEditor/Private/GeometryCollection/DetailCustomizations/WarningMessageCustomization.cpp
+++ b/Engine/Plugins/Experimental/GeometryCollectionPlugin/Source/GeometryCollectionEditor/Private/GeometryCollection/DetailCustomizations/WarningMessageCustomization.cpp
@@ -46,7 +46,7 @@ void FWarningMessageCustomization::CustomizeHeader(TSharedRef<IPropertyHandle> S
 	if (!bHasBegunPlay)
 	{
 		static const FText ToolTipMessage = NSLOCTEXT("GeometryCollectionDebugDrawWarningMessage", "WarningMessage_ToolTip", "The actor must be playing for these debug draw properties to have any effect.");
-		static const FText Message = NSLOCTEXT("GeometryCollectionDebugDrawWarningMessage", "WarningMessage_MessageWaiting", "Waiting for Play...");
+		static const FText Message = NSLOCTEXT("GeometryCollectionDebugDrawWarningMessage", "WarningMessage_MessageWaiting", "Waiting for Play");
 		static const FSlateColor Color = FEditorStyle::GetSlateColor("SelectionColor");
 
 		// Add buttons
