Copying //UE4/Dev-Physics to Dev-Main (//UE4/Dev-Main) [at] 7703071

diff --git a/Engine/Plugins/Experimental/GeometryCollectionPlugin/Source/GeometryCollectionExampleCore/Private/GeometryCollection/GeometryCollectionExampleBoneHierarchy.cpp b/Engine/Plugins/Experimental/GeometryCollectionPlugin/Source/GeometryCollectionExampleCore/Private/GeometryCollection/GeometryCollectionExampleBoneHierarchy.cpp
--- a/Engine/Plugins/Experimental/GeometryCollectionPlugin/Source/GeometryCollectionExampleCore/Private/GeometryCollection/GeometryCollectionExampleBoneHierarchy.cpp
+++ b/Engine/Plugins/Experimental/GeometryCollectionPlugin/Source/GeometryCollectionExampleCore/Private/GeometryCollection/GeometryCollectionExampleBoneHierarchy.cpp
@@ -2,7 +2,7 @@
 
 #include "GeometryCollection/GeometryCollectionExampleBoneHierarchy.h"
 
-#include "SolverObjects/AnalyticImplicitGroup.h"
+#include "PhysicsProxy/AnalyticImplicitGroup.h"
 #include "BoneHierarchy.h"
 
 #include "Chaos/Sphere.h"
diff --git a/Engine/Plugins/Experimental/GeometryCollectionPlugin/Source/GeometryCollectionExampleCore/Private/GeometryCollection/GeometryCollectionExampleClustering.cpp b/Engine/Plugins/Experimental/GeometryCollectionPlugin/Source/GeometryCollectionExampleCore/Private/GeometryCollection/GeometryCollectionExampleClustering.cpp
--- a/Engine/Plugins/Experimental/GeometryCollectionPlugin/Source/GeometryCollectionExampleCore/Private/GeometryCollection/GeometryCollectionExampleClustering.cpp
+++ b/Engine/Plugins/Experimental/GeometryCollectionPlugin/Source/GeometryCollectionExampleCore/Private/GeometryCollection/GeometryCollectionExampleClustering.cpp
@@ -17,7 +17,7 @@
 
 
 #include "GeometryCollection/GeometryDynamicCollection.h"
-#include "SolverObjects/SolverObjects.h"
+#include "PhysicsProxy/PhysicsProxies.h"
 #include "Chaos/ErrorReporter.h"
 #include "ChaosSolversModule.h"
 #include "Chaos/PBDRigidClustering.h"
@@ -31,6 +31,7 @@ DEFINE_LOG_CATEGORY_STATIC(GCTCL_Log, Verbose, All);
 // has to go and fix up so many callsites here and they're all pretty much
 // Identical. The similar code should be pulled out
 
+#if TODO_REIMPLEMENT_RIGID_CLUSTERING
 namespace GeometryCollectionExample
 {
 #if INCLUDE_CHAOS
@@ -86,7 +87,7 @@ namespace GeometryCollectionExample
 			BuildSimulationData(ErrorReporter, *RestCollection, InParams.Shared);
 		};
 
-		FGeometryCollectionPhysicsObject* PhysObject = new FGeometryCollectionPhysicsObject(nullptr, DynamicCollection.Get(), InitFunc, nullptr, nullptr);
+		FGeometryCollectionPhysicsProxy* PhysObject = new FGeometryCollectionPhysicsProxy(nullptr, DynamicCollection.Get(), InitFunc, nullptr, nullptr);
 		PhysObject->Initialize();
 
 		Chaos::FPBDRigidsSolver* Solver = FChaosSolversModule::GetModule()->CreateSolver(true);
@@ -178,7 +179,7 @@ namespace GeometryCollectionExample
 			BuildSimulationData(ErrorReporter, *RestCollection, InParams.Shared);
 		};
 
-		FGeometryCollectionPhysicsObject* PhysObject = new FGeometryCollectionPhysicsObject(nullptr, DynamicCollection.Get(), InitFunc, nullptr, nullptr);
+		FGeometryCollectionPhysicsProxy* PhysObject = new FGeometryCollectionPhysicsProxy(nullptr, DynamicCollection.Get(), InitFunc, nullptr, nullptr);
 		PhysObject->Initialize();
 
 		Chaos::FPBDRigidsSolver* Solver = FChaosSolversModule::GetModule()->CreateSolver(true);
@@ -324,7 +325,7 @@ namespace GeometryCollectionExample
 			BuildSimulationData(ErrorReporter, *RestCollection, InParams.Shared);
 		};
 
-		FGeometryCollectionPhysicsObject* PhysObject = new FGeometryCollectionPhysicsObject(nullptr, DynamicCollection.Get(), InitFunc, nullptr, nullptr);
+		FGeometryCollectionPhysicsProxy* PhysObject = new FGeometryCollectionPhysicsProxy(nullptr, DynamicCollection.Get(), InitFunc, nullptr, nullptr);
 		PhysObject->Initialize();
 		PhysObject->SetCollisionParticlesPerObjectFraction(1.0);
 
@@ -454,7 +455,7 @@ namespace GeometryCollectionExample
 			BuildSimulationData(ErrorReporter, *RestCollection, InParams.Shared);
 		};
 		
-		FGeometryCollectionPhysicsObject* PhysObject = new FGeometryCollectionPhysicsObject(nullptr, DynamicCollection.Get(), InitFunc, nullptr, nullptr);
+		FGeometryCollectionPhysicsProxy* PhysObject = new FGeometryCollectionPhysicsProxy(nullptr, DynamicCollection.Get(), InitFunc, nullptr, nullptr);
 		PhysObject->Initialize();
 
 		Chaos::FPBDRigidsSolver* Solver = FChaosSolversModule::GetModule()->CreateSolver(true);
@@ -586,7 +587,7 @@ namespace GeometryCollectionExample
 			BuildSimulationData(ErrorReporter, *RestCollection, InParams.Shared);
 		};
 
-		FGeometryCollectionPhysicsObject* PhysObject = new FGeometryCollectionPhysicsObject(nullptr, DynamicCollection.Get(), InitFunc, nullptr, nullptr);;
+		FGeometryCollectionPhysicsProxy* PhysObject = new FGeometryCollectionPhysicsProxy(nullptr, DynamicCollection.Get(), InitFunc, nullptr, nullptr);;
 		PhysObject->Initialize();
 
 		Chaos::FPBDRigidsSolver* Solver = FChaosSolversModule::GetModule()->CreateSolver(true);
@@ -744,7 +745,7 @@ namespace GeometryCollectionExample
 			BuildSimulationData(ErrorReporter, *RestCollection, InParams.Shared);
 		};
 
-		FGeometryCollectionPhysicsObject* PhysObject = new FGeometryCollectionPhysicsObject(nullptr, DynamicCollection.Get(), InitFunc, nullptr, nullptr);;
+		FGeometryCollectionPhysicsProxy* PhysObject = new FGeometryCollectionPhysicsProxy(nullptr, DynamicCollection.Get(), InitFunc, nullptr, nullptr);;
 		PhysObject->Initialize();
 
 		Chaos::FPBDRigidsSolver* Solver = FChaosSolversModule::GetModule()->CreateSolver(true);
@@ -877,7 +878,7 @@ namespace GeometryCollectionExample
 			BuildSimulationData(ErrorReporter, *RestCollection, InParams.Shared);
 		};
 
-		FGeometryCollectionPhysicsObject* PhysObject = new FGeometryCollectionPhysicsObject(nullptr, DynamicCollection.Get(), InitFunc, nullptr, nullptr);;
+		FGeometryCollectionPhysicsProxy* PhysObject = new FGeometryCollectionPhysicsProxy(nullptr, DynamicCollection.Get(), InitFunc, nullptr, nullptr);;
 		PhysObject->Initialize();
 
 		Chaos::FPBDRigidsSolver* Solver = FChaosSolversModule::GetModule()->CreateSolver(true);
@@ -1115,7 +1116,7 @@ namespace GeometryCollectionExample
 			BuildSimulationData(ErrorReporter, *RestCollection, InParams.Shared);
 		};
 
-		FGeometryCollectionPhysicsObject* PhysObject = new FGeometryCollectionPhysicsObject(nullptr, DynamicCollection.Get(), InitFunc, nullptr, nullptr);;
+		FGeometryCollectionPhysicsProxy* PhysObject = new FGeometryCollectionPhysicsProxy(nullptr, DynamicCollection.Get(), InitFunc, nullptr, nullptr);;
 		PhysObject->Initialize();
 
 		Chaos::FPBDRigidsSolver* Solver = FChaosSolversModule::GetModule()->CreateSolver(true);
@@ -1335,10 +1336,10 @@ namespace GeometryCollectionExample
 			BuildSimulationData(ErrorReporter, *RestCollection2, InParams.Shared);
 		};
 
-		FGeometryCollectionPhysicsObject* PhysObject = new FGeometryCollectionPhysicsObject(nullptr, DynamicCollection.Get(), InitFunc, nullptr, nullptr);;
+		FGeometryCollectionPhysicsProxy* PhysObject = new FGeometryCollectionPhysicsProxy(nullptr, DynamicCollection.Get(), InitFunc, nullptr, nullptr);;
 		PhysObject->Initialize();
 
-		FGeometryCollectionPhysicsObject* PhysObject2 = new FGeometryCollectionPhysicsObject(nullptr, DynamicCollection2.Get(), InitFunc2, nullptr, nullptr);;
+		FGeometryCollectionPhysicsProxy* PhysObject2 = new FGeometryCollectionPhysicsProxy(nullptr, DynamicCollection2.Get(), InitFunc2, nullptr, nullptr);;
 		PhysObject2->Initialize();
 
 
@@ -1473,10 +1474,10 @@ namespace GeometryCollectionExample
 			BuildSimulationData(ErrorReporter, *RestCollection2, InParams.Shared);
 		};
 
-		FGeometryCollectionPhysicsObject* PhysObject = new FGeometryCollectionPhysicsObject(nullptr, DynamicCollection.Get(), InitFunc, nullptr, nullptr);;
+		FGeometryCollectionPhysicsProxy* PhysObject = new FGeometryCollectionPhysicsProxy(nullptr, DynamicCollection.Get(), InitFunc, nullptr, nullptr);;
 		PhysObject->Initialize();
 
-		FGeometryCollectionPhysicsObject* PhysObject2 = new FGeometryCollectionPhysicsObject(nullptr, DynamicCollection2.Get(), InitFunc2, nullptr, nullptr);;
+		FGeometryCollectionPhysicsProxy* PhysObject2 = new FGeometryCollectionPhysicsProxy(nullptr, DynamicCollection2.Get(), InitFunc2, nullptr, nullptr);;
 		PhysObject2->Initialize();
 
 
@@ -1610,10 +1611,10 @@ namespace GeometryCollectionExample
 			BuildSimulationData(ErrorReporter, *RestCollection2, InParams.Shared);
 		};
 
-		FGeometryCollectionPhysicsObject* PhysObject = new FGeometryCollectionPhysicsObject(nullptr, DynamicCollection.Get(), InitFunc, nullptr, nullptr);;
+		FGeometryCollectionPhysicsProxy* PhysObject = new FGeometryCollectionPhysicsProxy(nullptr, DynamicCollection.Get(), InitFunc, nullptr, nullptr);;
 		PhysObject->Initialize();
 
-		FGeometryCollectionPhysicsObject* PhysObject2 = new FGeometryCollectionPhysicsObject(nullptr, DynamicCollection2.Get(), InitFunc2, nullptr, nullptr);;
+		FGeometryCollectionPhysicsProxy* PhysObject2 = new FGeometryCollectionPhysicsProxy(nullptr, DynamicCollection2.Get(), InitFunc2, nullptr, nullptr);;
 		PhysObject2->Initialize();
 
 		Chaos::FPBDRigidsSolver* Solver = FChaosSolversModule::GetModule()->CreateSolver(true);
@@ -1719,7 +1720,7 @@ namespace GeometryCollectionExample
 			BuildSimulationData(ErrorReporter, *RestCollection, InParams.Shared);
 		};
 
-		FGeometryCollectionPhysicsObject* PhysObject = new FGeometryCollectionPhysicsObject(nullptr, DynamicCollection.Get(), InitFunc, nullptr, nullptr);;
+		FGeometryCollectionPhysicsProxy* PhysObject = new FGeometryCollectionPhysicsProxy(nullptr, DynamicCollection.Get(), InitFunc, nullptr, nullptr);;
 		PhysObject->Initialize();
 
 		Chaos::FPBDRigidsSolver* Solver = FChaosSolversModule::GetModule()->CreateSolver(true);
@@ -1808,7 +1809,7 @@ namespace GeometryCollectionExample
 		};
 
 
-		FGeometryCollectionPhysicsObject* PhysObject = new FGeometryCollectionPhysicsObject(nullptr, DynamicCollection.Get(), InitFunc, nullptr, nullptr);;
+		FGeometryCollectionPhysicsProxy* PhysObject = new FGeometryCollectionPhysicsProxy(nullptr, DynamicCollection.Get(), InitFunc, nullptr, nullptr);;
 		PhysObject->Initialize();
 
 
@@ -1915,9 +1916,9 @@ namespace GeometryCollectionExample
 		FalloffField->Position = FVector(0.0, 0.0, 0.0);
 		FalloffField->Falloff = EFieldFalloffType::Field_FallOff_None;
 
-		FFieldSystemPhysicsObject* FieldObject = new FFieldSystemPhysicsObject(nullptr);
+		FFieldSystemPhysicsProxy* FieldObject = new FFieldSystemPhysicsProxy(nullptr);
 
-		FGeometryCollectionPhysicsObject* PhysObject = new FGeometryCollectionPhysicsObject(nullptr, DynamicCollection.Get(), InitFunc, nullptr, nullptr);
+		FGeometryCollectionPhysicsProxy* PhysObject = new FGeometryCollectionPhysicsProxy(nullptr, DynamicCollection.Get(), InitFunc, nullptr, nullptr);
 		Chaos::FPBDRigidsSolver* Solver = FChaosSolversModule::GetModule()->CreateSolver(true);
 		Chaos::FPBDRigidsSolver::FClusteringType & Clustering = Solver->GetRigidClustering();
 
@@ -1988,20 +1989,20 @@ namespace GeometryCollectionExample
 		P.SizeData.MaxLevelSetResolution = 20;
 
 		SimulationObjects<T>* Object = new SimulationObjects<T>(P, CreateClusteredBody_FracturedGeometry());
-		Object->PhysicsObject->Initialize();
-		Object->PhysicsObject->ActivateBodies();
+		Object->PhysicsProxy->Initialize();
+		Object->PhysicsProxy->ActivateBodies();
 
 		typedef TUniquePtr<Chaos::TImplicitObject<float, 3>> FImplicitPointer;
-		const TManagedArray<FImplicitPointer> & Implicits = Object->RestCollection->template GetAttribute<FImplicitPointer>(FGeometryCollectionPhysicsObject::ImplicitsAttribute, FTransformCollection::TransformGroup);
+		const TManagedArray<FImplicitPointer> & Implicits = Object->RestCollection->template GetAttribute<FImplicitPointer>(FGeometryCollectionPhysicsProxy::ImplicitsAttribute, FTransformCollection::TransformGroup);
 
 		typedef TUniquePtr< FCollisionStructureManager::FSimplicial > FSimplicialPointer;
-		const TManagedArray<FSimplicialPointer> & Simplicials = Object->RestCollection->template GetAttribute<FSimplicialPointer>(FGeometryCollectionPhysicsObject::SimplicialsAttribute, FTransformCollection::TransformGroup);
+		const TManagedArray<FSimplicialPointer> & Simplicials = Object->RestCollection->template GetAttribute<FSimplicialPointer>(FGeometryCollectionPhysicsProxy::SimplicialsAttribute, FTransformCollection::TransformGroup);
 
 
 		Chaos::FPBDRigidsSolver* Solver = FChaosSolversModule::GetModule()->CreateSolver(true);
 		Solver->SetHasFloor(false);
 		Solver->SetEnabled(true);
-		Solver->RegisterObject(Object->PhysicsObject.Get());
+		Solver->RegisterObject(Object->PhysicsProxy.Get());
 		Solver->AdvanceSolverBy(1 / 24.);
 		Chaos::TPBDRigidParticles<float, 3>& Particles = Solver->GetRigidParticles();
 
@@ -2013,9 +2014,9 @@ namespace GeometryCollectionExample
 			CollisionParticlesPerObjectFractionDefault = CVarCollisionParticlesPerObjectFractionDefault->GetFloat();
 		}
 
-		R.ExpectTrue(Particles.CollisionParticles(Object->PhysicsObject->RigidBodyIDArray_TestingAccess()[10])->Size() == (int)(Simplicials[10]->Size() * CollisionParticlesPerObjectFractionDefault));
-		R.ExpectTrue(Particles.CollisionParticles(Object->PhysicsObject->RigidBodyIDArray_TestingAccess()[11])->Size() == (int)(Simplicials[11]->Size() * CollisionParticlesPerObjectFractionDefault));
-		R.ExpectTrue(Particles.CollisionParticles(Object->PhysicsObject->RigidBodyIDArray_TestingAccess()[12])->Size() == (int)(Simplicials[12]->Size() * CollisionParticlesPerObjectFractionDefault));
+		R.ExpectTrue(Particles.CollisionParticles(Object->PhysicsProxy->RigidBodyIDArray_TestingAccess()[10])->Size() == (int)(Simplicials[10]->Size() * CollisionParticlesPerObjectFractionDefault));
+		R.ExpectTrue(Particles.CollisionParticles(Object->PhysicsProxy->RigidBodyIDArray_TestingAccess()[11])->Size() == (int)(Simplicials[11]->Size() * CollisionParticlesPerObjectFractionDefault));
+		R.ExpectTrue(Particles.CollisionParticles(Object->PhysicsProxy->RigidBodyIDArray_TestingAccess()[12])->Size() == (int)(Simplicials[12]->Size() * CollisionParticlesPerObjectFractionDefault));
 
 		// cleanup
 		//for (auto Obj : Collections) delete Obj;
@@ -2027,3 +2028,4 @@ namespace GeometryCollectionExample
 	}
 	template bool RigidBodiess_ClusterTest_ParticleImplicitCollisionGeometry<float>(ExampleResponse&& R);
 }
+#endif
diff --git a/Engine/Plugins/Experimental/GeometryCollectionPlugin/Source/GeometryCollectionExampleCore/Private/GeometryCollection/GeometryCollectionExampleEvents.cpp b/Engine/Plugins/Experimental/GeometryCollectionPlugin/Source/GeometryCollectionExampleCore/Private/GeometryCollection/GeometryCollectionExampleEvents.cpp
new file mode 100644
--- /dev/null
+++ b/Engine/Plugins/Experimental/GeometryCollectionPlugin/Source/GeometryCollectionExampleCore/Private/GeometryCollection/GeometryCollectionExampleEvents.cpp
@@ -0,0 +1,182 @@
+// Copyright 1998-2019 Epic Games, Inc. All Rights Reserved.
+
+#include "GeometryCollection/GeometryCollectionExampleEvents.h"
+#include "EventManager.h"
+#include "EventsData.h"
+#include "PBDRigidsSolver.h"
+#include "ChaosSolversModule.h"
+
+namespace GeometryCollectionExample
+{
+	// deliberately choosing values outside EEventType defaults
+	static const int32 CustomEvent1 = 5;
+	static const int32 CustomEvent2 = 6;
+
+	struct EventTestData
+	{
+		EventTestData(int InData1, FVector InData2) : Data1(InData1), Data2(InData2) {}
+		EventTestData() {}
+
+		bool operator==(const EventTestData& Other) const
+		{
+			return (Data1 == Other.Data1) && (Data2 == Other.Data2);
+		}
+
+		int Data1;
+		FVector Data2;
+	};
+
+	class MyEventHandler
+	{
+	public:
+
+		MyEventHandler(Chaos::FEventManager& InEventManager) : EventManager(InEventManager)
+		{
+		}
+		~MyEventHandler()
+		{
+			EventManager.Reset();
+		}
+
+		void HandleEvent(const EventTestData& EventData)
+		{
+			ResultFromHandler = EventData;
+		}
+
+		void HandleEvent(const TArray<EventTestData>& EventData)
+		{
+			ResultFromHandler2 = EventData;
+		}
+
+		void RegisterHandler1()
+		{
+			EventManager.RegisterHandler<EventTestData>(CustomEvent1, this, &MyEventHandler::HandleEvent);
+		}
+
+		void RegisterHandler2()
+		{
+			EventManager.RegisterHandler<TArray<EventTestData>>(CustomEvent2, this, &MyEventHandler::HandleEvent);
+		}
+
+		void UnregisterHandler1()
+		{
+			EventManager.UnregisterHandler(CustomEvent1, this);
+		}
+
+		void UnregisterHandler2()
+		{
+			EventManager.UnregisterHandler(CustomEvent2, this);
+		}
+
+		// Prove data is reference to original and not copy passed to the event handler
+		EventTestData ResultFromHandler;
+
+		// the data dispatched can also be a frames worth of events
+		TArray<EventTestData> ResultFromHandler2;
+
+	private:
+		Chaos::FEventManager& EventManager;
+	};
+
+
+	template<class T>
+	bool Event_Handler(ExampleResponse&& R)
+	{
+#if INCLUDE_CHAOS
+		Chaos::FEventManager EventManager(Chaos::EMultiBufferMode::Single);
+		Chaos::FPBDRigidsSolver* Solver = FChaosSolversModule::GetModule()->CreateSolver(true);
+
+		MyEventHandler HandlerTest(EventManager);
+		MyEventHandler AnotherHandlerTest(EventManager);
+
+		// the data injected into the buffer for CustomEvent1 will be whatever is currently the variable TestData
+		EventTestData TestData;
+		EventTestData* TestDataPtr = &TestData;
+		EventManager.RegisterEvent<EventTestData>(CustomEvent1, [TestDataPtr]
+		(const Chaos::FPBDRigidsSolver* Solver, EventTestData& MyData)
+		{
+			MyData = *TestDataPtr;
+		});
+
+		// the data injected into the buffer for CustomEvent2 will be whatever is currently the variable TestArrayData
+		TArray<EventTestData> TestArrayData;
+		TArray<EventTestData>* TestDataPtr2 = &TestArrayData;
+		EventManager.RegisterEvent<TArray<EventTestData>>(CustomEvent2, [TestDataPtr2]
+		(const Chaos::FPBDRigidsSolver* Solver, TArray<EventTestData>& MyData)
+		{
+			MyData = *TestDataPtr2;
+		});
+
+		HandlerTest.RegisterHandler1();
+		HandlerTest.RegisterHandler2();
+		AnotherHandlerTest.RegisterHandler2();
+
+		TestData.Data1 = 123;
+		TestData.Data2 = FVector(1, 2, 3);
+
+		EventManager.FillProducerData(Solver);
+		EventManager.FlipBuffersIfRequired();
+		EventManager.DispatchEvents();
+		R.ExpectTrue(HandlerTest.ResultFromHandler == TestData);
+
+		TestData.Data1 = 789;
+		TestData.Data2 = FVector(7, 8, 9);
+
+		EventManager.FillProducerData(Solver);
+		EventManager.FlipBuffersIfRequired();
+		EventManager.DispatchEvents();
+		R.ExpectTrue(HandlerTest.ResultFromHandler == TestData);
+
+		// Unregister - data should no longer update
+		HandlerTest.UnregisterHandler1();
+
+		EventTestData OriginalTestData = TestData;
+		TestData.Data1 = 999;
+		TestData.Data2 = FVector(9, 9, 9);
+
+		EventManager.FillProducerData(Solver);
+		EventManager.FlipBuffersIfRequired();
+		EventManager.DispatchEvents();
+		R.ExpectTrue(HandlerTest.ResultFromHandler == OriginalTestData);
+
+		EventTestData Data;
+		TestArrayData.Push(EventTestData(123, FVector(1, 2, 3)));
+		TestArrayData.Push(EventTestData(456, FVector(4, 5, 6)));
+		TestArrayData.Push(EventTestData(789, FVector(7, 8, 9)));
+
+		EventManager.FillProducerData(Solver);
+		EventManager.FlipBuffersIfRequired();
+		EventManager.DispatchEvents();
+		// dispatched to multiple handlers
+		R.ExpectTrue(HandlerTest.ResultFromHandler2 == TestArrayData);
+		R.ExpectTrue(AnotherHandlerTest.ResultFromHandler2 == TestArrayData);
+
+		// Unregister one of the handlers from the events
+		HandlerTest.UnregisterHandler2();
+
+		TArray<EventTestData> OriginalTestArrayData = TestArrayData;
+		TestArrayData.Reset();
+		TestArrayData.Push(EventTestData(999, FVector(9, 9, 9)));
+
+		EventManager.FillProducerData(Solver);
+		EventManager.FlipBuffersIfRequired();
+		EventManager.DispatchEvents();
+		R.ExpectTrue(HandlerTest.ResultFromHandler2 == OriginalTestArrayData); // Unregistered - data should no longer update
+		R.ExpectTrue(AnotherHandlerTest.ResultFromHandler2 == TestArrayData); // Still registered so should get updates
+
+		HandlerTest.RegisterHandler2();
+
+		EventManager.FillProducerData(Solver);
+		EventManager.FlipBuffersIfRequired();
+		EventManager.DispatchEvents();
+		R.ExpectTrue(HandlerTest.ResultFromHandler2 == TestArrayData);
+		R.ExpectTrue(AnotherHandlerTest.ResultFromHandler2 == TestArrayData);
+#endif
+		return !R.HasError();
+
+	}
+	template bool Event_Handler<float>(ExampleResponse&& R);
+
+}
+
+
diff --git a/Engine/Plugins/Experimental/GeometryCollectionPlugin/Source/GeometryCollectionExampleCore/Private/GeometryCollection/GeometryCollectionExampleImplicitSphere.cpp b/Engine/Plugins/Experimental/GeometryCollectionPlugin/Source/GeometryCollectionExampleCore/Private/GeometryCollection/GeometryCollectionExampleImplicitSphere.cpp
--- a/Engine/Plugins/Experimental/GeometryCollectionPlugin/Source/GeometryCollectionExampleCore/Private/GeometryCollection/GeometryCollectionExampleImplicitSphere.cpp
+++ b/Engine/Plugins/Experimental/GeometryCollectionPlugin/Source/GeometryCollectionExampleCore/Private/GeometryCollection/GeometryCollectionExampleImplicitSphere.cpp
@@ -2,20 +2,334 @@
 
 #include "GeometryCollection/GeometryCollectionExampleImplicitSphere.h"
 
+#include "Chaos/Pair.h"
 #include "Chaos/Sphere.h"
 
+#include <string>
+
 namespace GeometryCollectionExample
 {
+	//==========================================================================
+	// FindClosestIntersection() tests
+	//==========================================================================
+
+	typedef Chaos::Pair<Chaos::TVector<float, 3>, bool> IntersectionResult;
+
+	template <class T>
+	void RunTestFindClosestIntersectionHelper(
+		ExampleResponse& R, 
+		const FString TestName,
+		const Chaos::TVector<T, 3>& Center, 
+		const T Radius,
+		const T Thickness,
+		const Chaos::TVector<T, 3>& StartPt,
+		const Chaos::TVector<T, 3>& EndPt,
+		const IntersectionResult ExpectedRes)
+	{
+		Chaos::TSphere<T, 3> Sphere(Center, Radius);
+		IntersectionResult Res = Sphere.FindClosestIntersection(StartPt, EndPt, Thickness);
+		R.ExpectTrue(
+			(!Res.Second && !ExpectedRes.Second) ||
+			(Res.Second && ExpectedRes.Second && (Res.First - ExpectedRes.First).Size() < SMALL_NUMBER),
+			FString::Printf(
+				TEXT(
+					"%s - Sphere((%g, %g, %g), %g).FindClosestIntersection("
+						"StartPt(%g, %g, %g), "
+						"EndPt(%g, %g, %g), "
+						"Thickness:%g) = "
+					"RESULT: %d, Point: (%g, %g, %g), "
+					"EXPECTED: %d (%g, %g, %g)."),
+				*TestName,
+				Center[0], Center[1], Center[2],
+				Radius,
+
+				StartPt[0], StartPt[1], StartPt[2],
+				EndPt[0], EndPt[1], EndPt[2],
+				Thickness,
+				
+				Res.Second ? 1 : 0,
+				Res.First[0], Res.First[1], Res.First[2],
+				
+				ExpectedRes.Second ? 1 : 0,
+				ExpectedRes.First[0], ExpectedRes.First[1], ExpectedRes.First[2]
+			));
+	}
+
+	template <class T>
+	void RunTestFindClosestIntersection_Tangent(
+		ExampleResponse& R,
+		const Chaos::TVector<T, 3>& Center,
+		const T Radius,
+		const T Thickness)
+	{
+		Chaos::TVector<T, 3> StartPt;
+		Chaos::TVector<T, 3> EndPt;
+
+		// Start & End points coincident, lying on the surface of the sphere w/zero thickness
+		// Fails with no collision:
+		// "RunTestFindClosestIntersection_Tangent1 - Sphere((0, 0, 0), 1).FindClosestIntersection(StartPt(0, 0, 1), EndPt(0, 0, 1), Thickness:0.1) = RESULT: 0, Point: (0, 0, 0), EXPECTED: 1 (0, 0, 1.1)."	FString
+		StartPt = EndPt = Center + Chaos::TVector<T,3>(0, 0, Radius);
+		//RunTestFindClosestIntersectionHelper(R, "RunTestFindClosestIntersection_Tangent1", Center, Radius, Thickness,
+		//	StartPt, EndPt, 
+		//	IntersectionResult(Center+Chaos::TVector<T,3>(0,0,Radius+Thickness), true));
+
+		StartPt = Center + Chaos::TVector<T, 3>(Radius, 0, Radius+Thickness);
+		EndPt = Center + Chaos::TVector<T, 3>(-Radius, 0, Radius+Thickness);
+		RunTestFindClosestIntersectionHelper(R, "RunTestFindClosestIntersection_Tangent2", Center, Radius, Thickness,
+			StartPt, EndPt, 
+			IntersectionResult(Center+Chaos::TVector<T,3>(0,0,Radius+Thickness), true));
+		
+		// Start & End points coincident, lying on the surface of the sphere w/thickness
+		// Fails with incorrect collision point:
+		// "RunTestFindClosestIntersection_Tangent3 - Sphere((0, 0, 0), 1).FindClosestIntersection(StartPt(0, 0, 1.1), EndPt(0, 0, 1.1), Thickness:0.1) = RESULT: 1, Point: (0, 0, 1), EXPECTED: 1 (0, 0, 1.1)."	FString
+		// "RunTestFindClosestIntersection_Tangent3 - Sphere((0, 0, 0), 1).FindClosestIntersection(StartPt(0, 0, 1.5), EndPt(0, 0, 1.5), Thickness:0.5) = RESULT: 1, Point: (0, 0, 1), EXPECTED: 1 (0, 0, 1.5)."	FString
+		// "RunTestFindClosestIntersection_Tangent3 - Sphere((0, 0, 0), 1).FindClosestIntersection(StartPt(0, 0, 2), EndPt(0, 0, 2), Thickness:1) = RESULT: 1, Point: (0, 0, 1), EXPECTED: 1 (0, 0, 2)."	FString
+		// "RunTestFindClosestIntersection_Tangent3 - Sphere((0, 0, 0), 1).FindClosestIntersection(StartPt(0, 0, 3), EndPt(0, 0, 3), Thickness:2) = RESULT: 1, Point: (0, 0, 1), EXPECTED: 1 (0, 0, 3)."	FString
+		StartPt = EndPt = Center + Chaos::TVector<T,3>(0, 0, Radius+Thickness);
+		//RunTestFindClosestIntersectionHelper(R, "RunTestFindClosestIntersection_Tangent3", Center, Radius, Thickness,
+		//	StartPt, EndPt, 
+		//	IntersectionResult(Center+Chaos::TVector<T,3>(0,0,Radius+Thickness), true));
+
+		StartPt = Center + Chaos::TVector<T, 3>(Radius, 0, Radius+Thickness);
+		EndPt = Center + Chaos::TVector<T, 3>(-Radius, 0, Radius+Thickness);
+		RunTestFindClosestIntersectionHelper(R, "RunTestFindClosestIntersection_Tangent4", Center, Radius, Thickness,
+			StartPt, EndPt, 
+			IntersectionResult(Center+Chaos::TVector<T,3>(0,0,Radius+Thickness), true));
+
+		// Radius + Thickness + 1 (miss)
+		StartPt = EndPt = Center + Chaos::TVector<T,3>(0, 0, Radius+Thickness+1.0);
+		RunTestFindClosestIntersectionHelper(R, "RunTestFindClosestIntersection_Tangent5", Center, Radius, Thickness,
+			StartPt, EndPt, 
+			IntersectionResult(Center+Chaos::TVector<T,3>(Radius+Thickness,0,0), false));
+
+		StartPt = Center + Chaos::TVector<T, 3>(Radius, 0, Radius+Thickness+1.0);
+		EndPt = Center + Chaos::TVector<T, 3>(-Radius, 0, Radius+Thickness+1.0);
+		RunTestFindClosestIntersectionHelper(R, "RunTestFindClosestIntersection_Tangent6", Center, Radius, Thickness,
+			StartPt, EndPt, 
+			IntersectionResult(Center+Chaos::TVector<T,3>(0,0,Radius+Thickness), false));
+	}
+
+	template <class T>
+	void RunTestFindClosestIntersection_Shallow(
+		ExampleResponse& R,
+		const Chaos::TVector<T, 3>& Center,
+		const T Radius,
+		const T Thickness)
+	{
+		Chaos::TVector<T, 3> StartPt(0, 0, Radius * .999999);
+		Chaos::TVector<T, 3> EndPt = StartPt;
+
+		// Fails w/no collision:
+		// "RunTestFindClosestIntersection_Shallow1 - Sphere((0, 0, 0), 1).FindClosestIntersection(StartPt(0, 0, 0.999999), EndPt(0, 0, 0.999999), Thickness:0.1) = RESULT: 0, Point: (0, 0, 0), EXPECTED: 1 (1.1, 0...	FString
+		//RunTestFindClosestIntersectionHelper(R, "RunTestFindClosestIntersection_Shallow1", Center, Radius, Thickness,
+		//	StartPt, EndPt, 
+		//	IntersectionResult(Center+Chaos::TVector<T,3>(Radius+Thickness,0,0), true));
+
+		// Fails w/no collision:
+		// "RunTestFindClosestIntersection_Shallow2 - Sphere((0, 0, 0), 1).FindClosestIntersection(StartPt(0, 0, 0.999999), EndPt(0, 0, 1), Thickness:0.1) = RESULT: 0, Point: (0, 0, 0), EXPECTED: 1 (0, 0, 1.1)."	FString
+		EndPt += Chaos::TVector<T, 3>(0, 0, Radius * .000001);
+		//RunTestFindClosestIntersectionHelper(R, "RunTestFindClosestIntersection_Shallow2", Center, Radius, Thickness,
+		//	StartPt, EndPt, 
+		//	IntersectionResult(Center+Chaos::TVector<T,3>(0,0,Radius+Thickness), true));
+
+		// Fails w/no collision:
+		// "RunTestFindClosestIntersection_Shallow3 - Sphere((0, 0, 0), 1).FindClosestIntersection(StartPt(0, 0, 0.999999), EndPt(0, 0, 1), Thickness:0.1) = RESULT: 0, Point: (0, 0, 0), EXPECTED: 1 (0, 0, 1.1)."	FString
+		EndPt += Chaos::TVector<T, 3>(0, 0, Radius * .000001);
+		//RunTestFindClosestIntersectionHelper(R, "RunTestFindClosestIntersection_Shallow3", Center, Radius, Thickness,
+		//	StartPt, EndPt, 
+		//	IntersectionResult(Center+Chaos::TVector<T,3>(0,0,Radius+Thickness), true));
+
+		EndPt += Chaos::TVector<T, 3>(0, 0, Thickness);
+		RunTestFindClosestIntersectionHelper(R, "RunTestFindClosestIntersection_Shallow3.1", Center, Radius, Thickness,
+			StartPt, EndPt, 
+			IntersectionResult(Center+Chaos::TVector<T,3>(0,0,Radius+Thickness), true));
+
+		// Fails w/no collision:
+		// "RunTestFindClosestIntersection_Shallow4 - Sphere((0, 0, 0), 1).FindClosestIntersection(StartPt(0, 0, 1), EndPt(0, 0, 0.999999), Thickness:0.1) = RESULT: 0, Point: (0, 0, 0), EXPECTED: 1 (0, 0, 1.1)."	FString
+		EndPt = StartPt;
+		StartPt = EndPt + Chaos::TVector<T, 3>(0, 0, Radius * .000001);
+		//RunTestFindClosestIntersectionHelper(R, "RunTestFindClosestIntersection_Shallow4", Center, Radius, Thickness,
+		//	StartPt, EndPt, 
+		//	IntersectionResult(Center+Chaos::TVector<T,3>(0,0,Radius+Thickness), true));
+
+		// Fails w/no collision:
+		// "RunTestFindClosestIntersection_Shallow5 - Sphere((0, 0, 0), 1).FindClosestIntersection(StartPt(0, 0, 1), EndPt(0, 0, 0.999999), Thickness:0.1) = RESULT: 0, Point: (0, 0, 0), EXPECTED: 1 (0, 0, 1.1)."	FString
+		StartPt += Chaos::TVector<T, 3>(0, 0, Radius * .000001);
+		//RunTestFindClosestIntersectionHelper(R, "RunTestFindClosestIntersection_Shallow5", Center, Radius, Thickness,
+		//	StartPt, EndPt, 
+		//	IntersectionResult(Center+Chaos::TVector<T,3>(0,0,Radius+Thickness), true));
+
+		StartPt += Chaos::TVector<T, 3>(0, 0, Thickness);
+		RunTestFindClosestIntersectionHelper(R, "RunTestFindClosestIntersection_Shallow5.1", Center, Radius, Thickness,
+			StartPt, EndPt, 
+			IntersectionResult(Center+Chaos::TVector<T,3>(0,0,Radius+Thickness), true));
+	}
+
+	template <class T>
+	void RunTestFindClosestIntersection_Mid(
+		ExampleResponse& R, 
+		const Chaos::TVector<T, 3>& Center, 
+		const T Radius,
+		const T Thickness)
+	{
+		Chaos::TVector<T, 3> StartPt(0,0,Radius * .5);
+		Chaos::TVector<T, 3> EndPt = StartPt;
+
+		// Fails w/no collision:
+		// "RunTestFindClosestIntersection_Mid1 - Sphere((0, 0, 0), 1).FindClosestIntersection(StartPt(0, 0, 0.5), EndPt(0, 0, 0.5), Thickness:0.1) = RESULT: 0, Point: (0, 0, 0), EXPECTED: 1 (1.1, 0, 0)."	FString
+		//RunTestFindClosestIntersectionHelper(R, "RunTestFindClosestIntersection_Mid1", Center, Radius, Thickness,
+		//	StartPt, EndPt, 
+		//	IntersectionResult(Center+Chaos::TVector<T,3>(Radius+Thickness,0,0), true));
+
+		// Fails w/no collision:
+		// "RunTestFindClosestIntersection_Mid2 - Sphere((0, 0, 0), 1).FindClosestIntersection(StartPt(0, 0, 0.5), EndPt(0, 0, 1), Thickness:0.1) = RESULT: 0, Point: (0, 0, 0), EXPECTED: 1 (0, 0, 1.1)."	FString
+		//EndPt += Chaos::TVector<T, 3>(0, 0, Radius * .5);
+		//RunTestFindClosestIntersectionHelper(R, "RunTestFindClosestIntersection_Mid2", Center, Radius, Thickness,
+		//	StartPt, EndPt, 
+		//	IntersectionResult(Center+Chaos::TVector<T,3>(0,0,Radius+Thickness), true));
+		
+		// Fails w/no collision:
+		// "RunTestFindClosestIntersection_Mid3 - Sphere((0, 0, 0), 1).FindClosestIntersection(StartPt(0, 0, 0.5), EndPt(0, 0, 1), Thickness:0.1) = RESULT: 0, Point: (0, 0, 0), EXPECTED: 1 (0, 0, 1.1)."	FString
+		EndPt += Chaos::TVector<T, 3>(0, 0, Radius * .5);
+		//RunTestFindClosestIntersectionHelper(R, "RunTestFindClosestIntersection_Mid3", Center, Radius, Thickness,
+		//	StartPt, EndPt, 
+		//	IntersectionResult(Center+Chaos::TVector<T,3>(0,0,Radius+Thickness), true));
+
+		// Fails w/no collision:
+		// "RunTestFindClosestIntersection_Mid4 - Sphere((0, 0, 0), 1).FindClosestIntersection(StartPt(0, 0, 1), EndPt(0, 0, 0.5), Thickness:0.1) = RESULT: 0, Point: (0, 0, 0), EXPECTED: 1 (0, 0, 1.1)."	FString
+		EndPt = StartPt;
+		StartPt = EndPt + Chaos::TVector<T, 3>(0, 0, Radius * .5);
+		//RunTestFindClosestIntersectionHelper(R, "RunTestFindClosestIntersection_Mid4", Center, Radius, Thickness,
+		//	StartPt, EndPt, 
+		//	IntersectionResult(Center+Chaos::TVector<T,3>(0,0,Radius+Thickness), true));
+
+		StartPt += Chaos::TVector<T, 3>(0, 0, Radius * .5);
+		RunTestFindClosestIntersectionHelper(R, "RunTestFindClosestIntersection_Mid5", Center, Radius, Thickness,
+			StartPt, EndPt, 
+			IntersectionResult(Center+Chaos::TVector<T,3>(0,0,Radius+Thickness), true));
+	}
+
+	template <class T>
+	void RunTestFindClosestIntersection_Deep(
+		ExampleResponse& R, 
+		const Chaos::TVector<T, 3>& Center, 
+		const T Radius,
+		const T Thickness)
+	{
+		Chaos::TVector<T, 3> StartPt = Center;
+		Chaos::TVector<T, 3> EndPt = Center;
+
+		// Fails w/no collision
+		// "RunTestFindClosestIntersection_Deep1 - Sphere((0, 0, 0), 1).FindClosestIntersection(StartPt(0, 0, 0), EndPt(0, 0, 0), Thickness:0.1) = RESULT: 0, Point: (0, 0, 0), EXPECTED: 1 (1.1, 0, 0)."	FString
+		//RunTestFindClosestIntersectionHelper(R, "RunTestFindClosestIntersection_Deep1", Center, Radius, Thickness,
+		//	StartPt, EndPt, 
+		//	IntersectionResult(Center+Chaos::TVector<T,3>(Radius+Thickness,0,0), true));
+
+		// Fails w/no collision
+		// "RunTestFindClosestIntersection_Deep2 - Sphere((0, 0, 0), 1).FindClosestIntersection(StartPt(0, 0, 0), EndPt(0, 0, 0.5), Thickness:0.1) = RESULT: 0, Point: (0, 0, 0), EXPECTED: 1 (0, 0, 1.1)."	FString
+		EndPt += Chaos::TVector<T, 3>(0, 0, Radius*.5);
+		//RunTestFindClosestIntersectionHelper(R, "RunTestFindClosestIntersection_Deep2", Center, Radius, Thickness,
+		//	StartPt, EndPt, 
+		//	IntersectionResult(Center+Chaos::TVector<T, 3>(0, 0, Radius+Thickness), true));
+
+		EndPt += Chaos::TVector<T, 3>(0, 0, Radius);
+		RunTestFindClosestIntersectionHelper(R, "RunTestFindClosestIntersection_Deep3", Center, Radius, Thickness,
+			StartPt, EndPt, 
+			IntersectionResult(Center+Chaos::TVector<T, 3>(0, 0, Radius+Thickness), true));
+
+		// Fails w/no collision:
+		// "RunTestFindClosestIntersection_Deep4 - Sphere((0, 0, 0), 1).FindClosestIntersection(StartPt(0, 0, 0.5), EndPt(0, 0, 0), Thickness:0.1) = RESULT: 0, Point: (0, 0, 0), EXPECTED: 1 (0, 0, 1.1)."	FString
+		EndPt = Center;
+		StartPt += Chaos::TVector<T, 3>(0, 0, Radius*.5);
+		//RunTestFindClosestIntersectionHelper(R, "RunTestFindClosestIntersection_Deep4", Center, Radius, Thickness,
+		//	StartPt, EndPt, 
+		//	IntersectionResult(Center+Chaos::TVector<T, 3>(0, 0, Radius+Thickness), true));
+
+		StartPt += Chaos::TVector<T, 3>(0, 0, Radius);
+		RunTestFindClosestIntersectionHelper(R, "RunTestFindClosestIntersection_Deep5", Center, Radius, Thickness,
+			StartPt, EndPt, 
+			IntersectionResult(Center+Chaos::TVector<T, 3>(0, 0, Radius+Thickness), true));
+
+		// Miss
+		StartPt = Center + Chaos::TVector<T, 3>(Radius + Thickness + 1);
+		EndPt = StartPt + Chaos::TVector<T, 3>(Radius + Thickness + 1);
+		RunTestFindClosestIntersectionHelper(R, "RunTestFindClosestIntersection_Deep6", Center, Radius, Thickness,
+			StartPt, EndPt, 
+			IntersectionResult(StartPt, false));
+	}
+
+	template <class T>
+	void RunTestFindClosestIntersection(
+		ExampleResponse& R, const Chaos::TVector<T,3>& Center, const T Radius)
+	{
+		//
+		// Run different test configurations with varying thicknesses.
+		//
+
+		// Deep intersection
+		RunTestFindClosestIntersection_Deep<T>(R, Center, Radius, 0.0);
+		RunTestFindClosestIntersection_Deep<T>(R, Center, Radius, Radius*.1);
+//		RunTestFindClosestIntersection_Deep<T>(R, Center, Radius, Radius*.5); // fails with incorrect collision points
+//		RunTestFindClosestIntersection_Deep<T>(R, Center, Radius, Radius);
+//		RunTestFindClosestIntersection_Deep<T>(R, Center, Radius, Radius*2);
+
+		// Mid intersection
+		RunTestFindClosestIntersection_Mid<T>(R, Center, Radius, 0.0);
+		RunTestFindClosestIntersection_Mid<T>(R, Center, Radius, Radius*.1);
+//		RunTestFindClosestIntersection_Mid<T>(R, Center, Radius, Radius*.5); // fails with incorrect collision points
+//		RunTestFindClosestIntersection_Mid<T>(R, Center, Radius, Radius);
+//		RunTestFindClosestIntersection_Mid<T>(R, Center, Radius, Radius*2);
+
+		// Grazing intersection - all fail
+		RunTestFindClosestIntersection_Shallow<T>(R, Center, Radius, 0.0);
+//		RunTestFindClosestIntersection_Shallow<T>(R, Center, Radius, Radius*.1); // fails with incorrect collision points
+//		RunTestFindClosestIntersection_Shallow<T>(R, Center, Radius, Radius*.5);
+//		RunTestFindClosestIntersection_Shallow<T>(R, Center, Radius, Radius);
+//		RunTestFindClosestIntersection_Shallow<T>(R, Center, Radius, Radius*2);
+
+		// Tangent intersection
+		RunTestFindClosestIntersection_Tangent<T>(R, Center, Radius, 0.0);
+		RunTestFindClosestIntersection_Tangent<T>(R, Center, Radius, Radius*.1);
+		RunTestFindClosestIntersection_Tangent<T>(R, Center, Radius, Radius*.5);
+		RunTestFindClosestIntersection_Tangent<T>(R, Center, Radius, Radius);
+		RunTestFindClosestIntersection_Tangent<T>(R, Center, Radius, Radius*2);
+	}
+	template <class T>
+	void TestIntersections(ExampleResponse& R)
+	{
+		// At the origin
+		RunTestFindClosestIntersection(R, Chaos::TVector<T, 3>((T)0.0), (T)1.0);
+//		RunTestFindClosestIntersection(R, Chaos::TVector<T, 3>((T)0.0), (T)10.0);
+//		RunTestFindClosestIntersection(R, Chaos::TVector<T, 3>((T)0.0), (T)100.0);
+
+		// Off origin
+//		RunTestFindClosestIntersection(R, Chaos::TVector<T, 3>((T)1.0), (T)1.0);
+//		RunTestFindClosestIntersection(R, Chaos::TVector<T, 3>((T)10.0), (T)1.0);
+//		RunTestFindClosestIntersection(R, Chaos::TVector<T, 3>((T)100.0), (T)1.0);
+
+//		RunTestFindClosestIntersection(R, Chaos::TVector<T, 3>((T)1.0), (T)10.0);
+//		RunTestFindClosestIntersection(R, Chaos::TVector<T, 3>((T)10.0), (T)10.0);
+//		RunTestFindClosestIntersection(R, Chaos::TVector<T, 3>((T)100.0), (T)10.0);
+
+//		RunTestFindClosestIntersection(R, Chaos::TVector<T, 3>((T)1.0), (T)100.0);
+//		RunTestFindClosestIntersection(R, Chaos::TVector<T, 3>((T)10.0), (T)100.0);
+//		RunTestFindClosestIntersection(R, Chaos::TVector<T, 3>((T)100.0), (T)100.0);
+	}
+
+	//==========================================================================
+	// Sample points tests
+	//==========================================================================
+
 	template <class T, int d>
 	void RunTestComputeSamplePoints(ExampleResponse& R, const Chaos::TSphere<T, d> &Sphere)
 	{
 		R.ExpectTrue(Sphere.GetType() == Chaos::ImplicitObjectType::Sphere, FString("Implicit object type is not 'sphere'."));
 
-		Chaos::TVector<T, d> Point = Sphere.Center();
+		Chaos::TVector<T, d> Point = Sphere.GetCenter();
 		T Phi = Sphere.SignedDistance(Point);
-		R.ExpectTrue(FMath::Abs(Phi + Sphere.Radius()) <= SMALL_NUMBER, FString("Sphere failed phi depth sanity test."));
+		R.ExpectTrue(FMath::Abs(Phi + Sphere.GetRadius()) <= SMALL_NUMBER, FString("Sphere failed phi depth sanity test."));
 		
-		Point[0] += Sphere.Radius();
+		Point[0] += Sphere.GetRadius();
 		Phi = Sphere.SignedDistance(Point);
 		R.ExpectTrue(FMath::Abs(Phi) <= KINDA_SMALL_NUMBER, FString("Sphere failed phi surface sanity test."));
 
@@ -47,7 +361,7 @@ namespace GeometryCollectionExample
 			check(OnSurface);
 			R.ExpectTrue(OnSurface, FString("Produced a point not on the surface of the sphere."));
 
-			const bool BelowCenter = Pt[2] < Sphere.Center()[2] + KINDA_SMALL_NUMBER;
+			const bool BelowCenter = Pt[2] < Sphere.GetCenter()[2] + KINDA_SMALL_NUMBER;
 			check(BelowCenter);
 			R.ExpectTrue(BelowCenter, FString("Bottom semisphere produced a point above midline."));
 		}
@@ -61,7 +375,7 @@ namespace GeometryCollectionExample
 			check(OnSurface);
 			R.ExpectTrue(OnSurface, FString("Produced a point not on the surface of the sphere."));
 
-			const bool BelowCenter = Pt[2] > Sphere.Center()[2] - KINDA_SMALL_NUMBER;
+			const bool BelowCenter = Pt[2] > Sphere.GetCenter()[2] - KINDA_SMALL_NUMBER;
 			check(BelowCenter);
 			R.ExpectTrue(BelowCenter, FString("Top semisphere produced a point above midline."));
 		}
@@ -118,6 +432,11 @@ namespace GeometryCollectionExample
 		//TestComputeSamplePoints<T, 2>(R);
 		TestComputeSamplePoints_Sphere<T, 3>(R);
 		TestComputeSamplePoints_SemiSphere<T>(R);
+		check(!R.HasError());
+
+		TestIntersections<T>(R);
+		check(!R.HasError());
+
 		return !R.HasError();
 	}
 	template bool TestImplicitSphere<float>(ExampleResponse&& R);
diff --git a/Engine/Plugins/Experimental/GeometryCollectionPlugin/Source/GeometryCollectionExampleCore/Private/GeometryCollection/GeometryCollectionExampleSimulation.cpp b/Engine/Plugins/Experimental/GeometryCollectionPlugin/Source/GeometryCollectionExampleCore/Private/GeometryCollection/GeometryCollectionExampleSimulation.cpp
--- a/Engine/Plugins/Experimental/GeometryCollectionPlugin/Source/GeometryCollectionExampleCore/Private/GeometryCollection/GeometryCollectionExampleSimulation.cpp
+++ b/Engine/Plugins/Experimental/GeometryCollectionPlugin/Source/GeometryCollectionExampleCore/Private/GeometryCollection/GeometryCollectionExampleSimulation.cpp
@@ -12,9 +12,10 @@
 
 #include "GeometryCollection/GeometryDynamicCollection.h"
 #include "GeometryCollection/GeometryCollectionSimulationTypes.h"
-#include "SolverObjects/SolverObjects.h"
+#include "PhysicsProxy/PhysicsProxies.h"
 #include "Chaos/ErrorReporter.h"
 #include "ChaosSolversModule.h"
+#include "PBDRigidsSolver.h"
 
 #define SMALL_THRESHOLD 1e-4
 
@@ -51,11 +52,13 @@ namespace GeometryCollectionExample
 			BuildSimulationData(ErrorReporter, *RestCollection, InParams.Shared);
 		};
 
-		FGeometryCollectionPhysicsObject* PhysObject = new FGeometryCollectionPhysicsObject(nullptr, DynamicCollection.Get(), InitFunc, nullptr, nullptr);;
+		FGeometryCollectionPhysicsProxy* PhysObject = new FGeometryCollectionPhysicsProxy(nullptr, DynamicCollection.Get(), InitFunc, nullptr, nullptr);;
 		PhysObject->Initialize();
 
 		Chaos::FPBDRigidsSolver* Solver = FChaosSolversModule::GetModule()->CreateSolver(true);
+#if CHAOS_PARTICLEHANDLE_TODO
 		Solver->RegisterObject(PhysObject);
+#endif
 		Solver->SetHasFloor(false);
 		Solver->SetEnabled(true);
 		PhysObject->ActivateBodies();
@@ -113,11 +116,13 @@ namespace GeometryCollectionExample
 			BuildSimulationData(ErrorReporter, *RestCollection, InParams.Shared);
 		};
 
-		FGeometryCollectionPhysicsObject* PhysObject = new FGeometryCollectionPhysicsObject(nullptr, DynamicCollection.Get(), InitFunc, nullptr, nullptr);;
+		FGeometryCollectionPhysicsProxy* PhysObject = new FGeometryCollectionPhysicsProxy(nullptr, DynamicCollection.Get(), InitFunc, nullptr, nullptr);;
 		PhysObject->Initialize();
 
 		Chaos::FPBDRigidsSolver* Solver = FChaosSolversModule::GetModule()->CreateSolver(true);
+#if CHAOS_PARTICLEHANDLE_TODO
 		Solver->RegisterObject(PhysObject);
+#endif
 		Solver->SetHasFloor(true);
 		Solver->SetIsFloorAnalytic(true);
 		Solver->SetEnabled(true);
@@ -177,11 +182,13 @@ namespace GeometryCollectionExample
 			BuildSimulationData(ErrorReporter, *RestCollection, InParams.Shared);
 		};
 
-		FGeometryCollectionPhysicsObject* PhysObject = new FGeometryCollectionPhysicsObject(nullptr, DynamicCollection.Get(), InitFunc, nullptr, nullptr);;
+		FGeometryCollectionPhysicsProxy* PhysObject = new FGeometryCollectionPhysicsProxy(nullptr, DynamicCollection.Get(), InitFunc, nullptr, nullptr);;
 		PhysObject->Initialize();
 
 		Chaos::FPBDRigidsSolver* Solver = FChaosSolversModule::GetModule()->CreateSolver(true);
+#if CHAOS_PARTICLEHANDLE_TODO
 		Solver->RegisterObject(PhysObject);
+#endif
 		Solver->SetHasFloor(true);
 		Solver->SetIsFloorAnalytic(true);
 		Solver->SetEnabled(true);
@@ -245,11 +252,13 @@ namespace GeometryCollectionExample
 			BuildSimulationData(ErrorReporter, *RestCollection, InParams.Shared);
 		};
 
-		FGeometryCollectionPhysicsObject* PhysObject = new FGeometryCollectionPhysicsObject(nullptr, DynamicCollection.Get(), InitFunc, nullptr, nullptr);;
+		FGeometryCollectionPhysicsProxy* PhysObject = new FGeometryCollectionPhysicsProxy(nullptr, DynamicCollection.Get(), InitFunc, nullptr, nullptr);;
 		PhysObject->Initialize();
 
 		Chaos::FPBDRigidsSolver* Solver = FChaosSolversModule::GetModule()->CreateSolver(true);
+#if CHAOS_PARTICLEHANDLE_TODO
 		Solver->RegisterObject(PhysObject);
+#endif
 		Solver->SetHasFloor(true);
 		Solver->SetIsFloorAnalytic(true);
 		Solver->SetEnabled(true);
@@ -308,11 +317,13 @@ namespace GeometryCollectionExample
 			BuildSimulationData(ErrorReporter, *RestCollection, InParams.Shared);
 		};
 
-		FGeometryCollectionPhysicsObject* PhysObject = new FGeometryCollectionPhysicsObject(nullptr, DynamicCollection.Get(), InitFunc, nullptr, nullptr);;
+		FGeometryCollectionPhysicsProxy* PhysObject = new FGeometryCollectionPhysicsProxy(nullptr, DynamicCollection.Get(), InitFunc, nullptr, nullptr);;
 		PhysObject->Initialize();
 
 		Chaos::FPBDRigidsSolver* Solver = FChaosSolversModule::GetModule()->CreateSolver(true);
+#if CHAOS_PARTICLEHANDLE_TODO
 		Solver->RegisterObject(PhysObject);
+#endif
 		Solver->SetHasFloor(false);
 		Solver->SetIsFloorAnalytic(true);
 		Solver->SetEnabled(true);
@@ -374,7 +385,7 @@ namespace GeometryCollectionExample
 			BuildSimulationData(ErrorReporter, *RestCollection, InParams.Shared);
 		};
 
-		FGeometryCollectionPhysicsObject* PhysObject = new FGeometryCollectionPhysicsObject(nullptr, DynamicCollection.Get(), InitFunc, nullptr, nullptr);;
+		FGeometryCollectionPhysicsProxy* PhysObject = new FGeometryCollectionPhysicsProxy(nullptr, DynamicCollection.Get(), InitFunc, nullptr, nullptr);;
 		PhysObject->Initialize();
 
 		//
@@ -385,7 +396,9 @@ namespace GeometryCollectionExample
 		ObjectType[1] = (int32)EObjectStateTypeEnum::Chaos_Object_Dynamic;
 
 		Chaos::FPBDRigidsSolver* Solver = FChaosSolversModule::GetModule()->CreateSolver(true);
+#if CHAOS_PARTICLEHANDLE_TODO
 		Solver->RegisterObject(PhysObject);
+#endif
 		Solver->SetHasFloor(false);
 		Solver->SetIsFloorAnalytic(true);
 		Solver->SetEnabled(true);
@@ -448,17 +461,20 @@ namespace GeometryCollectionExample
 			BuildSimulationData(ErrorReporter, *RestCollection, InParams.Shared);
 		};
 
-		FGeometryCollectionPhysicsObject* PhysObject = new FGeometryCollectionPhysicsObject(nullptr, DynamicCollection.Get(), InitFunc, nullptr, nullptr);;
+		FGeometryCollectionPhysicsProxy* PhysObject = new FGeometryCollectionPhysicsProxy(nullptr, DynamicCollection.Get(), InitFunc, nullptr, nullptr);;
 		PhysObject->Initialize();
 		PhysObject->SetCollisionParticlesPerObjectFraction( 1.0 );
 
 		Chaos::FPBDRigidsSolver* Solver = FChaosSolversModule::GetModule()->CreateSolver(true);
+#if CHAOS_PARTICLEHANDLE_TODO
 		Solver->RegisterObject(PhysObject);
+#endif
 		Solver->SetHasFloor(true);
 		Solver->SetEnabled(true);
 		PhysObject->ActivateBodies();
 
 		Solver->AdvanceSolverBy(1 / 24.);
+#if TODO_REIMPLEMENT_GET_RIGID_PARTICLES
 		Chaos::TPBDRigidParticles<float, 3>& Particles = Solver->GetRigidParticles();
 
 		for (int Frame = 1; Frame < 200; Frame++)
@@ -503,9 +519,9 @@ namespace GeometryCollectionExample
 		R.ExpectTrue(Particles.X(2).Z > 50.0);
 		R.ExpectTrue(Particles.X(3).Z < -10000);
 		R.ExpectTrue(Particles.X(4).Z < -10000);
+#endif
 
 		FChaosSolversModule::GetModule()->DestroySolver(Solver);
-
 		delete PhysObject;
 #endif
 
@@ -527,7 +543,7 @@ namespace GeometryCollectionExample
 		P.SizeData.CollisionType = ECollisionTypeEnum::Chaos_Surface_Volumetric;
 
 		SimulationObjects<T>* Object = new SimulationObjects<T>(P);
-		Object->PhysicsObject->Initialize();
+		Object->PhysicsProxy->Initialize();
 
 		// check implicit domain
 		typedef TUniquePtr<Chaos::TImplicitObject<float, 3>> FImplicitPointer;
@@ -561,7 +577,7 @@ namespace GeometryCollectionExample
 
 		// check simplicial elements
 		typedef TUniquePtr< FCollisionStructureManager::FSimplicial > FSimplicialPointer;
-		const TManagedArray<FSimplicialPointer> & Simplicials = Object->RestCollection->template GetAttribute<FSimplicialPointer>(FGeometryCollectionPhysicsObject::SimplicialsAttribute, FTransformCollection::TransformGroup);
+		const TManagedArray<FSimplicialPointer> & Simplicials = Object->RestCollection->template GetAttribute<FSimplicialPointer>(FGeometryCollectionPhysicsProxy::SimplicialsAttribute, FTransformCollection::TransformGroup);
 		R.ExpectTrue(Simplicials.Num() == 1);
 		R.ExpectTrue(Simplicials[0].IsValid());
 		R.ExpectTrue(Simplicials[0]->Size() == 8);
diff --git a/Engine/Plugins/Experimental/GeometryCollectionPlugin/Source/GeometryCollectionExampleCore/Private/GeometryCollection/GeometryCollectionExampleSimulationField.cpp b/Engine/Plugins/Experimental/GeometryCollectionPlugin/Source/GeometryCollectionExampleCore/Private/GeometryCollection/GeometryCollectionExampleSimulationField.cpp
--- a/Engine/Plugins/Experimental/GeometryCollectionPlugin/Source/GeometryCollectionExampleCore/Private/GeometryCollection/GeometryCollectionExampleSimulationField.cpp
+++ b/Engine/Plugins/Experimental/GeometryCollectionPlugin/Source/GeometryCollectionExampleCore/Private/GeometryCollection/GeometryCollectionExampleSimulationField.cpp
@@ -13,9 +13,10 @@
 #include "Field/FieldSystemNodes.h"
 #include "GeometryCollection/GeometryDynamicCollection.h"
 #include "GeometryCollection/GeometryCollectionSimulationTypes.h"
-#include "SolverObjects/SolverObjects.h"
+#include "PhysicsProxy/PhysicsProxies.h"
 #include "Chaos/ErrorReporter.h"
 #include "Chaos/PBDRigidClustering.h"
+#include "PBDRigidsSolver.h"
 #include "ChaosSolversModule.h"
 
 #define SMALL_THRESHOLD 1e-4
@@ -62,11 +63,13 @@ namespace GeometryCollectionExample
 			BuildSimulationData(ErrorReporter, *RestCollection, InParams.Shared);
 		};
 
-		FGeometryCollectionPhysicsObject* PhysObject = new FGeometryCollectionPhysicsObject(nullptr, DynamicCollection.Get(), InitFunc, nullptr, nullptr);;
+		FGeometryCollectionPhysicsProxy* PhysObject = new FGeometryCollectionPhysicsProxy(nullptr, DynamicCollection.Get(), InitFunc, nullptr, nullptr);;
 		PhysObject->Initialize();
 
 		Chaos::FPBDRigidsSolver* Solver = FChaosSolversModule::GetModule()->CreateSolver(true);
+#if CHAOS_PARTICLEHANDLE_TODO
 		Solver->RegisterObject(PhysObject);
+#endif
 		Solver->SetHasFloor(false);
 		Solver->SetIsFloorAnalytic(true);
 		Solver->SetEnabled(true);
@@ -147,13 +150,15 @@ namespace GeometryCollectionExample
 		//
 		// Solver setup
 		//
-		FFieldSystemPhysicsObject* FieldObject = new FFieldSystemPhysicsObject(nullptr);
-		FGeometryCollectionPhysicsObject* PhysObject = new FGeometryCollectionPhysicsObject(nullptr, DynamicCollection.Get(), InitFunc, nullptr, nullptr);;
+		FFieldSystemPhysicsProxy* FieldObject = new FFieldSystemPhysicsProxy(nullptr);
+		FGeometryCollectionPhysicsProxy* PhysObject = new FGeometryCollectionPhysicsProxy(nullptr, DynamicCollection.Get(), InitFunc, nullptr, nullptr);;
 		PhysObject->Initialize();
 
 		Chaos::FPBDRigidsSolver* Solver = FChaosSolversModule::GetModule()->CreateSolver(true);
+#if CHAOS_PARTICLEHANDLE_TODO
 		Solver->RegisterObject(PhysObject);
 		Solver->RegisterObject(FieldObject);
+#endif
 		Solver->SetHasFloor(false);
 		Solver->SetIsFloorAnalytic(true);
 		Solver->SetEnabled(true);
@@ -246,14 +251,16 @@ namespace GeometryCollectionExample
 		//
 		// Solver setup
 		//
-		FGeometryCollectionPhysicsObject* PhysObject = new FGeometryCollectionPhysicsObject(nullptr, DynamicCollection.Get(), InitFunc, nullptr, nullptr);;
+		FGeometryCollectionPhysicsProxy* PhysObject = new FGeometryCollectionPhysicsProxy(nullptr, DynamicCollection.Get(), InitFunc, nullptr, nullptr);;
 		PhysObject->Initialize();
 
-		FFieldSystemPhysicsObject* FieldObject = new FFieldSystemPhysicsObject(nullptr);
+		FFieldSystemPhysicsProxy* FieldObject = new FFieldSystemPhysicsProxy(nullptr);
 
 		Chaos::FPBDRigidsSolver* Solver = FChaosSolversModule::GetModule()->CreateSolver(true);
+#if CHAOS_PARTICLEHANDLE_TODO
 		Solver->RegisterObject(PhysObject);
 		Solver->RegisterObject(FieldObject);
+#endif
 		Solver->SetHasFloor(false);
 		Solver->SetEnabled(true);
 		PhysObject->ActivateBodies();
@@ -274,8 +281,6 @@ namespace GeometryCollectionExample
 
 			FinalizeSolver(*Solver);
 
-			const Chaos::TPBDRigidParticles<float, 3>& Particles = Solver->GetRigidParticles();
-
 			if (Frame < 5)
 			{
 				R.ExpectTrue(FMath::Abs(Transform[0].GetTranslation().Z - 5.f) < SMALL_THRESHOLD);
@@ -345,14 +350,16 @@ namespace GeometryCollectionExample
 		//
 		// Solver setup
 		//
-		FGeometryCollectionPhysicsObject* PhysObject = new FGeometryCollectionPhysicsObject(nullptr, DynamicCollection.Get(), InitFunc, nullptr, nullptr);;
+		FGeometryCollectionPhysicsProxy* PhysObject = new FGeometryCollectionPhysicsProxy(nullptr, DynamicCollection.Get(), InitFunc, nullptr, nullptr);;
 		PhysObject->Initialize();
 
-		FFieldSystemPhysicsObject* FieldObject = new FFieldSystemPhysicsObject(nullptr);
+		FFieldSystemPhysicsProxy* FieldObject = new FFieldSystemPhysicsProxy(nullptr);
 
 		Chaos::FPBDRigidsSolver* Solver = FChaosSolversModule::GetModule()->CreateSolver(true);
+#if CHAOS_PARTICLEHANDLE_TODO
 		Solver->RegisterObject(PhysObject);
 		Solver->RegisterObject(FieldObject);
+#endif
 		Solver->SetHasFloor(false);
 		Solver->SetEnabled(true);
 		PhysObject->ActivateBodies();
@@ -369,7 +376,6 @@ namespace GeometryCollectionExample
 			}
 
 			Solver->AdvanceSolverBy(1 / 24.);
-			const Chaos::TPBDRigidParticles<float, 3>& Particles = Solver->GetRigidParticles();
 
 			FinalizeSolver(*Solver);
 
@@ -438,14 +444,16 @@ namespace GeometryCollectionExample
 		//
 		// Solver setup
 		//
-		FGeometryCollectionPhysicsObject* PhysObject = new FGeometryCollectionPhysicsObject(nullptr, DynamicCollection.Get(), InitFunc, nullptr, nullptr);;
+		FGeometryCollectionPhysicsProxy* PhysObject = new FGeometryCollectionPhysicsProxy(nullptr, DynamicCollection.Get(), InitFunc, nullptr, nullptr);;
 		PhysObject->Initialize();
 
-		FFieldSystemPhysicsObject* FieldObject = new FFieldSystemPhysicsObject(nullptr);
+		FFieldSystemPhysicsProxy* FieldObject = new FFieldSystemPhysicsProxy(nullptr);
 
 		Chaos::FPBDRigidsSolver* Solver = FChaosSolversModule::GetModule()->CreateSolver(true);
+#if CHAOS_PARTICLEHANDLE_TODO
 		Solver->RegisterObject(PhysObject);
 		Solver->RegisterObject(FieldObject);
+#endif
 		Solver->SetHasFloor(false);
 		Solver->SetEnabled(true);
 		PhysObject->ActivateBodies();
@@ -463,6 +471,8 @@ namespace GeometryCollectionExample
 			}
 
 			Solver->AdvanceSolverBy(1 / 24.);
+
+#if TODO_REIMPLEMENT_GET_RIGID_PARTICLES
 			const Chaos::TPBDRigidParticles<float, 3>& Particles = Solver->GetRigidParticles();
 
 			FinalizeSolver(*Solver);
@@ -478,6 +488,7 @@ namespace GeometryCollectionExample
 			}
 
 			PreviousY = Particles.W(0).Y;
+#endif
 
 		}
 
@@ -538,14 +549,16 @@ namespace GeometryCollectionExample
 		//
 		// Solver setup
 		//
-		FGeometryCollectionPhysicsObject* PhysObject = new FGeometryCollectionPhysicsObject(nullptr, DynamicCollection.Get(), InitFunc, nullptr, nullptr);;
+		FGeometryCollectionPhysicsProxy* PhysObject = new FGeometryCollectionPhysicsProxy(nullptr, DynamicCollection.Get(), InitFunc, nullptr, nullptr);;
 		PhysObject->Initialize();
 
-		FFieldSystemPhysicsObject* FieldObject = new FFieldSystemPhysicsObject(nullptr);
+		FFieldSystemPhysicsProxy* FieldObject = new FFieldSystemPhysicsProxy(nullptr);
 
 		Chaos::FPBDRigidsSolver* Solver = FChaosSolversModule::GetModule()->CreateSolver(true);
+#if CHAOS_PARTICLEHANDLE_TODO
 		Solver->RegisterObject(PhysObject);
 		Solver->RegisterObject(FieldObject);
+#endif
 		Solver->SetHasFloor(false);
 		Solver->SetEnabled(true);
 		PhysObject->ActivateBodies();
@@ -562,7 +575,6 @@ namespace GeometryCollectionExample
 
 			FinalizeSolver(*Solver);
 
-			const Chaos::TPBDRigidParticles<float, 3>& Particles = Solver->GetRigidParticles();
 			R.ExpectTrue(Transform[0].GetTranslation().Z < 20.f);
 			R.ExpectTrue(Transform[0].GetTranslation().Z > -10.);
 		}
@@ -620,14 +632,16 @@ namespace GeometryCollectionExample
 		//
 		// Solver setup
 		//
-		FGeometryCollectionPhysicsObject* PhysObject = new FGeometryCollectionPhysicsObject(nullptr, DynamicCollection.Get(), InitFunc, nullptr, nullptr);;
+		FGeometryCollectionPhysicsProxy* PhysObject = new FGeometryCollectionPhysicsProxy(nullptr, DynamicCollection.Get(), InitFunc, nullptr, nullptr);;
 		PhysObject->Initialize();
 
-		FFieldSystemPhysicsObject* FieldObject = new FFieldSystemPhysicsObject(nullptr);
+		FFieldSystemPhysicsProxy* FieldObject = new FFieldSystemPhysicsProxy(nullptr);
 
 		Chaos::FPBDRigidsSolver* Solver = FChaosSolversModule::GetModule()->CreateSolver(true);
+#if CHAOS_PARTICLEHANDLE_TODO
 		Solver->RegisterObject(PhysObject);
 		Solver->RegisterObject(FieldObject);
+#endif
 		Solver->SetHasFloor(false);
 		Solver->SetEnabled(true);
 		PhysObject->ActivateBodies();
@@ -637,7 +651,6 @@ namespace GeometryCollectionExample
 		Solver->AdvanceSolverBy(1 / 24.);
 		FinalizeSolver(*Solver);
 
-		const Chaos::TPBDRigidParticles<float, 3>& Particles = Solver->GetRigidParticles();
 		float PreviousX = 0.f;
 		TManagedArray<FTransform>& Transform = DynamicCollection->Transform;
 		for (int Frame = 1; Frame < 10; Frame++)
@@ -717,16 +730,19 @@ namespace GeometryCollectionExample
 			BuildSimulationData(ErrorReporter, *RestCollection, InParams.Shared);
 		};
 
-		FGeometryCollectionPhysicsObject* PhysObject = new FGeometryCollectionPhysicsObject(nullptr, DynamicCollection.Get(), InitFunc, nullptr, nullptr);;
+		FGeometryCollectionPhysicsProxy* PhysObject = new FGeometryCollectionPhysicsProxy(nullptr, DynamicCollection.Get(), InitFunc, nullptr, nullptr);;
 		PhysObject->Initialize();
 
 		Chaos::FPBDRigidsSolver* Solver = FChaosSolversModule::GetModule()->CreateSolver(true);
+#if CHAOS_PARTICLEHANDLE_TODO
 		Solver->RegisterObject(PhysObject);
+#endif
 		Solver->SetHasFloor(true);
 		Solver->SetEnabled(true);
 		PhysObject->ActivateBodies();
 
 
+#if TODO_REIMPLEMENT_GET_RIGID_PARTICLES
 		for (int Frame = 0; Frame < 60; Frame++)
 		{
 			Solver->AdvanceSolverBy(1 / 24.);
@@ -755,6 +771,7 @@ namespace GeometryCollectionExample
 		R.ExpectTrue(Particles.X(2).Z < 0);
 		R.ExpectTrue(Particles.X(3).Z < 0);
 		R.ExpectTrue(Particles.X(4).Z < 0);
+#endif
 
 		FChaosSolversModule::GetModule()->DestroySolver(Solver);
 
@@ -805,15 +822,20 @@ namespace GeometryCollectionExample
 		FalloffField->Falloff = EFieldFalloffType::Field_FallOff_None;
 
 		Chaos::FPBDRigidsSolver* Solver = FChaosSolversModule::GetModule()->CreateSolver(true);
-		FFieldSystemPhysicsObject* FieldObject = new FFieldSystemPhysicsObject(nullptr);
-		FGeometryCollectionPhysicsObject* PhysObject = new FGeometryCollectionPhysicsObject(nullptr, DynamicCollection.Get(), InitFunc, nullptr, nullptr);;
+		FFieldSystemPhysicsProxy* FieldObject = new FFieldSystemPhysicsProxy(nullptr);
+		FGeometryCollectionPhysicsProxy* PhysObject = new FGeometryCollectionPhysicsProxy(nullptr, DynamicCollection.Get(), InitFunc, nullptr, nullptr);;
+#if TODO_REIMPLEMENT_RIGID_CLUSTERING
 		Chaos::FPBDRigidsSolver::FClusteringType & Clustering = Solver->GetRigidClustering();
 		const Chaos::FPBDRigidsSolver::FClusteringType::FClusterMap & ClusterMap = Clustering.GetChildrenMap();
+#endif
+#if TODO_REIMPLEMENT_GET_RIGID_PARTICLES
 		const Chaos::TPBDRigidParticles<float, 3>& Particles = Solver->GetRigidParticles();
 
 		PhysObject->Initialize();
+#if CHAOS_PARTICLEHANDLE_TODO
 		Solver->RegisterObject(PhysObject);
 		Solver->RegisterObject(FieldObject);
+#endif
 		Solver->SetHasFloor(false);
 		Solver->SetEnabled(true);
 		PhysObject->ActivateBodies();
@@ -821,8 +843,6 @@ namespace GeometryCollectionExample
 		Solver->AdvanceSolverBy(1 / 24.);
 		FinalizeSolver(*Solver);
 
-		Chaos::TArrayCollectionArray<float>& InternalStrain = Clustering.GetStrainArray();
-
 		{
 			FName TargetName = GetFieldPhysicsName(EFieldPhysicsType::Field_ExternalClusterStrain);
 			FieldObject->BufferCommand(Solver, { TargetName, FalloffField->NewCopy() });
@@ -854,6 +874,7 @@ namespace GeometryCollectionExample
 			R.ExpectTrue(!Particles.Disabled(5));
 			R.ExpectTrue(Particles.Disabled(6));
 		}
+#endif
 
 		FChaosSolversModule::GetModule()->DestroySolver(Solver);
 
@@ -906,15 +927,21 @@ namespace GeometryCollectionExample
 		FalloffField->Falloff = EFieldFalloffType::Field_FallOff_None;
 
 		Chaos::FPBDRigidsSolver* Solver = FChaosSolversModule::GetModule()->CreateSolver(true);
-		FFieldSystemPhysicsObject* FieldObject = new FFieldSystemPhysicsObject(nullptr);
-		FGeometryCollectionPhysicsObject* PhysObject = new FGeometryCollectionPhysicsObject(nullptr, DynamicCollection.Get(), InitFunc, nullptr, nullptr);;
+		FFieldSystemPhysicsProxy* FieldObject = new FFieldSystemPhysicsProxy(nullptr);
+		FGeometryCollectionPhysicsProxy* PhysObject = new FGeometryCollectionPhysicsProxy(nullptr, DynamicCollection.Get(), InitFunc, nullptr, nullptr);;
+#if TODO_REIMPLEMENT_RIGID_CLUSTERING
 		Chaos::FPBDRigidsSolver::FClusteringType & Clustering = Solver->GetRigidClustering();
 		const Chaos::FPBDRigidsSolver::FClusteringType::FClusterMap & ClusterMap = Clustering.GetChildrenMap();
+#endif
+#if TODO_REIMPLEMENT_GET_RIGID_PARTICLES
 		const Chaos::TPBDRigidParticles<float, 3>& Particles = Solver->GetRigidParticles();
+#endif
 
 		PhysObject->Initialize();
+#if CHAOS_PARTICLEHANDLE_TODO
 		Solver->RegisterObject(PhysObject);
 		Solver->RegisterObject(FieldObject);
+#endif
 		Solver->SetHasFloor(false);
 		Solver->SetEnabled(true);
 		PhysObject->ActivateBodies();
@@ -922,8 +949,6 @@ namespace GeometryCollectionExample
 		Solver->AdvanceSolverBy(1 / 24.);
 		FinalizeSolver(*Solver);
 
-		Chaos::TArrayCollectionArray<float>& InternalStrain = Clustering.GetStrainArray();
-
 		{
 			FName TargetName = GetFieldPhysicsName(EFieldPhysicsType::Field_ExternalClusterStrain);
 			FFieldSystemCommand Command(TargetName, FalloffField->NewCopy());
@@ -931,6 +956,7 @@ namespace GeometryCollectionExample
 			Command.MetaData.Add(FFieldSystemMetaData::EMetaType::ECommandData_ProcessingResolution, TUniquePtr< FFieldSystemMetaDataProcessingResolution >(ResolutionData));
 			FieldObject->BufferCommand(Solver, Command);
 
+#if TODO_REIMPLEMENT_GET_RIGID_PARTICLES
 			R.ExpectTrue(Particles.Disabled(0));
 			R.ExpectTrue(Particles.Disabled(1));
 			R.ExpectTrue(Particles.Disabled(2));
@@ -940,6 +966,7 @@ namespace GeometryCollectionExample
 			R.ExpectTrue(Particles.Disabled(6));
 			R.ExpectTrue(Particles.Disabled(7));
 			R.ExpectTrue(!Particles.Disabled(8));
+#endif
 
 			Solver->AdvanceSolverBy(1 / 24.);
 			FinalizeSolver(*Solver);
@@ -947,12 +974,15 @@ namespace GeometryCollectionExample
 			Solver->AdvanceSolverBy(1 / 24.);
 			FinalizeSolver(*Solver);
 
+#if TODO_REIMPLEMENT_RIGID_CLUSTERING
 			R.ExpectTrue(ClusterMap.Num() == 1);
 			R.ExpectTrue(ClusterMap[6]->Num() == 3);
 			R.ExpectTrue(ClusterMap[6]->Contains(3));
 			R.ExpectTrue(ClusterMap[6]->Contains(4));
 			R.ExpectTrue(ClusterMap[6]->Contains(5));
+#endif
 
+#if TODO_REIMPLEMENT_GET_RIGID_PARTICLES
 			R.ExpectTrue(!Particles.Disabled(0));
 			R.ExpectTrue(!Particles.Disabled(1));
 			R.ExpectTrue(!Particles.Disabled(2));
@@ -962,6 +992,7 @@ namespace GeometryCollectionExample
 			R.ExpectTrue(!Particles.Disabled(6));
 			R.ExpectTrue(Particles.Disabled(7));
 			R.ExpectTrue(Particles.Disabled(8));
+#endif
 		}
 
 		FChaosSolversModule::GetModule()->DestroySolver(Solver);
@@ -1011,15 +1042,21 @@ namespace GeometryCollectionExample
 		FalloffField->Falloff = EFieldFalloffType::Field_FallOff_None;
 
 		Chaos::FPBDRigidsSolver* Solver = FChaosSolversModule::GetModule()->CreateSolver(true);
-		FFieldSystemPhysicsObject* FieldObject = new FFieldSystemPhysicsObject(nullptr);
-		FGeometryCollectionPhysicsObject* PhysObject = new FGeometryCollectionPhysicsObject(nullptr, DynamicCollection.Get(), InitFunc, nullptr, nullptr);;
+		FFieldSystemPhysicsProxy* FieldObject = new FFieldSystemPhysicsProxy(nullptr);
+		FGeometryCollectionPhysicsProxy* PhysObject = new FGeometryCollectionPhysicsProxy(nullptr, DynamicCollection.Get(), InitFunc, nullptr, nullptr);;
+#if TODO_REIMPLEMENT_RIGID_CLUSTERING
 		Chaos::FPBDRigidsSolver::FClusteringType & Clustering = Solver->GetRigidClustering();
 		const Chaos::FPBDRigidsSolver::FClusteringType::FClusterMap & ClusterMap = Clustering.GetChildrenMap();
+#endif
+#if TODO_REIMPLEMENT_GET_RIGID_PARTICLES
 		const Chaos::TPBDRigidParticles<float, 3>& Particles = Solver->GetRigidParticles();
+#endif
 
 		PhysObject->Initialize();
+#if CHAOS_PARTICLEHANDLE_TODO
 		Solver->RegisterObject(PhysObject);
 		Solver->RegisterObject(FieldObject);
+#endif
 		Solver->SetHasFloor(false);
 		Solver->SetEnabled(true);
 		PhysObject->ActivateBodies();
@@ -1027,8 +1064,6 @@ namespace GeometryCollectionExample
 		Solver->AdvanceSolverBy(1 / 24.);
 		FinalizeSolver(*Solver);
 
-		Chaos::TArrayCollectionArray<float>& InternalStrain = Clustering.GetStrainArray();
-
 		{
 
 			FName TargetName = GetFieldPhysicsName(EFieldPhysicsType::Field_ExternalClusterStrain);
@@ -1037,6 +1072,7 @@ namespace GeometryCollectionExample
 			Command.MetaData.Add(FFieldSystemMetaData::EMetaType::ECommandData_ProcessingResolution, TUniquePtr< FFieldSystemMetaDataProcessingResolution >(ResolutionData));
 			FieldObject->BufferCommand(Solver, Command);
 
+#if TODO_REIMPLEMENT_GET_RIGID_PARTICLES
 			R.ExpectTrue(Particles.Disabled(0));
 			R.ExpectTrue(Particles.Disabled(1));
 			R.ExpectTrue(Particles.Disabled(2));
@@ -1046,10 +1082,12 @@ namespace GeometryCollectionExample
 			R.ExpectTrue(Particles.Disabled(6));
 			R.ExpectTrue(Particles.Disabled(7));
 			R.ExpectTrue(!Particles.Disabled(8));
+#endif
 
 			Solver->AdvanceSolverBy(1 / 24.);
 			FinalizeSolver(*Solver);
 
+#if TODO_REIMPLEMENT_RIGID_CLUSTERING
 			R.ExpectTrue(ClusterMap.Num() == 2);
 			R.ExpectTrue(ClusterMap[6]->Num() == 3);
 			R.ExpectTrue(ClusterMap[6]->Contains(3));
@@ -1059,7 +1097,9 @@ namespace GeometryCollectionExample
 			R.ExpectTrue(ClusterMap[7]->Contains(0));
 			R.ExpectTrue(ClusterMap[7]->Contains(1));
 			R.ExpectTrue(ClusterMap[7]->Contains(2));
+#endif
 
+#if TODO_REIMPLEMENT_GET_RIGID_PARTICLES
 			R.ExpectTrue(Particles.Disabled(0));
 			R.ExpectTrue(Particles.Disabled(1));
 			R.ExpectTrue(Particles.Disabled(2));
@@ -1069,6 +1109,7 @@ namespace GeometryCollectionExample
 			R.ExpectTrue(!Particles.Disabled(6));
 			R.ExpectTrue(!Particles.Disabled(7));
 			R.ExpectTrue(Particles.Disabled(8));
+#endif
 		}
 
 		FChaosSolversModule::GetModule()->DestroySolver(Solver);
@@ -1119,16 +1160,21 @@ namespace GeometryCollectionExample
 		FalloffField->Falloff = EFieldFalloffType::Field_FallOff_None;
 
 		Chaos::FPBDRigidsSolver* Solver = FChaosSolversModule::GetModule()->CreateSolver(true);
-		FFieldSystemPhysicsObject* FieldObject = new FFieldSystemPhysicsObject(nullptr);
-		FGeometryCollectionPhysicsObject* PhysObject = new FGeometryCollectionPhysicsObject(nullptr, DynamicCollection.Get(), InitFunc, nullptr, nullptr);;
+		FFieldSystemPhysicsProxy* FieldObject = new FFieldSystemPhysicsProxy(nullptr);
+		FGeometryCollectionPhysicsProxy* PhysObject = new FGeometryCollectionPhysicsProxy(nullptr, DynamicCollection.Get(), InitFunc, nullptr, nullptr);;
+#if TODO_REIMPLEMENT_RIGID_CLUSTERING
 		Chaos::FPBDRigidsSolver::FClusteringType & Clustering = Solver->GetRigidClustering();
-
 		const Chaos::FPBDRigidsSolver::FClusteringType::FClusterMap & ClusterMap = Clustering.GetChildrenMap();
+#endif
+#if TODO_REIMPLEMENT_GET_RIGID_PARTICLES
 		const Chaos::TPBDRigidParticles<float, 3>& Particles = Solver->GetRigidParticles();
+#endif
 
 		PhysObject->Initialize();
+#if CHAOS_PARTICLEHANDLE_TODO
 		Solver->RegisterObject(PhysObject);
 		Solver->RegisterObject(FieldObject);
+#endif
 		Solver->SetHasFloor(false);
 		Solver->SetEnabled(true);
 		PhysObject->ActivateBodies();
@@ -1136,12 +1182,11 @@ namespace GeometryCollectionExample
 		Solver->AdvanceSolverBy(1 / 24.);
 		FinalizeSolver(*Solver);
 
-		Chaos::TArrayCollectionArray<float>& InternalStrain = Clustering.GetStrainArray();
-
 		{
 			FName TargetName = GetFieldPhysicsName(EFieldPhysicsType::Field_ExternalClusterStrain);
 			FieldObject->BufferCommand(Solver, { TargetName, FalloffField->NewCopy() });
 
+#if TODO_REIMPLEMENT_GET_RIGID_PARTICLES
 			R.ExpectTrue(Particles.Disabled(0));
 			R.ExpectTrue(Particles.Disabled(1));
 			R.ExpectTrue(Particles.Disabled(2));
@@ -1149,10 +1194,12 @@ namespace GeometryCollectionExample
 			R.ExpectTrue(Particles.Disabled(4));
 			R.ExpectTrue(Particles.Disabled(5));
 			R.ExpectTrue(!Particles.Disabled(6));
+#endif
 
 			Solver->AdvanceSolverBy(1 / 24.);
 			FinalizeSolver(*Solver);
 
+#if TODO_REIMPLEMENT_RIGID_CLUSTERING
 			R.ExpectTrue(ClusterMap.Num() == 2);
 			R.ExpectTrue(ClusterMap[4]->Num() == 2);
 			R.ExpectTrue(ClusterMap[4]->Contains(2));
@@ -1160,7 +1207,9 @@ namespace GeometryCollectionExample
 			R.ExpectTrue(ClusterMap[5]->Num() == 2);
 			R.ExpectTrue(ClusterMap[5]->Contains(0));
 			R.ExpectTrue(ClusterMap[5]->Contains(1));
+#endif
 
+#if TODO_REIMPLEMENT_GET_RIGID_PARTICLES
 			R.ExpectTrue(Particles.Disabled(0));
 			R.ExpectTrue(Particles.Disabled(1));
 			R.ExpectTrue(Particles.Disabled(2));
@@ -1168,6 +1217,7 @@ namespace GeometryCollectionExample
 			R.ExpectTrue(!Particles.Disabled(4));
 			R.ExpectTrue(!Particles.Disabled(5));
 			R.ExpectTrue(Particles.Disabled(6));
+#endif
 		}
 
 		FChaosSolversModule::GetModule()->DestroySolver(Solver);
diff --git a/Engine/Plugins/Experimental/GeometryCollectionPlugin/Source/GeometryCollectionExampleCore/Private/GeometryCollection/GeometryCollectionExampleSimulationSolver.cpp b/Engine/Plugins/Experimental/GeometryCollectionPlugin/Source/GeometryCollectionExampleCore/Private/GeometryCollection/GeometryCollectionExampleSimulationSolver.cpp
--- a/Engine/Plugins/Experimental/GeometryCollectionPlugin/Source/GeometryCollectionExampleCore/Private/GeometryCollection/GeometryCollectionExampleSimulationSolver.cpp
+++ b/Engine/Plugins/Experimental/GeometryCollectionPlugin/Source/GeometryCollectionExampleCore/Private/GeometryCollection/GeometryCollectionExampleSimulationSolver.cpp
@@ -16,11 +16,13 @@
 #include "Field/FieldSystemNodes.h"
 #include "GeometryCollection/GeometryDynamicCollection.h"
 #include "GeometryCollection/GeometryCollectionSimulationTypes.h"
-#include "SolverObjects/GeometryCollectionPhysicsObject.h"
-#include "SolverObjects/SolverObjects.h"
+#include "PhysicsProxy/GeometryCollectionPhysicsProxy.h"
+#include "PhysicsProxy/PhysicsProxies.h"
 #include "Chaos/ErrorReporter.h"
 #include "Chaos/Defines.h"
 #include "ChaosSolversModule.h"
+#include "PhysicsSolver.h"
+#include "EventsData.h"
 
 #define SMALL_THRESHOLD 1e-4
 
@@ -67,11 +69,13 @@ namespace GeometryCollectionExample
 			BuildSimulationData(ErrorReporter,*RestCollection, InParams.Shared);
 		};
 
-		FGeometryCollectionPhysicsObject* PhysObject = new FGeometryCollectionPhysicsObject(nullptr, DynamicCollection.Get(), InitFunc, nullptr, nullptr);;
+		FGeometryCollectionPhysicsProxy* PhysObject = new FGeometryCollectionPhysicsProxy(nullptr, DynamicCollection.Get(), InitFunc, nullptr, nullptr);;
 		PhysObject->Initialize();
 
 		Chaos::FPBDRigidsSolver* Solver = FChaosSolversModule::GetModule()->CreateSolver(true);
+#if CHAOS_PARTICLEHANDLE_TODO
 		Solver->RegisterObject(PhysObject);
+#endif
 		Solver->SetHasFloor(false);
 		Solver->SetEnabled(true);
 		PhysObject->ActivateBodies();
@@ -127,11 +131,13 @@ namespace GeometryCollectionExample
 		};
 
 
-		FGeometryCollectionPhysicsObject* PhysObject = new FGeometryCollectionPhysicsObject(nullptr, DynamicCollection.Get(), InitFunc, nullptr, nullptr);
+		FGeometryCollectionPhysicsProxy* PhysObject = new FGeometryCollectionPhysicsProxy(nullptr, DynamicCollection.Get(), InitFunc, nullptr, nullptr);
 		PhysObject->Initialize();
 
 		Chaos::FPBDRigidsSolver* Solver = FChaosSolversModule::GetModule()->CreateSolver(true);
+#if CHAOS_PARTICLEHANDLE_TODO
 		Solver->RegisterObject(PhysObject);
+#endif
 		Solver->SetHasFloor(true);
 		Solver->SetEnabled(true);
 		PhysObject->ActivateBodies();
@@ -142,7 +148,6 @@ namespace GeometryCollectionExample
 		for (int Frame = 0; Frame < 10; Frame++)
 		{
 			Solver->AdvanceSolverBy(1 / 24.);
-			const Chaos::TPBDRigidParticles<float, 3>& Particles = PhysObject->GetSolver()->GetRigidParticles();
 			CurrentRigidDistance = (Transform[1].GetTranslation() - Transform[0].GetTranslation()).Size();
 		}
 
@@ -192,10 +197,12 @@ namespace GeometryCollectionExample
 				BuildSimulationData(ErrorReporter, *RestCollection, InParams.Shared);
 			};
 
-			FGeometryCollectionPhysicsObject* PhysObject = new FGeometryCollectionPhysicsObject(nullptr, DynamicCollection.Get(), InitFunc, nullptr, nullptr);;
+			FGeometryCollectionPhysicsProxy* PhysObject = new FGeometryCollectionPhysicsProxy(nullptr, DynamicCollection.Get(), InitFunc, nullptr, nullptr);;
 			PhysObject->Initialize();
 
+#if CHAOS_PARTICLEHANDLE_TODO
 			Solver->RegisterObject(PhysObject);
+#endif
 			PhysObject->ActivateBodies();
 
 			RestArray.Add(RestCollection);
@@ -205,7 +212,8 @@ namespace GeometryCollectionExample
 
 		Solver->AdvanceSolverBy(1 / 24.);
 
-		const Chaos::TArrayCollectionArray<SolverObjectWrapper> & SolverMapping = Solver->GetSolverObjectReverseMapping();
+#if TODO_REIMPLEMENT_PHYSICS_PROXY_REVERSE_MAPPING
+		const Chaos::TArrayCollectionArray<PhysicsProxyWrapper> & SolverMapping = Solver->GetPhysicsProxyReverseMapping();
 		const Chaos::TArrayCollectionArray<int32> & ParticleMapping = Solver->GetParticleIndexReverseMapping();
 
 		R.ExpectTrue(SolverMapping.Num() == 11);
@@ -214,16 +222,17 @@ namespace GeometryCollectionExample
 		R.ExpectTrue(ParticleMapping[0] == INDEX_NONE);
 		R.ExpectTrue(ParticleMapping[1] == 0);
 
-		R.ExpectTrue(SolverMapping[0].SolverObject == nullptr);
-		R.ExpectTrue(SolverMapping[0].Type == ESolverObjectType::NoneType);
+		R.ExpectTrue(SolverMapping[0].PhysicsProxy == nullptr);
+		R.ExpectTrue(SolverMapping[0].Type == EPhysicsProxyType::NoneType);
 
-		R.ExpectTrue(SolverMapping[5].SolverObject != nullptr);
-		R.ExpectTrue(SolverMapping[5].Type == ESolverObjectType::GeometryCollectionType);
+		R.ExpectTrue(SolverMapping[5].PhysicsProxy != nullptr);
+		R.ExpectTrue(SolverMapping[5].Type == EPhysicsProxyType::GeometryCollectionType);
 
-		const TManagedArray<int32> & RigidBodyID = static_cast<FGeometryCollectionPhysicsObject*>(SolverMapping[5].SolverObject)->
+		const TManagedArray<int32> & RigidBodyID = static_cast<FGeometryCollectionPhysicsProxy*>(SolverMapping[5].PhysicsProxy)->
 			GetGeometryDynamicCollection_PhysicsThread()->GetAttribute<int32>("RigidBodyID", FGeometryCollection::TransformGroup);
 		R.ExpectTrue(RigidBodyID.Num() == 1);
 		R.ExpectTrue(RigidBodyID[0] == 5);
+#endif
 
 		FChaosSolversModule::GetModule()->DestroySolver(Solver);
 
@@ -283,11 +292,13 @@ namespace GeometryCollectionExample
 				BuildSimulationData(ErrorReporter, *RestCollection, InParams.Shared);
 			};
 
-			FGeometryCollectionPhysicsObject* PhysObject = new FGeometryCollectionPhysicsObject(nullptr, DynamicCollection.Get(), InitFunc, nullptr, nullptr);;
+			FGeometryCollectionPhysicsProxy* PhysObject = new FGeometryCollectionPhysicsProxy(nullptr, DynamicCollection.Get(), InitFunc, nullptr, nullptr);;
 			PhysObject->Initialize();
 			PhysObject->SetCollisionParticlesPerObjectFraction(1.0);
 
+#if CHAOS_PARTICLEHANDLE_TODO
 			Solver->RegisterObject(PhysObject);
+#endif
 			Solver->SetHasFloor(true);
 			Solver->SetEnabled(true);
 			PhysObject->ActivateBodies();
@@ -300,6 +311,32 @@ namespace GeometryCollectionExample
 	}
 #endif // INCLUDE_CHAOS
 
+
+	class EventHarvester
+	{
+	public:
+#if INCLUDE_CHAOS
+		EventHarvester(Chaos::FPBDRigidsSolver* Solver)
+		{
+			Solver->GetEventManager()->RegisterHandler<Chaos::FCollisionEventData>(Chaos::EEventType::Collision, this, &EventHarvester::HandleCollisionEvents);
+			Solver->GetEventManager()->RegisterHandler<Chaos::FBreakingEventData>(Chaos::EEventType::Breaking, this, &EventHarvester::HandleBreakingEvents);
+		}
+
+		void HandleCollisionEvents(const Chaos::FCollisionEventData& Events)
+		{
+			CollisionEventData = Events;
+		}
+
+		void HandleBreakingEvents(const Chaos::FBreakingEventData& Events)
+		{
+			BreakingEventData = Events;
+		}
+
+		Chaos::FCollisionEventData CollisionEventData;
+		Chaos::FBreakingEventData BreakingEventData;
+#endif
+	};
+
 	template<class T>
 	bool Solver_CollisionEventFilter(ExampleResponse&& R)
 	{
@@ -323,9 +360,11 @@ namespace GeometryCollectionExample
 
 		Solver->SetGenerateCollisionData(true);
 		Solver->SetCollisionFilterSettings(CollisionFilterSettings);
+		EventHarvester Events(Solver);
 
+#if TODO_REIMPLEMENT_GET_RIGID_PARTICLES
 		// modify mass for filter test
-		Chaos::TPBDRigidParticles<float, 3>& Particles = Solver->GetRigidParticles();
+		FParticlesType& Particles = (FParticlesType&)Solver->GetRigidParticles();
 		for (int32 i = 0; i < 10; i++)
 		{
 			Particles.M(i) = i;
@@ -338,15 +377,12 @@ namespace GeometryCollectionExample
 			Solver->SyncEvents_GameThread();
 			Solver->AdvanceSolverBy(1 / 24.);
 
-			Chaos::FPBDRigidsSolver::FScopedGetEventsData SopedAccess = Solver->ScopedGetEventsData();
-			const Chaos::FPBDRigidsSolver::FAllCollisionDataMaps& CollisionMaps = SopedAccess.GetAllCollisions_Maps();
-
-			Impact = (CollisionMaps.AllCollisionData != nullptr);
+			const auto& AllCollisionsArray = Events.CollisionEventData.CollisionData.AllCollisionsArray;
+			Impact = (AllCollisionsArray.Num() > 0);
 
 			if (Impact)
 			{
 				// any objects with a mass of less than 6 are removed from returned collision data
-				const auto& AllCollisionsArray = CollisionMaps.AllCollisionData->AllCollisionsArray;
 				R.ExpectTrue(AllCollisionsArray.Num() == 4);
 
 				for (const auto& Collision : AllCollisionsArray)
@@ -358,17 +394,16 @@ namespace GeometryCollectionExample
 					R.ExpectTrue(Collision.Velocity2.Z == 0.0f);
 				}
 			}
-
 		} while (!Impact);
 
+#endif
 		FChaosSolversModule::GetModule()->DestroySolver(Solver);
-
 #endif
+
 		return !R.HasError();
 	}
 	template bool Solver_CollisionEventFilter<float>(ExampleResponse&& R);
 
-
 	template<class T>
 	bool Solver_BreakingEventFilter(ExampleResponse&& R)
 	{
@@ -393,13 +428,15 @@ namespace GeometryCollectionExample
 		Solver->SetGenerateBreakingData(true);
 		Solver->SetBreakingFilterSettings(BreakingFilterSettings);
 
+		EventHarvester Events(Solver);
+
+#if TODO_REIMPLEMENT_GET_RIGID_PARTICLES
 		// modify mass for filter test
 		Chaos::TPBDRigidParticles<float, 3>& Particles = Solver->GetRigidParticles();
-		Particles.M(1) = TestMass+1.0f;
+		Particles.M(1) = TestMass + 1.0f;
 		Particles.M(2) = TestMass - 1.0f;
 		Particles.M(3) = TestMass - 2.0f;
-		Particles.M(4) = TestMass+2.0f;
-
+		Particles.M(4) = TestMass + 2.0f;
 
 		bool Impact = false;
 		do
@@ -409,10 +446,8 @@ namespace GeometryCollectionExample
 			Solver->AdvanceSolverBy(1 / 24.);
 			FinalizeSolver(*Solver);
 
-			Chaos::FPBDRigidsSolver::FScopedGetEventsData SopedAccess = Solver->ScopedGetEventsData();
-			const Chaos::FPBDRigidsSolver::FAllBreakingDataMaps& BreakingMaps = SopedAccess.GetAllBreakings_Maps();
-
-			Impact = (BreakingMaps.AllBreakingData != nullptr);
+			const auto& AllBreakingsArray = Events.BreakingEventData.BreakingData.AllBreakingsArray;
+			Impact = (AllBreakingsArray.Num() > 0);
 
 			if (Impact)
 			{
@@ -424,7 +459,6 @@ namespace GeometryCollectionExample
 				R.ExpectTrue(Particles.Disabled(5) == true); // cluster
 
 				// breaking data
-				const auto& AllBreakingsArray = BreakingMaps.AllBreakingData->AllBreakingsArray;
 				R.ExpectTrue(AllBreakingsArray.Num() == 2); // 2 pieces filtered out of 4
 
 				R.ExpectTrue(AllBreakingsArray[0].ParticleIndex == 4);
@@ -433,11 +467,10 @@ namespace GeometryCollectionExample
 				R.ExpectTrue(AllBreakingsArray[1].Mass == TestMass + 1.0f);
 
 			}
-
 		} while (!Impact);
 
+#endif
 		FChaosSolversModule::GetModule()->DestroySolver(Solver);
-
 #endif
 		return !R.HasError();
 	}
diff --git a/Engine/Plugins/Experimental/GeometryCollectionPlugin/Source/GeometryCollectionExampleCore/Private/GeometryCollection/GeometryCollectionExampleSimulationStreaming.cpp b/Engine/Plugins/Experimental/GeometryCollectionPlugin/Source/GeometryCollectionExampleCore/Private/GeometryCollection/GeometryCollectionExampleSimulationStreaming.cpp
--- a/Engine/Plugins/Experimental/GeometryCollectionPlugin/Source/GeometryCollectionExampleCore/Private/GeometryCollection/GeometryCollectionExampleSimulationStreaming.cpp
+++ b/Engine/Plugins/Experimental/GeometryCollectionPlugin/Source/GeometryCollectionExampleCore/Private/GeometryCollection/GeometryCollectionExampleSimulationStreaming.cpp
@@ -14,9 +14,10 @@
 #include "Field/FieldSystemNodes.h"
 #include "GeometryCollection/GeometryDynamicCollection.h"
 #include "GeometryCollection/GeometryCollectionSimulationTypes.h"
-#include "SolverObjects/SolverObjects.h"
+#include "PhysicsProxy/PhysicsProxies.h"
 #include "Chaos/ErrorReporter.h"
 #include "ChaosSolversModule.h"
+#include "PBDRigidsSolver.h"
 
 #define SMALL_THRESHOLD 1e-4
 
@@ -31,6 +32,7 @@ namespace GeometryCollectionExample
 		Solver->SetEnabled(true);
 
 		Solver->AdvanceSolverBy(1 / 24.);
+#if TODO_REIMPLEMENT_GET_RIGID_PARTICLES
 		Chaos::TPBDRigidParticles<float, 3>& Particles = Solver->GetRigidParticles();
 
 		TArray< SimulationObjects<T>* > Collections;
@@ -40,9 +42,9 @@ namespace GeometryCollectionExample
 			if (Frame % 100 == 0)
 			{
 				Collections.Add(new SimulationObjects<T>());
-				Solver->RegisterObject(Collections.Last()->PhysicsObject.Get());
-				Collections.Last()->PhysicsObject.Get()->Initialize();
-				Collections.Last()->PhysicsObject.Get()->ActivateBodies();
+				Solver->RegisterObject(Collections.Last()->PhysicsProxy.Get());
+				Collections.Last()->PhysicsProxy.Get()->Initialize();
+				Collections.Last()->PhysicsProxy.Get()->ActivateBodies();
 			}
 		}
 
@@ -53,6 +55,7 @@ namespace GeometryCollectionExample
 		// cleanup
 		for (auto Obj : Collections)
 			delete Obj;
+#endif
 		delete Solver;
 #endif
 		return !R.HasError();
@@ -71,6 +74,7 @@ namespace GeometryCollectionExample
 		Solver->SetHasFloor(false);
 		Solver->SetEnabled(true);
 		Solver->AdvanceSolverBy(1 / 24.);
+#if TODO_REIMPLEMENT_GET_RIGID_PARTICLES
 		Chaos::TPBDRigidParticles<float, 3>& Particles = Solver->GetRigidParticles();
 
 		TArray< SimulationObjects<T>* > Collections;
@@ -80,12 +84,12 @@ namespace GeometryCollectionExample
 			if (Frame % 100 == 0)
 			{
 				Collections.Add(new SimulationObjects<T>(P));
-				Solver->RegisterObject(Collections.Last()->PhysicsObject.Get());
-				Collections.Last()->PhysicsObject.Get()->Initialize();
+				Solver->RegisterObject(Collections.Last()->PhysicsProxy.Get());
+				Collections.Last()->PhysicsProxy.Get()->Initialize();
 			}
 		}
 		for (auto Obj : Collections)
-			Obj->PhysicsObject->ActivateBodies();
+			Obj->PhysicsProxy->ActivateBodies();
 
 		for (int32 Frame = 1; Frame < 100; Frame++)
 		{
@@ -100,6 +104,7 @@ namespace GeometryCollectionExample
 		// cleanup
 		for (auto Obj : Collections)
 			delete Obj;
+#endif
 		delete Solver;
 #endif
 		return !R.HasError();
@@ -121,6 +126,7 @@ namespace GeometryCollectionExample
 		Solver->SetHasFloor(false);
 		Solver->SetEnabled(true);
 		Solver->AdvanceSolverBy(1 / 24.);
+#if TODO_REIMPLEMENT_GET_RIGID_PARTICLES
 		Chaos::TPBDRigidParticles<float, 3>& Particles = Solver->GetRigidParticles();
 
 		TArray< SimulationObjects<T>* > Collections;
@@ -130,8 +136,8 @@ namespace GeometryCollectionExample
 			if (Frame % 100 == 0)
 			{
 				Collections.Add(new SimulationObjects<T>(P, GeometryCollection::MakeCubeElement(FTransform(FQuat(ForceInit), FVector(Frame)), FVector(1.0))));
-				Solver->RegisterObject(Collections.Last()->PhysicsObject.Get());
-				Collections.Last()->PhysicsObject.Get()->Initialize();
+				Solver->RegisterObject(Collections.Last()->PhysicsProxy.Get());
+				Collections.Last()->PhysicsProxy.Get()->Initialize();
 			}
 		}
 
@@ -143,7 +149,7 @@ namespace GeometryCollectionExample
 		}
 
 		for (auto Obj : Collections)
-			Obj->PhysicsObject->ActivateBodies();
+			Obj->PhysicsProxy->ActivateBodies();
 
 		// all particles should be enabled
 		R.ExpectTrue(Particles.Size() == 9);
@@ -166,10 +172,10 @@ namespace GeometryCollectionExample
 		R.ExpectTrue(Particles.Disabled(Particles.Size()-1) == false);
 		R.ExpectTrue(Particles.X(Particles.Size() - 1).Z < -1.f);
 
-
 		// cleanup
 		for (auto Obj : Collections)
 			delete Obj;
+#endif
 		delete Solver;
 #endif
 		return !R.HasError();
diff --git a/Engine/Plugins/Experimental/GeometryCollectionPlugin/Source/GeometryCollectionExampleCore/Private/GeometryCollection/GeometryCollectionExampleUtility.cpp b/Engine/Plugins/Experimental/GeometryCollectionPlugin/Source/GeometryCollectionExampleCore/Private/GeometryCollection/GeometryCollectionExampleUtility.cpp
--- a/Engine/Plugins/Experimental/GeometryCollectionPlugin/Source/GeometryCollectionExampleCore/Private/GeometryCollection/GeometryCollectionExampleUtility.cpp
+++ b/Engine/Plugins/Experimental/GeometryCollectionPlugin/Source/GeometryCollectionExampleCore/Private/GeometryCollection/GeometryCollectionExampleUtility.cpp
@@ -9,7 +9,7 @@
 #include "GeometryCollection/GeometryCollectionAlgo.h"
 
 #include "PBDRigidsSolver.h"
-#include "SolverObjects/SolverObjects.h"
+#include "PhysicsProxy/PhysicsProxies.h"
 #include "../Resource/FracturedGeometry.h"
 
 
@@ -37,13 +37,11 @@ namespace GeometryCollectionExample {
 	void FinalizeSolver(Chaos::FPBDRigidsSolver& InSolver)
 	{
 #if INCLUDE_CHAOS
-		FSolverObjectStorage& Objects = InSolver.GetObjectStorage();
-
-		Objects.ForEachSolverObject([](auto* Object)
+		InSolver.ForEachPhysicsProxy([](auto* Object)
 		{
-			Object->CacheResults();
-			Object->FlipCache();
-			Object->SyncToCache();
+			Object->BufferPhysicsResults();
+			Object->FlipBuffer();
+			Object->PullFromPhysicsState();
 		});
 #endif
 	}
diff --git a/Engine/Plugins/Experimental/GeometryCollectionPlugin/Source/GeometryCollectionExampleCore/Public/GeometryCollection/GeometryCollectionExampleClustering.h b/Engine/Plugins/Experimental/GeometryCollectionPlugin/Source/GeometryCollectionExampleCore/Public/GeometryCollection/GeometryCollectionExampleClustering.h
--- a/Engine/Plugins/Experimental/GeometryCollectionPlugin/Source/GeometryCollectionExampleCore/Public/GeometryCollection/GeometryCollectionExampleClustering.h
+++ b/Engine/Plugins/Experimental/GeometryCollectionPlugin/Source/GeometryCollectionExampleCore/Public/GeometryCollection/GeometryCollectionExampleClustering.h
@@ -3,9 +3,12 @@
 
 #include "GeometryCollection/GeometryCollectionExample.h"
 #include "GeometryCollection/GeometryCollectionExampleResponse.h"
+#include "Chaos/Declares.h"
+
 namespace GeometryCollectionExample
 {
 
+#if TODO_REIMPLEMENT_RIGID_CLUSTERING
 	template<class T>
 	bool RigidBodies_ClusterTest_SingleLevelNonBreaking(ExampleResponse&& R);
 
@@ -51,4 +54,6 @@ namespace GeometryCollectionExample
 	template<class T>
 	bool RigidBodiess_ClusterTest_ParticleImplicitCollisionGeometry(ExampleResponse&& R);
 
+#endif
+
 }
diff --git a/Engine/Plugins/Experimental/GeometryCollectionPlugin/Source/GeometryCollectionExampleCore/Public/GeometryCollection/GeometryCollectionExampleEvents.h b/Engine/Plugins/Experimental/GeometryCollectionPlugin/Source/GeometryCollectionExampleCore/Public/GeometryCollection/GeometryCollectionExampleEvents.h
new file mode 100644
--- /dev/null
+++ b/Engine/Plugins/Experimental/GeometryCollectionPlugin/Source/GeometryCollectionExampleCore/Public/GeometryCollection/GeometryCollectionExampleEvents.h
@@ -0,0 +1,12 @@
+// Copyright 1998-2019 Epic Games, Inc. All Rights Reserved.
+#pragma once
+
+#include "GeometryCollection/GeometryCollectionExample.h"
+#include "GeometryCollection/GeometryCollectionExampleResponse.h"
+
+namespace GeometryCollectionExample
+{
+	template<class T>
+	bool Event_Handler(ExampleResponse&& R);
+
+}
diff --git a/Engine/Plugins/Experimental/GeometryCollectionPlugin/Source/GeometryCollectionExampleCore/Public/GeometryCollection/GeometryCollectionExampleSimulationObject.h b/Engine/Plugins/Experimental/GeometryCollectionPlugin/Source/GeometryCollectionExampleCore/Public/GeometryCollection/GeometryCollectionExampleSimulationObject.h
--- a/Engine/Plugins/Experimental/GeometryCollectionPlugin/Source/GeometryCollectionExampleCore/Public/GeometryCollection/GeometryCollectionExampleSimulationObject.h
+++ b/Engine/Plugins/Experimental/GeometryCollectionPlugin/Source/GeometryCollectionExampleCore/Public/GeometryCollection/GeometryCollectionExampleSimulationObject.h
@@ -7,7 +7,7 @@
 #include "Chaos/ErrorReporter.h"
 #include "GeometryCollection/GeometryCollectionSimulationTypes.h"
 #include "GeometryCollection/GeometryCollectionUtility.h"
-#include "SolverObjects/GeometryCollectionPhysicsObject.h"
+#include "PhysicsProxy/GeometryCollectionPhysicsProxy.h"
 #include "Templates/SharedPointer.h"
 
 #if INCLUDE_CHAOS
@@ -51,7 +51,7 @@ namespace GeometryCollectionExample {
 			, RestCollection(RestCollectionIn)
 			, DynamicCollection(GeometryCollectionToGeometryDynamicCollection(RestCollection.Get()))
 			, PhysicalMaterial(new Chaos::TChaosPhysicsMaterial<T>())
-			, PhysicsObject(new FGeometryCollectionPhysicsObject(nullptr, DynamicCollection.Get(), [&](FSimulationParameters& P) {Init(P); }, nullptr, nullptr))
+			, PhysicsProxy(new FGeometryCollectionPhysicsProxy(nullptr, DynamicCollection.Get(), [&](FSimulationParameters& P) {Init(P); }, nullptr, nullptr))
 		{
 			PhysicalMaterial->Friction = 0;
 			PhysicalMaterial->Restitution = 0;
@@ -70,13 +70,13 @@ namespace GeometryCollectionExample {
 			, RestCollection(InRestCollection)
 			, DynamicCollection(InDynamicCollection)
 			, PhysicalMaterial(InPhysicalMaterial)
-			, PhysicsObject(new FGeometryCollectionPhysicsObject(nullptr, DynamicCollection.Get(), [&](FSimulationParameters& P) {Init(P); }, nullptr, nullptr))
+			, PhysicsProxy(new FGeometryCollectionPhysicsProxy(nullptr, DynamicCollection.Get(), [&](FSimulationParameters& P) {Init(P); }, nullptr, nullptr))
 		{
 		}
 
 		void Init(FSimulationParameters& InParams)
 		{
-			PhysicsObject->SetCollisionParticlesPerObjectFraction(1.0);
+			PhysicsProxy->SetCollisionParticlesPerObjectFraction(1.0);
 			InParams.RestCollection = RestCollection.Get();
 			InParams.DynamicCollection = DynamicCollection.Get();
 			InParams.PhysicalMaterial = MakeSerializable(PhysicalMaterial);
@@ -99,7 +99,7 @@ namespace GeometryCollectionExample {
 		TSharedPtr<FGeometryCollection> RestCollection;
 		TSharedPtr<FGeometryDynamicCollection> DynamicCollection;
 		TUniquePtr<Chaos::TChaosPhysicsMaterial<T>> PhysicalMaterial;
-		TSharedPtr<FGeometryCollectionPhysicsObject> PhysicsObject;
+		TSharedPtr<FGeometryCollectionPhysicsProxy> PhysicsProxy;
 		FSharedSimulationSizeSpecificData SimulationData;
 	};
 
diff --git a/Engine/Plugins/Experimental/GeometryCollectionPlugin/Source/GeometryCollectionExampleCore/Public/GeometryCollection/GeometryCollectionExampleUtility.h b/Engine/Plugins/Experimental/GeometryCollectionPlugin/Source/GeometryCollectionExampleCore/Public/GeometryCollection/GeometryCollectionExampleUtility.h
--- a/Engine/Plugins/Experimental/GeometryCollectionPlugin/Source/GeometryCollectionExampleCore/Public/GeometryCollection/GeometryCollectionExampleUtility.h
+++ b/Engine/Plugins/Experimental/GeometryCollectionPlugin/Source/GeometryCollectionExampleCore/Public/GeometryCollection/GeometryCollectionExampleUtility.h
@@ -4,7 +4,7 @@
 #include "ChaosSolversModule.h"
 #include "Chaos/Serializable.h"
 #include "GeometryCollection/GeometryCollectionSimulationTypes.h"
-#include "SolverObjects/GeometryCollectionPhysicsObject.h"
+#include "PhysicsProxy/GeometryCollectionPhysicsProxy.h"
 #include "Templates/SharedPointer.h"
 
 class FGeometryCollection;
diff --git a/Engine/Plugins/Experimental/HairStrands/HairStrands.uplugin b/Engine/Plugins/Experimental/HairStrands/HairStrands.uplugin
new file mode 100644
--- /dev/null
+++ b/Engine/Plugins/Experimental/HairStrands/HairStrands.uplugin
@@ -0,0 +1,52 @@
+{
+	"FileVersion": 3,
+	"Version": 1,
+	"VersionName": "0.1",
+	"FriendlyName": "Hair Strands",
+	"Description": "Build Hair Strands for rendering and simulation",
+	"Category": "Geometry",
+	"CreatedBy": "Epic Games, Inc.",
+	"CreatedByURL": "http://epicgames.com",
+	"DocsURL": "",
+	"MarketplaceURL": "",
+	"SupportURL": "",
+	"EnabledByDefault": false,
+	"CanContainContent": true,
+	"IsBetaVersion": true,
+	"Installed": false,
+	"Modules": [
+		{
+			"Name": "HairStrandsCore",
+			"Type": "Runtime",
+			"LoadingPhase": "Default",
+			"WhitelistPlatforms": [
+				"Win64",
+				"Win32"
+			] 
+		},
+		{
+			"Name": "HairStrandsEditor",
+			"Type": "Editor",
+			"LoadingPhase": "Default",
+			"WhitelistPlatforms": [
+				"Win64",
+				"Win32"
+			]
+		},
+		{
+			"Name": "HairStrandsNiagara",
+			"Type": "RunTime",
+			"LoadingPhase": "Default",
+			"WhitelistPlatforms": [
+				"Win64",
+				"Win32"
+			]
+		}
+	],
+	"Plugins": [
+		{
+			"Name": "Niagara",
+			"Enabled": true
+		}
+	]
+}
