Fix making GeometryCollection from empty Static Mesh crashing.

diff --git a/Engine/Plugins/Experimental/GeometryCollectionPlugin/Source/GeometryCollectionEditor/Private/GeometryCollection/GeometryCollectionConversion.cpp b/Engine/Plugins/Experimental/GeometryCollectionPlugin/Source/GeometryCollectionEditor/Private/GeometryCollection/GeometryCollectionConversion.cpp
--- a/Engine/Plugins/Experimental/GeometryCollectionPlugin/Source/GeometryCollectionEditor/Private/GeometryCollection/GeometryCollectionConversion.cpp
+++ b/Engine/Plugins/Experimental/GeometryCollectionPlugin/Source/GeometryCollectionEditor/Private/GeometryCollection/GeometryCollectionConversion.cpp
@@ -28,7 +28,12 @@ DEFINE_LOG_CATEGORY_STATIC(UGeometryCollectionConversionLogging, Log, All);
 void FGeometryCollectionConversion::AppendStaticMesh(const UStaticMesh * StaticMesh, const UStaticMeshComponent *StaticMeshComponent, const FTransform & StaticMeshTransform, UGeometryCollection * GeometryCollectionObject, bool ReindexMaterials)
 {
 	//UE_LOG(UGeometryCollectionConversionLogging, Log, TEXT("FGeometryCollectionConversion::AppendStaticMesh()"));
-	check(StaticMesh);
+
+	if (StaticMesh == nullptr)
+	{
+		return;
+	}
+
 	check(GeometryCollectionObject);
 	TSharedPtr<FGeometryCollection, ESPMode::ThreadSafe> GeometryCollectionPtr = GeometryCollectionObject->GetGeometryCollection();
 	FGeometryCollection* GeometryCollection = GeometryCollectionPtr.Get();
