Fix some compile errors running with /permissive- on Windows.

#rnx

diff --git a/Engine/Plugins/Experimental/GeometryCollectionPlugin/Source/GeometryCollectionExampleCore/Private/GeometryCollection/GeometryCollectionExampleFields.cpp b/Engine/Plugins/Experimental/GeometryCollectionPlugin/Source/GeometryCollectionExampleCore/Private/GeometryCollection/GeometryCollectionExampleFields.cpp
--- a/Engine/Plugins/Experimental/GeometryCollectionPlugin/Source/GeometryCollectionExampleCore/Private/GeometryCollection/GeometryCollectionExampleFields.cpp
+++ b/Engine/Plugins/Experimental/GeometryCollectionPlugin/Source/GeometryCollectionExampleCore/Private/GeometryCollection/GeometryCollectionExampleFields.cpp
@@ -20,7 +20,6 @@ DEFINE_LOG_CATEGORY_STATIC(GCTF_Log, Verbose, All);
 
 namespace GeometryCollectionExample
 {
-	template<class T>
 	bool Fields_NoiseSample(ExampleResponse&& R)
 	{
 		int Bounds = 100;
@@ -78,9 +77,7 @@ namespace GeometryCollectionExample
 		return !R.HasError();
 
 	}
-	template bool Fields_NoiseSample<float>(ExampleResponse&& R);
 
-	template<class T>
 	bool Fields_RadialIntMask(ExampleResponse&& R)
 	{
 		TArray<ContextIndex> IndicesArray;
@@ -128,9 +125,7 @@ namespace GeometryCollectionExample
 
 		return !R.HasError();
 	}
-	template bool Fields_RadialIntMask<float>(ExampleResponse&& R);
 
-	template<class T>
 	bool Fields_RadialFalloff(ExampleResponse&& R)
 	{
 		TArray<ContextIndex> IndicesArray;
@@ -183,9 +178,7 @@ namespace GeometryCollectionExample
 
 		return !R.HasError();
 	}
-	template bool Fields_RadialFalloff<float>(ExampleResponse&& R);
 
-	template<class T>
 	bool Fields_PlaneFalloff(ExampleResponse&& R)
 	{
 		TArray<ContextIndex> IndicesArray;
@@ -233,10 +226,7 @@ namespace GeometryCollectionExample
 
 		return !R.HasError();
 	}
-	template bool Fields_PlaneFalloff<float>(ExampleResponse&& R);
-
 
-	template<class T>
 	bool Fields_UniformVector(ExampleResponse&& R)
 	{
 		TArray<ContextIndex> IndicesArray;
@@ -280,10 +270,7 @@ namespace GeometryCollectionExample
 
 		return !R.HasError();
 	}
-	template bool Fields_UniformVector<float>(ExampleResponse&& R);
-
 
-	template<class T>
 	bool Fields_RaidalVector(ExampleResponse&& R)
 	{
 		TArray<ContextIndex> IndicesArray;
@@ -327,10 +314,7 @@ namespace GeometryCollectionExample
 
 		return !R.HasError();
 	}
-	template bool Fields_RaidalVector<float>(ExampleResponse&& R);
 
-
-	template<class T>
 	bool Fields_SumVectorFullMult(ExampleResponse&& R)
 	{
 		TArray<ContextIndex> IndicesArray;
@@ -403,9 +387,7 @@ namespace GeometryCollectionExample
 
 		return !R.HasError();
 	}
-	template bool Fields_SumVectorFullMult<float>(ExampleResponse&& R);
 
-	template<class T>
 	bool Fields_SumVectorFullDiv(ExampleResponse&& R)
 	{
 		TArray<ContextIndex> IndicesArray;
@@ -478,10 +460,8 @@ namespace GeometryCollectionExample
 
 		return !R.HasError();
 	}
-	template bool Fields_SumVectorFullDiv<float>(ExampleResponse&& R);
 
 
-	template<class T>
 	bool Fields_SumVectorFullAdd(ExampleResponse&& R)
 	{
 		TArray<ContextIndex> IndicesArray;
@@ -555,10 +535,8 @@ namespace GeometryCollectionExample
 
 		return !R.HasError();
 	}
-	template bool Fields_SumVectorFullAdd<float>(ExampleResponse&& R);
 
 
-	template<class T>
 	bool Fields_SumVectorFullSub(ExampleResponse&& R)
 	{
 		TArray<ContextIndex> IndicesArray;
@@ -632,10 +610,8 @@ namespace GeometryCollectionExample
 
 		return !R.HasError();
 	}
-	template bool Fields_SumVectorFullSub<float>(ExampleResponse&& R);
 
 
-	template<class T>
 	bool Fields_SumVectorLeftSide(ExampleResponse&& R)
 	{
 		TArray<ContextIndex> IndicesArray;
@@ -709,9 +685,7 @@ namespace GeometryCollectionExample
 
 		return !R.HasError();
 	}
-	template bool Fields_SumVectorLeftSide<float>(ExampleResponse&& R);
 
-	template<class T>
 	bool Fields_SumVectorRightSide(ExampleResponse&& R)
 	{
 		TArray<ContextIndex> IndicesArray;
@@ -787,10 +761,8 @@ namespace GeometryCollectionExample
 
 		return !R.HasError();
 	}
-	template bool Fields_SumVectorRightSide<float>(ExampleResponse&& R);
 
 
-	template<class T>
 	bool Fields_SumScalar(ExampleResponse&& R)
 	{
 		int32 NumPoints = 20;
@@ -867,10 +839,8 @@ namespace GeometryCollectionExample
 
 		return !R.HasError();
 	}
-	template bool Fields_SumScalar<float>(ExampleResponse&& R);
 
 
-	template<class T>
 	bool Fields_SumScalarRightSide(ExampleResponse&& R)
 	{
 		int32 NumPoints = 20;
@@ -939,10 +909,8 @@ namespace GeometryCollectionExample
 
 		return !R.HasError();
 	}
-	template bool Fields_SumScalarRightSide<float>(ExampleResponse&& R);
 
 
-	template<class T>
 	bool Fields_SumScalarLeftSide(ExampleResponse&& R)
 	{
 		int32 NumPoints = 20;
@@ -1012,12 +980,10 @@ namespace GeometryCollectionExample
 
 		return !R.HasError();
 	}
-	template bool Fields_SumScalarLeftSide<float>(ExampleResponse&& R);
 
 
 
 
-	template<class T>
 	bool Fields_Culling(ExampleResponse&& R)
 	{
 		int32 NumPoints = 20;
@@ -1087,7 +1053,6 @@ namespace GeometryCollectionExample
 
 		return !R.HasError();
 	}
-	template bool Fields_Culling<float>(ExampleResponse&& R);
 
 
 	FFieldSystemCommand SaveAndLoad(FFieldSystemCommand & CommandOut)
diff --git a/Engine/Plugins/Experimental/GeometryProcessing/Source/GeometricObjects/Public/Intersection/Intersector1.h b/Engine/Plugins/Experimental/GeometryProcessing/Source/GeometricObjects/Public/Intersection/Intersector1.h
--- a/Engine/Plugins/Experimental/GeometryProcessing/Source/GeometricObjects/Public/Intersection/Intersector1.h
+++ b/Engine/Plugins/Experimental/GeometryProcessing/Source/GeometricObjects/Public/Intersection/Intersector1.h
@@ -69,7 +69,7 @@ public:
 	 */
 	bool Find()
 	{
-		if (U.Max < V.Min || U.Min > V.Max) 
+		if ((U.Max < V.Min) || (U.Min > V.Max))
 		{
 			NumIntersections = 0;
 		}
diff --git a/Engine/Plugins/Runtime/nDisplay/Source/PicpMPCDI/Private/PicpMPCDIShader.h b/Engine/Plugins/Runtime/nDisplay/Source/PicpMPCDI/Private/PicpMPCDIShader.h
--- a/Engine/Plugins/Runtime/nDisplay/Source/PicpMPCDI/Private/PicpMPCDIShader.h
+++ b/Engine/Plugins/Runtime/nDisplay/Source/PicpMPCDI/Private/PicpMPCDIShader.h
@@ -13,6 +13,7 @@
 
 #include "IMPCDI.h"
 #include "IPicpMPCDI.h"
+#include "MPCDIRegion.h"
 
 #include "Overlay/PicpProjectionOverlayRender.h"
 
diff --git a/Engine/Source/Developer/Windows/LiveCoding/Private/External/LC_PointerUtil.h b/Engine/Source/Developer/Windows/LiveCoding/Private/External/LC_PointerUtil.h
--- a/Engine/Source/Developer/Windows/LiveCoding/Private/External/LC_PointerUtil.h
+++ b/Engine/Source/Developer/Windows/LiveCoding/Private/External/LC_PointerUtil.h
@@ -8,6 +8,12 @@
 
 namespace pointer
 {
+	template <typename T>
+	T As(void* anyPointer);
+
+	template <typename T>
+	T As(const void* anyPointer);
+
 	template <typename T>
 	T AsInteger(const void* anyPointer)
 	{
diff --git a/Engine/Source/Developer/Windows/LiveCodingServer/Private/External/LC_Coff.cpp b/Engine/Source/Developer/Windows/LiveCodingServer/Private/External/LC_Coff.cpp
--- a/Engine/Source/Developer/Windows/LiveCodingServer/Private/External/LC_Coff.cpp
+++ b/Engine/Source/Developer/Windows/LiveCodingServer/Private/External/LC_Coff.cpp
@@ -411,9 +411,9 @@ namespace coff
 	template <typename RawCoffType, coffDetail::CoffType::Enum Type>
 	static RawCoffType* ReadRaw(ObjFile* file, uint32_t uniqueId, coff::ReadFlags::Enum readFlags)
 	{
-		using CoffHeader = coffDetail::CoffType::TypeMap<Type>::HeaderType;
-		using CoffSymbol = coffDetail::CoffType::TypeMap<Type>::SymbolType;
-		using CoffAuxSymbol = coffDetail::CoffType::TypeMap<Type>::AuxSymbolType;
+		using CoffHeader = typename coffDetail::CoffType::TypeMap<Type>::HeaderType;
+		using CoffSymbol = typename coffDetail::CoffType::TypeMap<Type>::SymbolType;
+		using CoffAuxSymbol = typename coffDetail::CoffType::TypeMap<Type>::AuxSymbolType;
 
 		RawCoffType* rawCoff = new RawCoffType;
 		{
@@ -612,9 +612,9 @@ namespace coff
 	template <coffDetail::CoffType::Enum Type, typename RawCoffType>
 	static void WriteRaw(const wchar_t* filename, const RawCoffType* rawCoff, SymbolRemovalStrategy::Enum removalStrategy)
 	{
-		using CoffHeader = coffDetail::CoffType::TypeMap<Type>::HeaderType;
-		using CoffSymbol = coffDetail::CoffType::TypeMap<Type>::SymbolType;
-		using CoffAuxSymbol = coffDetail::CoffType::TypeMap<Type>::AuxSymbolType;
+		using CoffHeader = typename coffDetail::CoffType::TypeMap<Type>::HeaderType;
+		using CoffSymbol = typename coffDetail::CoffType::TypeMap<Type>::SymbolType;
+		using CoffAuxSymbol = typename coffDetail::CoffType::TypeMap<Type>::AuxSymbolType;
 
 		// the file structure is as follows:
 		//	- COFF file header
@@ -1009,7 +1009,7 @@ namespace coff
 	template <coffDetail::CoffType::Enum Type>
 	static types::vector<std::string> ExtractLinkerDirectives(const ObjFile* file)
 	{
-		using CoffHeader = coffDetail::CoffType::TypeMap<Type>::HeaderType;
+		using CoffHeader = typename coffDetail::CoffType::TypeMap<Type>::HeaderType;
 
 		const uint32_t sectionCount = coffDetail::GetNumberOfSections<CoffHeader>(file->memoryFile->base);
 		const IMAGE_SECTION_HEADER* sectionHeader = coffDetail::GetSectionHeader<CoffHeader>(file->memoryFile->base);
@@ -1195,9 +1195,9 @@ namespace coff
 	template <coffDetail::CoffType::Enum Type>
 	static CoffDB* GatherDatabase(ObjFile* file, uint32_t uniqueId, coff::ReadFlags::Enum readFlags)
 	{
-		using CoffHeader = coffDetail::CoffType::TypeMap<Type>::HeaderType;
-		using CoffSymbol = coffDetail::CoffType::TypeMap<Type>::SymbolType;
-		using CoffAuxSymbol = coffDetail::CoffType::TypeMap<Type>::AuxSymbolType;
+		using CoffHeader = typename coffDetail::CoffType::TypeMap<Type>::HeaderType;
+		using CoffSymbol = typename coffDetail::CoffType::TypeMap<Type>::SymbolType;
+		using CoffAuxSymbol = typename coffDetail::CoffType::TypeMap<Type>::AuxSymbolType;
 
 		CoffDB* coffDb = new CoffDB;
 		coffDb->symbolAllocator = new PoolAllocator<PoolAllocatorSingleThreadPolicy>("COFF symbols", sizeof(coff::Symbol), alignof(coff::Symbol), 8192u);
@@ -1733,9 +1733,9 @@ namespace coff
 	template <coffDetail::CoffType::Enum Type>
 	static UnresolvedSymbolDB* GatherUnresolvedSymbolDatabase(ObjFile* file, uint32_t uniqueId, coff::ReadFlags::Enum readFlags)
 	{
-		using CoffHeader = coffDetail::CoffType::TypeMap<Type>::HeaderType;
-		using CoffSymbol = coffDetail::CoffType::TypeMap<Type>::SymbolType;
-		using CoffAuxSymbol = coffDetail::CoffType::TypeMap<Type>::AuxSymbolType;
+		using CoffHeader = typename coffDetail::CoffType::TypeMap<Type>::HeaderType;
+		using CoffSymbol = typename coffDetail::CoffType::TypeMap<Type>::SymbolType;
+		using CoffAuxSymbol = typename coffDetail::CoffType::TypeMap<Type>::AuxSymbolType;
 
 		UnresolvedSymbolDB* symbolDb = new UnresolvedSymbolDB;
 
@@ -1804,9 +1804,9 @@ namespace coff
 	template <coffDetail::CoffType::Enum Type>
 	static ExternalSymbolDB* GatherExternalSymbolDatabase(ObjFile* file, uint32_t uniqueId, coff::ReadFlags::Enum readFlags)
 	{
-		using CoffHeader = coffDetail::CoffType::TypeMap<Type>::HeaderType;
-		using CoffSymbol = coffDetail::CoffType::TypeMap<Type>::SymbolType;
-		using CoffAuxSymbol = coffDetail::CoffType::TypeMap<Type>::AuxSymbolType;
+		using CoffHeader = typename coffDetail::CoffType::TypeMap<Type>::HeaderType;
+		using CoffSymbol = typename coffDetail::CoffType::TypeMap<Type>::SymbolType;
+		using CoffAuxSymbol = typename coffDetail::CoffType::TypeMap<Type>::AuxSymbolType;
 
 		ExternalSymbolDB* symbolDb = new ExternalSymbolDB;
 
diff --git a/Engine/Source/Editor/Experimental/EditorInteractiveToolsFramework/Public/InteractiveToolsCommands.h b/Engine/Source/Editor/Experimental/EditorInteractiveToolsFramework/Public/InteractiveToolsCommands.h
--- a/Engine/Source/Editor/Experimental/EditorInteractiveToolsFramework/Public/InteractiveToolsCommands.h
+++ b/Engine/Source/Editor/Experimental/EditorInteractiveToolsFramework/Public/InteractiveToolsCommands.h
@@ -280,12 +280,12 @@ void TInteractiveToolCommands<CommandContextType>::RegisterUIToolCommand(const F
 	const FString DotString = FString(TEXT(".")) + ToolAction.ActionName;
 
 	FUICommandInfo::MakeCommandInfo(
-		AsShared(),
+		this->AsShared(),
 		UICommandInfo,
 		*ToolAction.ActionName,
 		ToolAction.ShortName,
 		ToolAction.Description,
-		FSlateIcon(GetStyleSetName(), ISlateStyle::Join(GetContextName(), TCHAR_TO_ANSI(*DotString))),
+		FSlateIcon(this->GetStyleSetName(), ISlateStyle::Join(this->GetContextName(), TCHAR_TO_ANSI(*DotString))),
 		EUserInterfaceActionType::Button,
 		FInputChord(ToolAction.DefaultModifiers, ToolAction.DefaultKey)
 		);
