Copying //UE4/Dev-Physics to Dev-Main (//UE4/Dev-Main) [at] 9098278

diff --git a/Engine/Plugins/Experimental/GeometryCollectionPlugin/Source/GeometryCollectionExampleCore/Private/GeometryCollection/GeometryCollectionExampleClustering.cpp b/Engine/Plugins/Experimental/GeometryCollectionPlugin/Source/GeometryCollectionExampleCore/Private/GeometryCollection/GeometryCollectionExampleClustering.cpp
--- a/Engine/Plugins/Experimental/GeometryCollectionPlugin/Source/GeometryCollectionExampleCore/Private/GeometryCollection/GeometryCollectionExampleClustering.cpp
+++ b/Engine/Plugins/Experimental/GeometryCollectionPlugin/Source/GeometryCollectionExampleCore/Private/GeometryCollection/GeometryCollectionExampleClustering.cpp
@@ -20,6 +20,7 @@
 #include "PhysicsProxy/PhysicsProxies.h"
 #include "Chaos/ErrorReporter.h"
 #include "ChaosSolversModule.h"
+#include "PhysicsSolver.h"
 #include "Chaos/PBDRigidClustering.h"
 
 #include "HAL/IConsoleManager.h"
@@ -31,16 +32,15 @@ DEFINE_LOG_CATEGORY_STATIC(GCTCL_Log, Verbose, All);
 // has to go and fix up so many callsites here and they're all pretty much
 // Identical. The similar code should be pulled out
 
-#if TODO_REIMPLEMENT_RIGID_CLUSTERING
 namespace GeometryCollectionExample
 {
-#if INCLUDE_CHAOS
+#if TODO_REIMPLEMENT_RIGID_CLUSTERING
 	bool ClusterMapContains(const Chaos::FPBDRigidsSolver::FClusteringType::FClusterMap & ClusterMap, int Key, TArray<int32> Elements)
 	{
 		if (ClusterMap.Num())
-			if( ClusterMap.Contains(Key))
-				if( ClusterMap[Key] != nullptr)
-					if( ClusterMap[Key]->Num() == Elements.Num())
+			if (ClusterMap.Contains(Key))
+				if (ClusterMap[Key] != nullptr)
+					if (ClusterMap[Key]->Num() == Elements.Num())
 					{
 						for (int32 Element : Elements)
 							if (!ClusterMap[Key]->Contains(Element))
@@ -52,23 +52,18 @@ namespace GeometryCollectionExample
 #endif
 
 	template<class T>
-	bool RigidBodies_ClusterTest_SingleLevelNonBreaking(ExampleResponse&& R)
+	void RigidBodies_ClusterTest_SingleLevelNonBreaking()
 	{
 #if INCLUDE_CHAOS
 		TUniquePtr<Chaos::TChaosPhysicsMaterial<T>> PhysicalMaterial = MakeUnique<Chaos::TChaosPhysicsMaterial<T>>();
-		PhysicalMaterial->Friction = 0;
-		PhysicalMaterial->Restitution = 0;
-		PhysicalMaterial->SleepingLinearThreshold = 0;
-		PhysicalMaterial->SleepingAngularThreshold = 0;
-		PhysicalMaterial->DisabledLinearThreshold = 0;
-		PhysicalMaterial->DisabledAngularThreshold = 0;
+		InitMaterialToZero(PhysicalMaterial);
 
 		TSharedPtr<FGeometryCollection> RestCollection = GeometryCollection::MakeCubeElement(FTransform(FQuat::MakeFromEuler(FVector(0, 0, 0.)), FVector(0, -10, 10)),FVector(1.0));
 		RestCollection->AppendGeometry(*GeometryCollection::MakeCubeElement(FTransform(FQuat::MakeFromEuler(FVector(0, 0, 0.)), FVector(0, 10, 10)), FVector(1.0)));
-		R.ExpectTrue(RestCollection->Transform.Num() == 2);
+		EXPECT_EQ(RestCollection->Transform.Num(), 2);
 
 		FGeometryCollectionClusteringUtility::ClusterAllBonesUnderNewRoot(RestCollection.Get());
-		R.ExpectTrue(RestCollection->Transform.Num() == 3);
+		EXPECT_EQ(RestCollection->Transform.Num(), 3);
 		RestCollection->Transform[2] = FTransform(FQuat::MakeFromEuler(FVector(90.f, 0, 0.)), FVector(0, 0, 40));
 
 		//GeometryCollectionAlgo::PrintParentHierarchy(RestCollection.Get());
@@ -91,7 +86,9 @@ namespace GeometryCollectionExample
 		PhysObject->Initialize();
 
 		Chaos::FPBDRigidsSolver* Solver = FChaosSolversModule::GetModule()->CreateSolver(true);
+#if CHAOS_PARTICLEHANDLE_TODO
 		Solver->RegisterObject(PhysObject);
+#endif
 		Solver->SetHasFloor(true);
 		Solver->SetEnabled(true);
 		PhysObject->ActivateBodies();
@@ -100,49 +97,49 @@ namespace GeometryCollectionExample
 		float StartingRigidDistance = (Transform[1].GetTranslation() - Transform[0].GetTranslation()).Size(), CurrentRigidDistance = 0.f;
 
 		Solver->AdvanceSolverBy(1 / 24.);
+#if TODO_REIMPLEMENT_GET_RIGID_PARTICLES
 		const Chaos::TPBDRigidParticles<float, 3>& Particles = PhysObject->GetSolver()->GetRigidParticles();
+#endif
+#if TODO_REIMPLEMENT_RIGID_CLUSTERING
 		Chaos::FPBDRigidsSolver::FClusteringType & Clustering = Solver->GetRigidClustering();
-		const Chaos::FPBDRigidsSolver::FClusteringType::FClusterMap & ClusterMap = Solver->GetRigidClustering().GetChildrenMap();
-		R.ExpectTrue(ClusterMapContains(ClusterMap, 3, { 1,2 }));
+		//const Chaos::FPBDRigidsSolver::FClusteringType::FClusterMap & ClusterMap = Solver->GetRigidClustering().GetChildrenMap();
+		const auto & ClusterMap = Solver->GetRigidClustering().GetChildrenMap();
+		//EXPECT_TRUE(ClusterMapContains(ClusterMap, 3, { 1,2 }));
+#endif
 
 
 		for (int Frame = 1; Frame < 10; Frame++)
 		{
 			Solver->AdvanceSolverBy(1 / 24.);
 			CurrentRigidDistance = (Transform[1].GetTranslation() - Transform[0].GetTranslation()).Size();
-
-			R.ExpectTrue(Particles.Disabled(0) == false);
-			R.ExpectTrue(Particles.Disabled(1) == true);
-			R.ExpectTrue(Particles.Disabled(2) == true);
-			R.ExpectTrue(Particles.Disabled(3) == false);
-
-			R.ExpectTrue(FMath::Abs(CurrentRigidDistance - StartingRigidDistance) < 1e-4);
+#if TODO_REIMPLEMENT_GET_RIGID_PARTICLES
+			EXPECT_FALSE(Particles.Disabled(0));
+			EXPECT_TRUE(Particles.Disabled(1));
+			EXPECT_TRUE(Particles.Disabled(2));
+			EXPECT_FALSE(Particles.Disabled(3));
+#endif
+			EXPECT_LT(FMath::Abs(CurrentRigidDistance - StartingRigidDistance), 1e-4);
 		}
-		R.ExpectTrue(ClusterMapContains(ClusterMap, 3, { 1,2 }));
+
+#if TODO_REIMPLEMENT_RIGID_CLUSTERING
+		EXPECT_TRUE(ClusterMapContains(ClusterMap, 3, { 1,2 }));
+#endif
 
 		FChaosSolversModule::GetModule()->DestroySolver(Solver);
 
 		delete PhysObject;
 
 #endif
-
-
-		return !R.HasError();
 	}
-	template bool RigidBodies_ClusterTest_SingleLevelNonBreaking<float>(ExampleResponse&& R);
+	template void RigidBodies_ClusterTest_SingleLevelNonBreaking<float>();
 
 
 	template<class T>
-	bool RigidBodies_ClusterTest_DeactivateClusterParticle(ExampleResponse&& R)
+	void RigidBodies_ClusterTest_DeactivateClusterParticle()
 	{
 #if INCLUDE_CHAOS
 		TUniquePtr<Chaos::TChaosPhysicsMaterial<T>> PhysicalMaterial = MakeUnique<Chaos::TChaosPhysicsMaterial<T>>();
-		PhysicalMaterial->Friction = 0;
-		PhysicalMaterial->Restitution = 0;
-		PhysicalMaterial->SleepingLinearThreshold = 0;
-		PhysicalMaterial->SleepingAngularThreshold = 0;
-		PhysicalMaterial->DisabledLinearThreshold = 0;
-		PhysicalMaterial->DisabledAngularThreshold = 0;
+		InitMaterialToZero(PhysicalMaterial);
 
 		TSharedPtr<FGeometryCollection> RestCollection = GeometryCollection::MakeCubeElement(FTransform(FQuat::MakeFromEuler(FVector(0.f)), FVector(20.f)),FVector(1.0));
 		RestCollection->AppendGeometry(*GeometryCollection::MakeCubeElement(FTransform(FQuat::MakeFromEuler(FVector(0.f)), FVector(30.f)), FVector(1.0)));
@@ -183,7 +180,9 @@ namespace GeometryCollectionExample
 		PhysObject->Initialize();
 
 		Chaos::FPBDRigidsSolver* Solver = FChaosSolversModule::GetModule()->CreateSolver(true);
+#if CHAOS_PARTICLEHANDLE_TODO
 		Solver->RegisterObject(PhysObject);
+#endif
 		Solver->SetHasFloor(true);
 		Solver->SetEnabled(true);
 		PhysObject->ActivateBodies();
@@ -192,12 +191,15 @@ namespace GeometryCollectionExample
 		float StartingRigidDistance = (Transform[1].GetTranslation() - Transform[0].GetTranslation()).Size(), CurrentRigidDistance = 0.f;
 
 		Solver->AdvanceSolverBy(1 / 24.);
+#if TODO_REIMPLEMENT_GET_RIGID_PARTICLES
 		const Chaos::TPBDRigidParticles<float, 3>& Particles = PhysObject->GetSolver()->GetRigidParticles();
+
+#if TODO_REIMPLEMENT_RIGID_CLUSTERING
 		const Chaos::FPBDRigidsSolver::FClusteringType::FClusterMap & ClusterMap = Solver->GetRigidClustering().GetChildrenMap();
-		R.ExpectTrue(ClusterMapContains(ClusterMap, 9, { 1,8 }));
-		R.ExpectTrue(ClusterMapContains(ClusterMap, 8, { 7,2 }));
-		R.ExpectTrue(ClusterMapContains(ClusterMap, 7, { 6,3 }));
-		R.ExpectTrue(ClusterMapContains(ClusterMap, 6, { 5,4 }));
+		EXPECT_TRUE(ClusterMapContains(ClusterMap, 9, { 1,8 }));
+		EXPECT_TRUE(ClusterMapContains(ClusterMap, 8, { 7,2 }));
+		EXPECT_TRUE(ClusterMapContains(ClusterMap, 7, { 6,3 }));
+		EXPECT_TRUE(ClusterMapContains(ClusterMap, 6, { 5,4 }));
 
 
 		TArray<bool> Conditions = { false,false };
@@ -232,9 +234,9 @@ namespace GeometryCollectionExample
 					Particles.Disabled(9) == false)
 				{
 					Conditions[0] = true;
-					R.ExpectTrue(Particles.InvM(9) == 0.f); // kinematic cluster
-					R.ExpectTrue(Particles.InvM(8) == 0.f); // disabled child
-					R.ExpectTrue(Particles.InvM(1) == 0.f); // disabled child
+					EXPECT_EQ(Particles.InvM(9), 0.f); // kinematic cluster
+					EXPECT_EQ(Particles.InvM(8), 0.f); // disabled child
+					EXPECT_EQ(Particles.InvM(1), 0.f); // disabled child
 				}
 			}
 			else if (Conditions[0] == true && Conditions[1] == false && Frame == 2)
@@ -251,21 +253,23 @@ namespace GeometryCollectionExample
 					Particles.Disabled(9) == true)
 				{
 					Conditions[1] = true;
-					R.ExpectTrue(Particles.InvM(9) == 0.f); // disabled cluster body
-					R.ExpectTrue(Particles.InvM(1) == 0.f); // enabled child
-					R.ExpectTrue(Particles.InvM(8) == 0.f); // enabled child
-
-					R.ExpectTrue(!ClusterMap.Contains(9));
-					R.ExpectTrue(ClusterMapContains(ClusterMap, 8, { 7,2 }));
-					R.ExpectTrue(ClusterMapContains(ClusterMap, 7, { 6,3 }));
-					R.ExpectTrue(ClusterMapContains(ClusterMap, 6, { 5,4 }));
+					EXPECT_EQ(Particles.InvM(9), 0.f); // disabled cluster body
+					EXPECT_EQ(Particles.InvM(1), 0.f); // enabled child
+					EXPECT_EQ(Particles.InvM(8), 0.f); // enabled child
+
+					EXPECT_TRUE(!ClusterMap.Contains(9));
+					EXPECT_TRUE(ClusterMapContains(ClusterMap, 8, { 7,2 }));
+					EXPECT_TRUE(ClusterMapContains(ClusterMap, 7, { 6,3 }));
+					EXPECT_TRUE(ClusterMapContains(ClusterMap, 6, { 5,4 }));
 				}
 			}
 		}
 		for (int i = 0; i < Conditions.Num(); i++)
 		{
-			R.ExpectTrue(Conditions[i]);
+			EXPECT_TRUE(Conditions[i]);
 		}
+#endif
+#endif
 
 		FChaosSolversModule::GetModule()->DestroySolver(Solver);
 
@@ -274,14 +278,14 @@ namespace GeometryCollectionExample
 #endif
 
 
-		return !R.HasError();
+		
 	}
-	template bool RigidBodies_ClusterTest_DeactivateClusterParticle<float>(ExampleResponse&& R);
+	template void RigidBodies_ClusterTest_DeactivateClusterParticle<float>();
 
 
 
 	template<class T>
-	bool RigidBodies_ClusterTest_SingleLevelBreaking(ExampleResponse&& R)
+	void RigidBodies_ClusterTest_SingleLevelBreaking()
 	{
 		//
 		// Test overview:
@@ -292,20 +296,15 @@ namespace GeometryCollectionExample
 		//
 #if INCLUDE_CHAOS
 		TUniquePtr<Chaos::TChaosPhysicsMaterial<T>> PhysicalMaterial = MakeUnique<Chaos::TChaosPhysicsMaterial<T>>();
-		PhysicalMaterial->Friction = 0;
-		PhysicalMaterial->Restitution = 0;
-		PhysicalMaterial->SleepingLinearThreshold = 0;
-		PhysicalMaterial->SleepingAngularThreshold = 0;
-		PhysicalMaterial->DisabledLinearThreshold = 0;
-		PhysicalMaterial->DisabledAngularThreshold = 0;
+		InitMaterialToZero(PhysicalMaterial);
 
 		TSharedPtr<FGeometryCollection> RestCollection = GeometryCollection::MakeCubeElement(FTransform(FQuat::MakeFromEuler(FVector(0, 0, 0.)), FVector(0, -10, 10)),FVector(1.0));
 		RestCollection->AppendGeometry(*GeometryCollection::MakeCubeElement(FTransform(FQuat::MakeFromEuler(FVector(0, 0, 0.)), FVector(0, 10, 10)), FVector(1.0)));
-		R.ExpectTrue(RestCollection->Transform.Num() == 2);
+		EXPECT_EQ(RestCollection->Transform.Num(), 2);
 
 
 		FGeometryCollectionClusteringUtility::ClusterAllBonesUnderNewRoot(RestCollection.Get());
-		R.ExpectTrue(RestCollection->Transform.Num() == 3);
+		EXPECT_EQ(RestCollection->Transform.Num(), 3);
 		RestCollection->Transform[2] = FTransform(FQuat::MakeFromEuler(FVector(90.f, 0, 0.)), FVector(0, 0, 40));
 
 		//GeometryCollectionAlgo::PrintParentHierarchy(RestCollection.Get());
@@ -330,7 +329,9 @@ namespace GeometryCollectionExample
 		PhysObject->SetCollisionParticlesPerObjectFraction(1.0);
 
 		Chaos::FPBDRigidsSolver* Solver = FChaosSolversModule::GetModule()->CreateSolver(true);
+#if CHAOS_PARTICLEHANDLE_TODO
 		Solver->RegisterObject(PhysObject);
+#endif
 		Solver->SetHasFloor(true);
 		Solver->SetEnabled(true);
 		PhysObject->ActivateBodies();
@@ -339,10 +340,13 @@ namespace GeometryCollectionExample
 		float StartingRigidDistance = (Transform[1].GetTranslation() - Transform[0].GetTranslation()).Size(), CurrentRigidDistance = 0.f;
 
 		Solver->AdvanceSolverBy(1 / 24.);
+#if TODO_REIMPLEMENT_GET_RIGID_PARTICLES
 		const Chaos::TPBDRigidParticles<float, 3>& Particles = PhysObject->GetSolver()->GetRigidParticles();
+#endif
+#if TODO_REIMPLEMENT_RIGID_CLUSTERING
 		const Chaos::FPBDRigidsSolver::FClusteringType::FClusterMap & ClusterMap = Solver->GetRigidClustering().GetChildrenMap();
-		R.ExpectTrue(ClusterMapContains(ClusterMap, 3, { 1,2 }));
-
+		EXPECT_TRUE(ClusterMapContains(ClusterMap, 3, { 1,2 }));
+#endif
 
 		// Particles array contains the following:
 		// 0: Ground
@@ -358,84 +362,90 @@ namespace GeometryCollectionExample
 			if (Frame < 5)
 			{
 				// The two boxes are dropping to the ground as a cluster
-				R.ExpectTrue(Particles.Disabled(0) == false);
-				R.ExpectTrue(Particles.Disabled(1) == true);
-				R.ExpectTrue(Particles.Disabled(2) == true);
-				R.ExpectTrue(Particles.Disabled(3) == false);
-				R.ExpectTrue(FMath::Abs(CurrentRigidDistance - StartingRigidDistance) < 1e-4);
+#if TODO_REIMPLEMENT_GET_RIGID_PARTICLES
+				EXPECT_FALSE(Particles.Disabled(0));
+				EXPECT_TRUE(Particles.Disabled(1));
+				EXPECT_TRUE(Particles.Disabled(2));
+				EXPECT_FALSE(Particles.Disabled(3));
+#endif
+				EXPECT_LT(FMath::Abs(CurrentRigidDistance - StartingRigidDistance), 1e-4);
 			}
 			else if (Frame == 5)
 			{
 				// The cluster has just hit the ground and should have broken.
 				// The boxes are still separated by StartingRigidDistance (when Rewind is disabled).
 				// All children should have zero velocity.
-				R.ExpectTrue(Particles.Disabled(0) == false);
-				R.ExpectTrue(Particles.Disabled(1) == false);
-				R.ExpectTrue(Particles.Disabled(2) == false);
-				R.ExpectTrue(Particles.Disabled(3) == true);
-				R.ExpectTrue(Particles.V(1).Size() < 1.e-4);
-				R.ExpectTrue(Particles.V(2).Size() < 1.e-4);
-				R.ExpectTrue(ClusterMap.Num() == 0);
-				R.ExpectTrue(FMath::Abs(CurrentRigidDistance - StartingRigidDistance) < 1e-4);
+#if TODO_REIMPLEMENT_GET_RIGID_PARTICLES
+				EXPECT_FALSE(Particles.Disabled(0));
+				EXPECT_FALSE(Particles.Disabled(1));
+				EXPECT_FALSE(Particles.Disabled(2));
+				EXPECT_TRUE(Particles.Disabled(3));
+				EXPECT_LT(Particles.V(1).Size(), 1.e-4);
+				EXPECT_LT(Particles.V(2).Size(), 1.e-4);
+#endif
+#if TODO_REIMPLEMENT_RIGID_CLUSTERING
+				EXPECT_EQ(ClusterMap.Num(), 0);
+#endif
+				EXPECT_LT(FMath::Abs(CurrentRigidDistance - StartingRigidDistance), 1e-4);
 			}
 			else if (Frame == 6)
 			{
 				// The boxes are now moving independently, but they had zero velocity 
 				// last frame, so they should still be separated by StartingRigidDistance. 
-				R.ExpectTrue(Particles.Disabled(0) == false);
-				R.ExpectTrue(Particles.Disabled(1) == false);
-				R.ExpectTrue(Particles.Disabled(2) == false);
-				R.ExpectTrue(Particles.Disabled(3) == true);
-				R.ExpectTrue(ClusterMap.Num() == 0);
-				R.ExpectTrue(FMath::Abs(CurrentRigidDistance - StartingRigidDistance) < 1e-4);
+#if TODO_REIMPLEMENT_GET_RIGID_PARTICLES
+				EXPECT_FALSE(Particles.Disabled(0));
+				EXPECT_FALSE(Particles.Disabled(1));
+				EXPECT_FALSE(Particles.Disabled(2));
+				EXPECT_TRUE(Particles.Disabled(3));
+#endif
+#if TODO_REIMPLEMENT_RIGID_CLUSTERING
+				EXPECT_EQ(ClusterMap.Num(), 0);
+#endif
+				EXPECT_LT(FMath::Abs(CurrentRigidDistance - StartingRigidDistance), 1e-4);
 			}
 			else
 			{
 				// The boxes are now moving independently - the bottom one is on the ground and should be stopped.
 				// The top one is still falling, so they should be closer together
-				R.ExpectTrue(Particles.Disabled(0) == false);
-				R.ExpectTrue(Particles.Disabled(1) == false);
-				R.ExpectTrue(Particles.Disabled(2) == false);
-				R.ExpectTrue(Particles.Disabled(3) == true);
-				R.ExpectTrue(ClusterMap.Num()==0);
+#if TODO_REIMPLEMENT_GET_RIGID_PARTICLES
+				EXPECT_FALSE(Particles.Disabled(0));
+				EXPECT_FALSE(Particles.Disabled(1));
+				EXPECT_FALSE(Particles.Disabled(2));
+				EXPECT_TRUE(Particles.Disabled(3));
+#endif
+#if TODO_REIMPLEMENT_RIGID_CLUSTERING
+				EXPECT_EQ(ClusterMap.Num(), 0);
+#endif
 			}
 		}
 		
-		R.ExpectTrue(FMath::Abs(CurrentRigidDistance - StartingRigidDistance) > 1e-4);
+		EXPECT_GT(FMath::Abs(CurrentRigidDistance - StartingRigidDistance), 1e-4);
 		FChaosSolversModule::GetModule()->DestroySolver(Solver);
 
 		delete PhysObject;
 
 #endif
-
-
-		return !R.HasError();
 	}
-	template bool RigidBodies_ClusterTest_SingleLevelBreaking<float>(ExampleResponse&& R);
+	template void RigidBodies_ClusterTest_SingleLevelBreaking<float>();
 
 
 	template<class T>
-	bool RigidBodies_ClusterTest_NestedCluster(ExampleResponse&& R)
+	void RigidBodies_ClusterTest_NestedCluster()
 	{
 #if INCLUDE_CHAOS
 		TUniquePtr<Chaos::TChaosPhysicsMaterial<T>> PhysicalMaterial = MakeUnique<Chaos::TChaosPhysicsMaterial<T>>();
-		PhysicalMaterial->Friction = 0;
-		PhysicalMaterial->Restitution = 0;
-		PhysicalMaterial->SleepingLinearThreshold = 0;
-		PhysicalMaterial->SleepingAngularThreshold = 0;
-		PhysicalMaterial->DisabledLinearThreshold = 0;
-		PhysicalMaterial->DisabledAngularThreshold = 0;
+		InitMaterialToZero(PhysicalMaterial);
 
 		TSharedPtr<FGeometryCollection> RestCollection = GeometryCollection::MakeCubeElement(FTransform(FQuat::MakeFromEuler(FVector(0, 0, 0.)), FVector(0, -10, 10)),FVector(1.0));
 		RestCollection->AppendGeometry(*GeometryCollection::MakeCubeElement(FTransform(FQuat::MakeFromEuler(FVector(0, 0, 0.)), FVector(0, 10, 10)), FVector(1.0)));
-		R.ExpectTrue(RestCollection->Transform.Num() == 2);
+		EXPECT_EQ(RestCollection->Transform.Num(), 2);
 
 		FGeometryCollectionClusteringUtility::ClusterAllBonesUnderNewRoot(RestCollection.Get());
-		R.ExpectTrue(RestCollection->Transform.Num() == 3);
+		EXPECT_EQ(RestCollection->Transform.Num(), 3);
 		RestCollection->Transform[2] = FTransform(FQuat::MakeFromEuler(FVector(90.f, 0, 0.)), FVector(0, 0, 40));
 
 		FGeometryCollectionClusteringUtility::ClusterBonesUnderNewNode(RestCollection.Get(), 3, { 2 }, true);
-		R.ExpectTrue(RestCollection->Transform.Num() == 4);
+		EXPECT_EQ(RestCollection->Transform.Num(), 4);
 		RestCollection->Transform[3] = FTransform(FQuat::MakeFromEuler(FVector(0.f, 0, 0.)), FVector(0, 0, 10));
 
 		//GeometryCollectionAlgo::PrintParentHierarchy(RestCollection.Get());
@@ -459,7 +469,9 @@ namespace GeometryCollectionExample
 		PhysObject->Initialize();
 
 		Chaos::FPBDRigidsSolver* Solver = FChaosSolversModule::GetModule()->CreateSolver(true);
+#if CHAOS_PARTICLEHANDLE_TODO
 		Solver->RegisterObject(PhysObject);
+#endif
 		Solver->SetHasFloor(true);
 		Solver->SetEnabled(true);
 		PhysObject->ActivateBodies();
@@ -468,10 +480,12 @@ namespace GeometryCollectionExample
 		float StartingRigidDistance = (Transform[1].GetTranslation() - Transform[0].GetTranslation()).Size(), CurrentRigidDistance = 0.f;
 
 		Solver->AdvanceSolverBy(1 / 24.);
+#if TODO_REIMPLEMENT_GET_RIGID_PARTICLES
 		const Chaos::TPBDRigidParticles<float, 3>& Particles = PhysObject->GetSolver()->GetRigidParticles();
+#if TODO_REIMPLEMENT_RIGID_CLUSTERING
 		const Chaos::FPBDRigidsSolver::FClusteringType::FClusterMap & ClusterMap = Solver->GetRigidClustering().GetChildrenMap();
-		R.ExpectTrue(ClusterMapContains(ClusterMap, 3, { 1,2 }));
-		R.ExpectTrue(ClusterMapContains(ClusterMap, 4, { 3, }));
+		EXPECT_TRUE(ClusterMapContains(ClusterMap, 3, { 1,2 }));
+		EXPECT_TRUE(ClusterMapContains(ClusterMap, 4, { 3, }));
 
 		TArray<bool> Conditions = {false,false,false};
 
@@ -501,9 +515,9 @@ namespace GeometryCollectionExample
 					Particles.Disabled(4) == true)
 				{
 					Conditions[1] = true;
-					R.ExpectTrue(ClusterMapContains(ClusterMap, 3, { 1,2 }));
-					R.ExpectTrue(ClusterMap.Num()==1);
-					R.ExpectTrue(!ClusterMap.Contains(4));
+					EXPECT_TRUE(ClusterMapContains(ClusterMap, 3, { 1,2 }));
+					EXPECT_EQ(ClusterMap.Num(), 1);
+					EXPECT_TRUE(!ClusterMap.Contains(4));
 				}
 			}
 			else if (Conditions[1] == true && Conditions[2] == false)
@@ -515,14 +529,16 @@ namespace GeometryCollectionExample
 					Particles.Disabled(4) == true)
 				{
 					Conditions[2] = true;
-					R.ExpectTrue(ClusterMap.Num() == 0);
+					EXPECT_EQ(ClusterMap.Num(), 0);
 				}
 			}
 		}
 		for (int i = 0; i < Conditions.Num(); i++)
 		{
-			R.ExpectTrue(Conditions[i]);
+			EXPECT_TRUE(Conditions[i]);
 		}
+#endif
+#endif
 
 		FChaosSolversModule::GetModule()->DestroySolver(Solver);
 
@@ -530,22 +546,17 @@ namespace GeometryCollectionExample
 
 #endif
 
-		return !R.HasError();
+		
 	}
-	template bool RigidBodies_ClusterTest_NestedCluster<float>(ExampleResponse&& R);
+	template void RigidBodies_ClusterTest_NestedCluster<float>();
 
 
 	template<class T>
-	bool RigidBodies_ClusterTest_NestedCluster_MultiStrain(ExampleResponse&& R)
+	void RigidBodies_ClusterTest_NestedCluster_MultiStrain()
 	{
 #if INCLUDE_CHAOS
 		TUniquePtr<Chaos::TChaosPhysicsMaterial<T>> PhysicalMaterial = MakeUnique<Chaos::TChaosPhysicsMaterial<T>>();
-		PhysicalMaterial->Friction = 0;
-		PhysicalMaterial->Restitution = 0;
-		PhysicalMaterial->SleepingLinearThreshold = 0;
-		PhysicalMaterial->SleepingAngularThreshold = 0;
-		PhysicalMaterial->DisabledLinearThreshold = 0;
-		PhysicalMaterial->DisabledAngularThreshold = 0;
+		InitMaterialToZero(PhysicalMaterial);
 
 		TSharedPtr<FGeometryCollection> RestCollection = GeometryCollection::MakeCubeElement(FTransform(FQuat::MakeFromEuler(FVector(0.f)), FVector(20.f)),FVector(1.0));
 		RestCollection->AppendGeometry(*GeometryCollection::MakeCubeElement(FTransform(FQuat::MakeFromEuler(FVector(0.f)), FVector(30.f)), FVector(1.0)));
@@ -591,7 +602,9 @@ namespace GeometryCollectionExample
 		PhysObject->Initialize();
 
 		Chaos::FPBDRigidsSolver* Solver = FChaosSolversModule::GetModule()->CreateSolver(true);
+#if CHAOS_PARTICLEHANDLE_TODO
 		Solver->RegisterObject(PhysObject);
+#endif
 		Solver->SetHasFloor(true);
 		Solver->SetEnabled(true);
 		PhysObject->ActivateBodies();
@@ -602,13 +615,15 @@ namespace GeometryCollectionExample
 		TArray<bool> Conditions = { false,false,false };
 
 		Solver->AdvanceSolverBy(1 / 24.);
+#if TODO_REIMPLEMENT_GET_RIGID_PARTICLES
 		const Chaos::TPBDRigidParticles<float, 3>& Particles = PhysObject->GetSolver()->GetRigidParticles();
+#if TODO_REIMPLEMENT_RIGID_CLUSTERING
 		const Chaos::FPBDRigidsSolver::FClusteringType::FClusterMap & ClusterMap = Solver->GetRigidClustering().GetChildrenMap();
-		R.ExpectTrue(ClusterMap.Num()==4);
-		R.ExpectTrue(ClusterMapContains(ClusterMap, 9, { 1,8 }));
-		R.ExpectTrue(ClusterMapContains(ClusterMap, 8, { 7,2 }));
-		R.ExpectTrue(ClusterMapContains(ClusterMap, 7, { 6,3 }));
-		R.ExpectTrue(ClusterMapContains(ClusterMap, 6, { 5,4 }));
+		EXPECT_EQ(ClusterMap.Num(), 4);
+		EXPECT_TRUE(ClusterMapContains(ClusterMap, 9, { 1,8 }));
+		EXPECT_TRUE(ClusterMapContains(ClusterMap, 8, { 7,2 }));
+		EXPECT_TRUE(ClusterMapContains(ClusterMap, 7, { 6,3 }));
+		EXPECT_TRUE(ClusterMapContains(ClusterMap, 6, { 5,4 }));
 
 		for (int Frame = 1; Frame < 20; Frame++)
 		{
@@ -647,10 +662,10 @@ namespace GeometryCollectionExample
 				{
 					Conditions[1] = true;
 
-					R.ExpectTrue(ClusterMap.Num() == 3);
-					R.ExpectTrue(ClusterMapContains(ClusterMap, 8, { 7,2 }));
-					R.ExpectTrue(ClusterMapContains(ClusterMap, 7, { 6,3 }));
-					R.ExpectTrue(ClusterMapContains(ClusterMap, 6, { 5,4 }));
+					EXPECT_EQ(ClusterMap.Num(), 3);
+					EXPECT_TRUE(ClusterMapContains(ClusterMap, 8, { 7,2 }));
+					EXPECT_TRUE(ClusterMapContains(ClusterMap, 7, { 6,3 }));
+					EXPECT_TRUE(ClusterMapContains(ClusterMap, 6, { 5,4 }));
 				}
 			}
 			else if (Conditions[1] == true && Conditions[2] == false)
@@ -668,16 +683,18 @@ namespace GeometryCollectionExample
 				{
 					Conditions[2] = true;
 
-					R.ExpectTrue(ClusterMap.Num() == 2);
-					R.ExpectTrue(ClusterMapContains(ClusterMap, 7, { 6,3 }));
-					R.ExpectTrue(ClusterMapContains(ClusterMap, 6, { 5,4 }));
+					EXPECT_EQ(ClusterMap.Num(), 2);
+					EXPECT_TRUE(ClusterMapContains(ClusterMap, 7, { 6,3 }));
+					EXPECT_TRUE(ClusterMapContains(ClusterMap, 6, { 5,4 }));
 				}
 			}
 		}
 		for (int i = 0; i < Conditions.Num(); i++)
 		{
-			R.ExpectTrue(Conditions[i]);
+			EXPECT_TRUE(Conditions[i]);
 		}
+#endif
+#endif
 
 		FChaosSolversModule::GetModule()->DestroySolver(Solver);
 
@@ -686,23 +703,18 @@ namespace GeometryCollectionExample
 #endif
 
 
-		return !R.HasError();
+		
 	}
-	template bool RigidBodies_ClusterTest_NestedCluster_MultiStrain<float>(ExampleResponse&& R);
+	template void RigidBodies_ClusterTest_NestedCluster_MultiStrain<float>();
 
 
 
 	template<class T>
-	bool RigidBodies_ClusterTest_NestedCluster_Halt(ExampleResponse&& R)
+	void RigidBodies_ClusterTest_NestedCluster_Halt()
 	{
 #if INCLUDE_CHAOS
 		TUniquePtr<Chaos::TChaosPhysicsMaterial<T>> PhysicalMaterial = MakeUnique<Chaos::TChaosPhysicsMaterial<T>>();
-		PhysicalMaterial->Friction = 0;
-		PhysicalMaterial->Restitution = 0;
-		PhysicalMaterial->SleepingLinearThreshold = 0;
-		PhysicalMaterial->SleepingAngularThreshold = 0;
-		PhysicalMaterial->DisabledLinearThreshold = 0;
-		PhysicalMaterial->DisabledAngularThreshold = 0;
+		InitMaterialToZero(PhysicalMaterial);
 
 		TSharedPtr<FGeometryCollection> RestCollection = GeometryCollection::MakeCubeElement(FTransform(FQuat::MakeFromEuler(FVector(0.f)), FVector(20.f)),FVector(1.0));
 		RestCollection->AppendGeometry(*GeometryCollection::MakeCubeElement(FTransform(FQuat::MakeFromEuler(FVector(0.f)), FVector(30.f)), FVector(1.0)));
@@ -749,7 +761,9 @@ namespace GeometryCollectionExample
 		PhysObject->Initialize();
 
 		Chaos::FPBDRigidsSolver* Solver = FChaosSolversModule::GetModule()->CreateSolver(true);
+#if CHAOS_PARTICLEHANDLE_TODO
 		Solver->RegisterObject(PhysObject);
+#endif
 		Solver->SetHasFloor(true);
 		Solver->SetEnabled(true);
 		PhysObject->ActivateBodies();
@@ -757,13 +771,16 @@ namespace GeometryCollectionExample
 		TManagedArray<FTransform>& Transform = DynamicCollection->Transform;
 		float StartingRigidDistance = (Transform[1].GetTranslation() - Transform[0].GetTranslation()).Size(), CurrentRigidDistance = 0.f;
 
+#if TODO_REIMPLEMENT_GET_RIGID_PARTICLES
 		TArray<bool> Conditions = { false,false };
 
 		for (int Frame = 0; Frame < 10; Frame++)
 		{
 			Solver->AdvanceSolverBy(1 / 24.);
 
+
 			const Chaos::TPBDRigidParticles<float, 3>& Particles = PhysObject->GetSolver()->GetRigidParticles();
+
 			CurrentRigidDistance = (Transform[1].GetTranslation() - Transform[0].GetTranslation()).Size();
 
 			//UE_LOG(GCTCL_Log, Verbose, TEXT("FRAME : %d"), Frame);
@@ -778,7 +795,7 @@ namespace GeometryCollectionExample
 			//UE_LOG(GCTCL_Log, Verbose, TEXT("StartingRigidDistance : %3.5f"), StartingRigidDistance);
 			//UE_LOG(GCTCL_Log, Verbose, TEXT("DeltaRigidDistance : %3.5f"), CurrentRigidDistance - StartingRigidDistance);
 			
-			
+#if TODO_REIMPLEMENT_GET_RIGID_PARTICLES
 			if (Conditions[0] == false)
 			{
 				if (Particles.Disabled(0) == false &&
@@ -811,11 +828,13 @@ namespace GeometryCollectionExample
 					Conditions[1] = true;
 				}
 			}
+#endif
 		}
 		for (int i = 0; i < Conditions.Num(); i++)
 		{
-			R.ExpectTrue(Conditions[i]);
+			EXPECT_TRUE(Conditions[i]);
 		}
+#endif
 
 		FChaosSolversModule::GetModule()->DestroySolver(Solver);
 
@@ -824,24 +843,19 @@ namespace GeometryCollectionExample
 #endif
 
 
-		return !R.HasError();
+		
 	}
-	template bool RigidBodies_ClusterTest_NestedCluster_Halt<float>(ExampleResponse&& R);
+	template void RigidBodies_ClusterTest_NestedCluster_Halt<float>();
 
 	template<class T>
-	bool RigidBodies_ClusterTest_KinematicAnchor(ExampleResponse&& R)
+	void RigidBodies_ClusterTest_KinematicAnchor()
 	{
 		// Test : Set one element kinematic. When the cluster breaks the elements that do not contain the kinematic
 		//        rigid body should be dynamic, while the clusters that contain the kinematic body should remain 
 		//        kinematic. 
 #if INCLUDE_CHAOS
 		TUniquePtr<Chaos::TChaosPhysicsMaterial<T>> PhysicalMaterial = MakeUnique<Chaos::TChaosPhysicsMaterial<T>>();
-		PhysicalMaterial->Friction = 0;
-		PhysicalMaterial->Restitution = 0;
-		PhysicalMaterial->SleepingLinearThreshold = 0;
-		PhysicalMaterial->SleepingAngularThreshold = 0;
-		PhysicalMaterial->DisabledLinearThreshold = 0;
-		PhysicalMaterial->DisabledAngularThreshold = 0;
+		InitMaterialToZero(PhysicalMaterial);
 
 		TSharedPtr<FGeometryCollection> RestCollection = GeometryCollection::MakeCubeElement(FTransform(FQuat::MakeFromEuler(FVector(0.f)), FVector(20.f)), FVector(1.0));
 		RestCollection->AppendGeometry(*GeometryCollection::MakeCubeElement(FTransform(FQuat::MakeFromEuler(FVector(0.f)), FVector(30.f)), FVector(1.0)));
@@ -882,9 +896,13 @@ namespace GeometryCollectionExample
 		PhysObject->Initialize();
 
 		Chaos::FPBDRigidsSolver* Solver = FChaosSolversModule::GetModule()->CreateSolver(true);
+#if TODO_REIMPLEMENT_RIGID_CLUSTERING
 		Chaos::FPBDRigidsSolver::FClusteringType & Clustering = Solver->GetRigidClustering();
+#endif
 
+#if CHAOS_PARTICLEHANDLE_TODO
 		Solver->RegisterObject(PhysObject);
+#endif
 		Solver->SetHasFloor(true);
 		Solver->SetEnabled(true);
 		PhysObject->ActivateBodies();
@@ -895,14 +913,17 @@ namespace GeometryCollectionExample
 		TArray<bool> Conditions = { false,false,false,false };
 
 		Solver->AdvanceSolverBy(1 / 24.);
+#if TODO_REIMPLEMENT_GET_RIGID_PARTICLES
 		const Chaos::TPBDRigidParticles<float, 3>& Particles = PhysObject->GetSolver()->GetRigidParticles();
+
+#if TODO_REIMPLEMENT_RIGID_CLUSTERING
 		const Chaos::FPBDRigidsSolver::FClusteringType::FClusterMap & ClusterMap = Clustering.GetChildrenMap();
 
-		R.ExpectTrue(ClusterMap.Num() == 4);
-		R.ExpectTrue(ClusterMapContains(ClusterMap, 9, { 1,8 }));
-		R.ExpectTrue(ClusterMapContains(ClusterMap, 8, { 7,2 }));
-		R.ExpectTrue(ClusterMapContains(ClusterMap, 7, { 6,3 }));
-		R.ExpectTrue(ClusterMapContains(ClusterMap, 6, { 5,4 }));
+		EXPECT_EQ(ClusterMap.Num(), 4);
+		EXPECT_TRUE(ClusterMapContains(ClusterMap, 9, { 1,8 }));
+		EXPECT_TRUE(ClusterMapContains(ClusterMap, 8, { 7,2 }));
+		EXPECT_TRUE(ClusterMapContains(ClusterMap, 7, { 6,3 }));
+		EXPECT_TRUE(ClusterMapContains(ClusterMap, 6, { 5,4 }));
 
 		for (int Frame = 1; Frame < 10; Frame++)
 		{
@@ -924,14 +945,14 @@ namespace GeometryCollectionExample
 			//	UE_LOG(GCTCL_Log, Verbose, TEXT("... ... ...    InvM[%d] : %f"), rdx, Particles.InvM(rdx));
 			//}
 
-			R.ExpectTrue(Particles.InvM(0) == 0.f); // floor
-			R.ExpectTrue(Particles.InvM(1) != 0.f); // dynamic rigid
-			R.ExpectTrue(Particles.InvM(2) == 0.f); // kinematic rigid
-			R.ExpectTrue(Particles.InvM(3) != 0.f); // dynamic rigid
-			R.ExpectTrue(Particles.InvM(4) != 0.f); // dynamic rigid
-			R.ExpectTrue(Particles.InvM(5) != 0.f); // dynamic rigid
-			R.ExpectTrue(Particles.InvM(6) != 0.f); // dynamic rigid
-			R.ExpectTrue(Particles.InvM(7) != 0.f); // dynamic cluster
+			EXPECT_EQ(Particles.InvM(0), 0.f); // floor
+			EXPECT_NE(Particles.InvM(1), 0.f); // dynamic rigid
+			EXPECT_EQ(Particles.InvM(2), 0.f); // kinematic rigid
+			EXPECT_NE(Particles.InvM(3), 0.f); // dynamic rigid
+			EXPECT_NE(Particles.InvM(4), 0.f); // dynamic rigid
+			EXPECT_NE(Particles.InvM(5), 0.f); // dynamic rigid
+			EXPECT_NE(Particles.InvM(6), 0.f); // dynamic rigid
+			EXPECT_NE(Particles.InvM(7), 0.f); // dynamic cluster
 
 			FVector Ref1, Ref2, Ref7; // RigidBody0(Dynamic), RigidBody1(Kinematic), RigidBody6(Kinematic then Dynamic)
 			if (Conditions[0] == false && Frame == 1)
@@ -951,10 +972,10 @@ namespace GeometryCollectionExample
 					Ref1 = Particles.X(1);
 					Ref2 = Particles.X(2);
 					Ref7 = Particles.X(7);
-					R.ExpectTrue(Particles.InvM(8) == 0.f); // kinematic cluster
-					R.ExpectTrue(Particles.ObjectState(8) == Chaos::EObjectStateType::Kinematic); // kinematic cluster
-					R.ExpectTrue(Particles.InvM(9) == 0.f); // kinematic cluster
-					R.ExpectTrue(Particles.ObjectState(9)==Chaos::EObjectStateType::Kinematic); // kinematic cluster
+					EXPECT_EQ(Particles.InvM(8), 0.f); // kinematic cluster
+					EXPECT_EQ(Particles.ObjectState(8), Chaos::EObjectStateType::Kinematic); // kinematic cluster
+					EXPECT_EQ(Particles.InvM(9), 0.f); // kinematic cluster
+					EXPECT_EQ(Particles.ObjectState(9), Chaos::EObjectStateType::Kinematic); // kinematic cluster
 				}
 			}
 			else if (Conditions[0] == true && Conditions[1] == false && Frame == 2)
@@ -975,18 +996,18 @@ namespace GeometryCollectionExample
 					FVector X2 = Particles.X(2);
 					FVector X7 = Particles.X(7);
 
-					R.ExpectTrue(FMath::Abs(X1.Size() - Ref1.Size()) <= KINDA_SMALL_NUMBER, "Kinematic body1 moved");
-					R.ExpectTrue(FMath::Abs(X2.Size() - Ref2.Size()) <= KINDA_SMALL_NUMBER, "Kinematic body2 moved");
-					R.ExpectTrue(FMath::Abs(X7.Size() - Ref7.Size()) <= KINDA_SMALL_NUMBER, "Kinematic body7 moved");
-					R.ExpectTrue(Particles.InvM(8) == 0.f); // kinematic cluster
-					R.ExpectTrue(Particles.ObjectState(8) == Chaos::EObjectStateType::Kinematic); // kinematic cluster
-					R.ExpectTrue(Particles.InvM(9) == 0.f); 
-					R.ExpectTrue(Particles.ObjectState(9) == Chaos::EObjectStateType::Kinematic);
-
-					R.ExpectTrue(ClusterMap.Num() == 3);
-					R.ExpectTrue(ClusterMapContains(ClusterMap, 8, { 7,2 }));
-					R.ExpectTrue(ClusterMapContains(ClusterMap, 7, { 6,3 }));
-					R.ExpectTrue(ClusterMapContains(ClusterMap, 6, { 5,4 }));
+					EXPECT_NEAR(FMath::Abs(X1.Size() - Ref1.Size()), 0, KINDA_SMALL_NUMBER) << *FString("Kinematic body1 moved");
+					EXPECT_NEAR(FMath::Abs(X2.Size() - Ref2.Size()), 0, KINDA_SMALL_NUMBER) << *FString("Kinematic body2 moved");
+					EXPECT_NEAR(FMath::Abs(X7.Size() - Ref7.Size()), 0, KINDA_SMALL_NUMBER) << *FString("Kinematic body7 moved");
+					EXPECT_EQ(Particles.InvM(8), 0.f); // kinematic cluster
+					EXPECT_EQ(Particles.ObjectState(8), Chaos::EObjectStateType::Kinematic); // kinematic cluster
+					EXPECT_EQ(Particles.InvM(9), 0.f); 
+					EXPECT_EQ(Particles.ObjectState(9), Chaos::EObjectStateType::Kinematic);
+
+					EXPECT_EQ(ClusterMap.Num(), 3);
+					EXPECT_TRUE(ClusterMapContains(ClusterMap, 8, { 7,2 }));
+					EXPECT_TRUE(ClusterMapContains(ClusterMap, 7, { 6,3 }));
+					EXPECT_TRUE(ClusterMapContains(ClusterMap, 6, { 5,4 }));
 				}
 			}
 			else if (Conditions[1] == true && Conditions[2] == false && Frame == 4)
@@ -1007,15 +1028,15 @@ namespace GeometryCollectionExample
 					FVector X2 = Particles.X(2);
 					FVector X7 = Particles.X(7);
 
-					R.ExpectTrue(FMath::Abs(X1.Size() - Ref1.Size()) > KINDA_SMALL_NUMBER,  "Dynamic body failed to move");
-					R.ExpectTrue(FMath::Abs(X2.Size() - Ref2.Size()) <= KINDA_SMALL_NUMBER, "Kinematic body moved");
-					R.ExpectTrue(FMath::Abs(X7.Size() - Ref7.Size()) <= KINDA_SMALL_NUMBER, "Kinematic body moved");
-					R.ExpectTrue(Particles.InvM(8) == 0.f); // kinematic cluster
-					R.ExpectTrue(Particles.ObjectState(8) == Chaos::EObjectStateType::Kinematic); // kinematic cluster
+					EXPECT_GT(FMath::Abs(X1.Size() - Ref1.Size()), KINDA_SMALL_NUMBER) << *FString("Dynamic body failed to move");
+					EXPECT_NEAR(FMath::Abs(X2.Size() - Ref2.Size()), 0, KINDA_SMALL_NUMBER) << *FString("Kinematic body moved");
+					EXPECT_NEAR(FMath::Abs(X7.Size() - Ref7.Size()), 0, KINDA_SMALL_NUMBER) << FString("Kinematic body moved");
+					EXPECT_EQ(Particles.InvM(8), 0.f); // kinematic cluster
+					EXPECT_EQ(Particles.ObjectState(8), Chaos::EObjectStateType::Kinematic); // kinematic cluster
 
-					R.ExpectTrue(ClusterMap.Num() == 2);
-					R.ExpectTrue(ClusterMapContains(ClusterMap, 7, { 6,3 }));
-					R.ExpectTrue(ClusterMapContains(ClusterMap, 6, { 5,4 }));
+					EXPECT_EQ(ClusterMap.Num(), 2);
+					EXPECT_TRUE(ClusterMapContains(ClusterMap, 7, { 6,3 }));
+					EXPECT_TRUE(ClusterMapContains(ClusterMap, 6, { 5,4 }));
 				}
 			}
 			else if (Conditions[2] == true && Conditions[3] == false && Frame == 6)
@@ -1036,23 +1057,25 @@ namespace GeometryCollectionExample
 					FVector X2 = Particles.X(2);
 					FVector X7 = Particles.X(7);
 
-					R.ExpectTrue(FMath::Abs(X1.Size() - Ref1.Size()) > KINDA_SMALL_NUMBER, "Dynamic body 1 failed to move.");
-					R.ExpectTrue(FMath::Abs(X2.Size() - Ref2.Size()) <= KINDA_SMALL_NUMBER, "Kinematic body moved");
-					R.ExpectTrue(FMath::Abs(X7.Size() - Ref7.Size()) > KINDA_SMALL_NUMBER, "Dynamic body 7 failed to move");
-					R.ExpectTrue(Particles.InvM(8) == 0.f); // kinematic cluster
-					R.ExpectTrue(Particles.ObjectState(8) == Chaos::EObjectStateType::Kinematic); // kinematic cluster
+					EXPECT_GT(FMath::Abs(X1.Size() - Ref1.Size()), KINDA_SMALL_NUMBER) << *FString("Dynamic body 1 failed to move.");
+					EXPECT_NEAR(FMath::Abs(X2.Size() - Ref2.Size()), 0, KINDA_SMALL_NUMBER) << *FString("Kinematic body moved");
+					EXPECT_GT(FMath::Abs(X7.Size() - Ref7.Size()), KINDA_SMALL_NUMBER) << *FString("Dynamic body 7 failed to move");
+					EXPECT_EQ(Particles.InvM(8), 0.f); // kinematic cluster
+					EXPECT_EQ(Particles.ObjectState(8), Chaos::EObjectStateType::Kinematic); // kinematic cluster
 
 
-					R.ExpectTrue(ClusterMap.Num() == 2);
-					R.ExpectTrue(ClusterMapContains(ClusterMap, 7, { 6,3 }));
-					R.ExpectTrue(ClusterMapContains(ClusterMap, 6, { 5,4 }));
+					EXPECT_EQ(ClusterMap.Num(), 2);
+					EXPECT_TRUE(ClusterMapContains(ClusterMap, 7, { 6,3 }));
+					EXPECT_TRUE(ClusterMapContains(ClusterMap, 6, { 5,4 }));
 				}
 			}
 		}
 		for (int i = 0; i < Conditions.Num(); i++)
 		{
-			R.ExpectTrue(Conditions[i]);
+			EXPECT_TRUE(Conditions[i]);
 		}
+#endif
+#endif
 
 		FChaosSolversModule::GetModule()->DestroySolver(Solver);
 
@@ -1061,25 +1084,20 @@ namespace GeometryCollectionExample
 #endif
 
 
-		return !R.HasError();
+		
 	}
-	template bool RigidBodies_ClusterTest_KinematicAnchor<float>(ExampleResponse&& R);
+	template void RigidBodies_ClusterTest_KinematicAnchor<float>();
 
 
 	template<class T>
-	bool RigidBodies_ClusterTest_StaticAnchor(ExampleResponse&& R)
+	void RigidBodies_ClusterTest_StaticAnchor()
 	{
 		// Test : Set one element static. When the cluster breaks the elements that do not contain the static
 		//        rigid body should be dynamic, while the clusters that contain the static body should remain 
 		//        static. 
 #if INCLUDE_CHAOS
 		TUniquePtr<Chaos::TChaosPhysicsMaterial<T>> PhysicalMaterial = MakeUnique<Chaos::TChaosPhysicsMaterial<T>>();
-		PhysicalMaterial->Friction = 0;
-		PhysicalMaterial->Restitution = 0;
-		PhysicalMaterial->SleepingLinearThreshold = 0;
-		PhysicalMaterial->SleepingAngularThreshold = 0;
-		PhysicalMaterial->DisabledLinearThreshold = 0;
-		PhysicalMaterial->DisabledAngularThreshold = 0;
+		InitMaterialToZero(PhysicalMaterial);
 
 		TSharedPtr<FGeometryCollection> RestCollection = GeometryCollection::MakeCubeElement(FTransform(FQuat::MakeFromEuler(FVector(0.f)), FVector(20.f)), FVector(1.0));
 		RestCollection->AppendGeometry(*GeometryCollection::MakeCubeElement(FTransform(FQuat::MakeFromEuler(FVector(0.f)), FVector(30.f)), FVector(1.0)));
@@ -1120,9 +1138,12 @@ namespace GeometryCollectionExample
 		PhysObject->Initialize();
 
 		Chaos::FPBDRigidsSolver* Solver = FChaosSolversModule::GetModule()->CreateSolver(true);
+#if TODO_REIMPLEMENT_RIGID_CLUSTERING
 		Chaos::FPBDRigidsSolver::FClusteringType & Clustering = Solver->GetRigidClustering();
 
+#if CHAOS_PARTICLEHANDLE_TODO
 		Solver->RegisterObject(PhysObject);
+#endif
 		Solver->SetHasFloor(true);
 		Solver->SetEnabled(true);
 		PhysObject->ActivateBodies();
@@ -1133,6 +1154,7 @@ namespace GeometryCollectionExample
 		TArray<bool> Conditions = { false,false,false,false };
 
 		Solver->AdvanceSolverBy(1 / 24.);
+#if TODO_REIMPLEMENT_GET_RIGID_PARTICLES
 		const Chaos::TPBDRigidParticles<float, 3>& Particles = PhysObject->GetSolver()->GetRigidParticles();
 		const Chaos::FPBDRigidsSolver::FClusteringType::FClusterMap & ClusterMap = Clustering.GetChildrenMap();
 
@@ -1156,14 +1178,14 @@ namespace GeometryCollectionExample
 			//	UE_LOG(GCTCL_Log, Verbose, TEXT("... ... ...    InvM[%d] : %f"), rdx, Particles.InvM(rdx));
 			//}
 
-			R.ExpectTrue(Particles.InvM(0) == 0.f); // floor
-			R.ExpectTrue(Particles.InvM(1) != 0.f); // dynamic rigid
-			R.ExpectTrue(Particles.InvM(2) == 0.f); // static rigid
-			R.ExpectTrue(Particles.InvM(3) != 0.f); // dynamic rigid
-			R.ExpectTrue(Particles.InvM(4) != 0.f); // dynamic rigid
-			R.ExpectTrue(Particles.InvM(5) != 0.f); // dynamic rigid
-			R.ExpectTrue(Particles.InvM(6) != 0.f); // dynamic rigid
-			R.ExpectTrue(Particles.InvM(7) != 0.f); // dynamic cluster
+			EXPECT_EQ(Particles.InvM(0), 0.f); // floor
+			EXPECT_NE(Particles.InvM(1), 0.f); // dynamic rigid
+			EXPECT_EQ(Particles.InvM(2), 0.f); // static rigid
+			EXPECT_NE(Particles.InvM(3), 0.f); // dynamic rigid
+			EXPECT_NE(Particles.InvM(4), 0.f); // dynamic rigid
+			EXPECT_NE(Particles.InvM(5), 0.f); // dynamic rigid
+			EXPECT_NE(Particles.InvM(6), 0.f); // dynamic rigid
+			EXPECT_NE(Particles.InvM(7), 0.f); // dynamic cluster
 
 			FVector Ref1, Ref2, Ref7; // RigidBody0(Dynamic), RigidBody1(static), RigidBody6(static then Dynamic)
 			if (Conditions[0] == false && Frame == 1)
@@ -1183,10 +1205,10 @@ namespace GeometryCollectionExample
 					Ref1 = Particles.X(1);
 					Ref2 = Particles.X(2);
 					Ref7 = Particles.X(7);
-					R.ExpectTrue(Particles.InvM(8) == 0.f); // static cluster
-					R.ExpectTrue(Particles.ObjectState(8) == Chaos::EObjectStateType::Static); // Static cluster
-					R.ExpectTrue(Particles.InvM(9) == 0.f); // static cluster
-					R.ExpectTrue(Particles.ObjectState(9) == Chaos::EObjectStateType::Static); // Static cluster
+					EXPECT_EQ(Particles.InvM(8), 0.f); // static cluster
+					EXPECT_EQ(Particles.ObjectState(8), Chaos::EObjectStateType::Static); // Static cluster
+					EXPECT_EQ(Particles.InvM(9), 0.f); // static cluster
+					EXPECT_EQ(Particles.ObjectState(9), Chaos::EObjectStateType::Static); // Static cluster
 				}
 			}
 			else if (Conditions[0] == true && Conditions[1] == false && Frame == 2)
@@ -1207,13 +1229,13 @@ namespace GeometryCollectionExample
 					FVector X2 = Particles.X(2);
 					FVector X7 = Particles.X(7);
 
-					R.ExpectTrue(FMath::Abs(X1.Size() - Ref1.Size()) <= KINDA_SMALL_NUMBER, "static body1 moved");
-					R.ExpectTrue(FMath::Abs(X2.Size() - Ref2.Size()) <= KINDA_SMALL_NUMBER, "static body2 moved");
-					R.ExpectTrue(FMath::Abs(X7.Size() - Ref7.Size()) <= KINDA_SMALL_NUMBER, "static body7 moved");
-					R.ExpectTrue(Particles.InvM(8) == 0.f); // static cluster
-					R.ExpectTrue(Particles.ObjectState(8) == Chaos::EObjectStateType::Static); // static cluster
-					R.ExpectTrue(Particles.InvM(9) == 0.f);
-					R.ExpectTrue(Particles.ObjectState(9) == Chaos::EObjectStateType::Static);
+					EXPECT_NEAR(X1.Size() - Ref1.Size(), 0, KINDA_SMALL_NUMBER) << *FString("static body1 moved");
+					EXPECT_NEAR(X2.Size() - Ref2.Size(), 0, KINDA_SMALL_NUMBER) << *FString("static body2 moved");
+					EXPECT_NEAR(X7.Size() - Ref7.Size(), 0, KINDA_SMALL_NUMBER) << *FString("static body7 moved");
+					EXPECT_EQ(Particles.InvM(8), 0.f); // static cluster
+					EXPECT_EQ(Particles.ObjectState(8), Chaos::EObjectStateType::Static); // static cluster
+					EXPECT_EQ(Particles.InvM(9), 0.f);
+					EXPECT_EQ(Particles.ObjectState(9), Chaos::EObjectStateType::Static);
 				}
 			}
 			else if (Conditions[1] == true && Conditions[2] == false && Frame == 4)
@@ -1234,11 +1256,11 @@ namespace GeometryCollectionExample
 					FVector X2 = Particles.X(2);
 					FVector X7 = Particles.X(7);
 
-					R.ExpectTrue(FMath::Abs(X1.Size() - Ref1.Size()) > KINDA_SMALL_NUMBER, "Dynamic body failed to move");
-					R.ExpectTrue(FMath::Abs(X2.Size() - Ref2.Size()) <= KINDA_SMALL_NUMBER, "static body moved");
-					R.ExpectTrue(FMath::Abs(X7.Size() - Ref7.Size()) <= KINDA_SMALL_NUMBER, "static body moved");
-					R.ExpectTrue(Particles.InvM(8) == 0.f); // static cluster
-					R.ExpectTrue(Particles.ObjectState(8) == Chaos::EObjectStateType::Static); // static cluster
+					EXPECT_GT(FMath::Abs(X1.Size() - Ref1.Size()), KINDA_SMALL_NUMBER) << *FString("Dynamic body failed to move");
+					EXPECT_NEAR(FX2.Size() - Ref2.Size(), 0, KINDA_SMALL_NUMBER) << *FString("static body moved");
+					EXPECT_NEAR(X7.Size() - Ref7.Size(), 0, KINDA_SMALL_NUMBER) << *FString("static body moved");
+					EXPECT_EQ(Particles.InvM(8), 0.f); // static cluster
+					EXPECT_EQ(Particles.ObjectState(8), Chaos::EObjectStateType::Static); // static cluster
 				}
 			}
 			else if (Conditions[2] == true && Conditions[3] == false && Frame == 6)
@@ -1259,44 +1281,37 @@ namespace GeometryCollectionExample
 					FVector X2 = Particles.X(2);
 					FVector X7 = Particles.X(7);
 
-					R.ExpectTrue(FMath::Abs(X1.Size() - Ref1.Size()) > KINDA_SMALL_NUMBER, "Dynamic body 1 failed to move.");
-					R.ExpectTrue(FMath::Abs(X2.Size() - Ref2.Size()) <= KINDA_SMALL_NUMBER, "static body moved");
-					R.ExpectTrue(FMath::Abs(X7.Size() - Ref7.Size()) > KINDA_SMALL_NUMBER, "Dynamic body 7 failed to move");
-					R.ExpectTrue(Particles.InvM(8) == 0.f); // static cluster
-					R.ExpectTrue(Particles.ObjectState(8) == Chaos::EObjectStateType::Static); // static cluster
+					EXPECT_GT(FMath::Abs(X1.Size() - Ref1.Size()), KINDA_SMALL_NUMBER) << *FString("Dynamic body 1 failed to move.");
+					EXPECT_NEAR(X2.Size() - Ref2.Size(), 0, KINDA_SMALL_NUMBER) << *FString("static body moved");
+					EXPECT_GT(FMath::Abs(X7.Size() - Ref7.Size()), KINDA_SMALL_NUMBER) << *FString("Dynamic body 7 failed to move");
+					EXPECT_EQ(Particles.InvM(8), 0.f); // static cluster
+					EXPECT_EQ(Particles.ObjectState(8), Chaos::EObjectStateType::Static); // static cluster
 				}
 			}
 		}
 		for (int i = 0; i < Conditions.Num(); i++)
 		{
-			R.ExpectTrue(Conditions[i]);
+			EXPECT_TRUE(Conditions[i]);
 		}
-
+#endif
+#endif
 		FChaosSolversModule::GetModule()->DestroySolver(Solver);
 
 		delete PhysObject;
 
 #endif
-
-
-		return !R.HasError();
 	}
-	template bool RigidBodies_ClusterTest_StaticAnchor<float>(ExampleResponse&& R);
+	template void RigidBodies_ClusterTest_StaticAnchor<float>();
 
 
 	template<class T>
-	bool RigidBodies_ClusterTest_UnionClusters(ExampleResponse&& R)
+	void RigidBodies_ClusterTest_UnionClusters()
 	{
 		// Test : Set one element kinematic. When the cluster breaks the elements that do not contain the kinematic
 		//        Rigid body should be dynamic, while the clusters that contain the kinematic body should remain kinematic. 
 #if INCLUDE_CHAOS
 		TUniquePtr<Chaos::TChaosPhysicsMaterial<T>> PhysicalMaterial = MakeUnique<Chaos::TChaosPhysicsMaterial<T>>();
-		PhysicalMaterial->Friction = 0;
-		PhysicalMaterial->Restitution = 0;
-		PhysicalMaterial->SleepingLinearThreshold = 0;
-		PhysicalMaterial->SleepingAngularThreshold = 0;
-		PhysicalMaterial->DisabledLinearThreshold = 0;
-		PhysicalMaterial->DisabledAngularThreshold = 0;
+		InitMaterialToZero(PhysicalMaterial);
 
 		TSharedPtr<FGeometryCollection> RestCollection = CreateClusteredBody( FVector(0,0,100) );
 		TSharedPtr<FGeometryDynamicCollection> DynamicCollection = GeometryCollectionToGeometryDynamicCollection(RestCollection.Get(), (uint8)EObjectStateTypeEnum::Chaos_Object_Dynamic);
@@ -1344,8 +1359,10 @@ namespace GeometryCollectionExample
 
 
 		Chaos::FPBDRigidsSolver* Solver = FChaosSolversModule::GetModule()->CreateSolver(true);
+#if CHAOS_PARTICLEHANDLE_TODO
 		Solver->RegisterObject(PhysObject);
 		Solver->RegisterObject(PhysObject2);
+#endif
 		Solver->SetHasFloor(true);
 		Solver->SetEnabled(true);
 		PhysObject->ActivateBodies();
@@ -1359,8 +1376,10 @@ namespace GeometryCollectionExample
 		{
 			Solver->AdvanceSolverBy(1 / 24.);
 			FinalizeSolver(*Solver);
+#if TODO_REIMPLEMENT_GET_RIGID_PARTICLES
 			const Chaos::TPBDRigidParticles<float, 3>& Particles = PhysObject->GetSolver()->GetRigidParticles();
 			const Chaos::FPBDRigidsSolver::FClusteringType::FClusterMap & ClusterMap = Solver->GetRigidClustering().GetChildrenMap();
+#endif
 
 			if (Frame == 0)
 			{
@@ -1380,9 +1399,10 @@ namespace GeometryCollectionExample
 					}
 				}
 
-				R.ExpectTrue(ClusterMap.Num() == 1);
-				R.ExpectTrue(ClusterMapContains(ClusterMap, 7, { 1,2,5,4 }));
-
+#if TODO_REIMPLEMENT_RIGID_CLUSTERING
+				EXPECT_EQ(ClusterMap.Num(), 1);
+				EXPECT_TRUE(ClusterMapContains(ClusterMap, 7, { 1,2,5,4 }));
+#endif
 			}
 		}
 
@@ -1403,7 +1423,7 @@ namespace GeometryCollectionExample
 		}
 		for (int i = 0; i < Distances.Num()/2.0; i++)
 		{
-			R.ExpectTrue( FMath::Abs(Distances[i] - Distances2[i]) < 0.1 );
+			EXPECT_LT( FMath::Abs(Distances[i] - Distances2[i]), 0.1 );
 		}
 
 
@@ -1412,27 +1432,19 @@ namespace GeometryCollectionExample
 		delete PhysObject2;
 
 #endif
-
-
-		return !R.HasError();
 	}
-	template bool RigidBodies_ClusterTest_UnionClusters<float>(ExampleResponse&& R);
+	template void RigidBodies_ClusterTest_UnionClusters<float>();
 
 
 
 
 
 	template<class T>
-	bool RigidBodies_ClusterTest_ReleaseClusterParticle_ClusteredNode(ExampleResponse&& R)
+	void RigidBodies_ClusterTest_ReleaseClusterParticle_ClusteredNode()
 	{
 #if INCLUDE_CHAOS
 		TUniquePtr<Chaos::TChaosPhysicsMaterial<T>> PhysicalMaterial = MakeUnique<Chaos::TChaosPhysicsMaterial<T>>();
-		PhysicalMaterial->Friction = 0;
-		PhysicalMaterial->Restitution = 0;
-		PhysicalMaterial->SleepingLinearThreshold = 0;
-		PhysicalMaterial->SleepingAngularThreshold = 0;
-		PhysicalMaterial->DisabledLinearThreshold = 0;
-		PhysicalMaterial->DisabledAngularThreshold = 0;
+		InitMaterialToZero(PhysicalMaterial);
 
 		TSharedPtr<FGeometryCollection> RestCollection = CreateClusteredBody(FVector(0, 0, 100));
 		TSharedPtr<FGeometryDynamicCollection> DynamicCollection = GeometryCollectionToGeometryDynamicCollection(RestCollection.Get(), (uint8)EObjectStateTypeEnum::Chaos_Object_Dynamic);
@@ -1482,9 +1494,13 @@ namespace GeometryCollectionExample
 
 
 		Chaos::FPBDRigidsSolver* Solver = FChaosSolversModule::GetModule()->CreateSolver(true);
+#if TODO_REIMPLEMENT_RIGID_CLUSTERING
 		Chaos::FPBDRigidsSolver::FClusteringType & Clustering = Solver->GetRigidClustering();
+#endif
+#if CHAOS_PARTICLEHANDLE_TODO
 		Solver->RegisterObject(PhysObject);
 		Solver->RegisterObject(PhysObject2);
+#endif
 		Solver->SetHasFloor(true);
 		Solver->SetEnabled(true);
 		PhysObject->ActivateBodies();
@@ -1500,26 +1516,28 @@ namespace GeometryCollectionExample
 		TArray<FTransform> InitialGlobalTransforms2;
 		GeometryCollectionAlgo::GlobalMatrices(DynamicCollection2->Transform, DynamicCollection2->Parent, InitialGlobalTransforms2);
 
-
-
-
-
-
+#if TODO_REIMPLEMENT_GET_RIGID_PARTICLES
 		const Chaos::TPBDRigidParticles<float, 3>& Particles = PhysObject->GetSolver()->GetRigidParticles();
+#endif
+#if TODO_REIMPLEMENT_RIGID_CLUSTERING
 		const Chaos::TArrayCollectionArray<Chaos::ClusterId> & ClusterIdsArray = Solver->GetRigidClustering().GetClusterIdsArray();
 		const Chaos::FPBDRigidsSolver::FClusteringType::FClusterMap & ClusterMap = Solver->GetRigidClustering().GetChildrenMap();
 
-		R.ExpectTrue(ClusterMap.Num() == 1);
-		R.ExpectTrue(ClusterMapContains(ClusterMap, 7, { 1,2,4,5}));
+		EXPECT_EQ(ClusterMap.Num(), 1);
+		EXPECT_TRUE(ClusterMapContains(ClusterMap, 7, { 1,2,4,5}));
+#endif
 
 		for (int Frame = 1; Frame < 10; Frame++)
 		{
 			Solver->AdvanceSolverBy(1 / 24.);
 
+#if TODO_REIMPLEMENT_RIGID_CLUSTERING
 			if (Frame == 5)
 			{
 				Solver->GetRigidClustering().ReleaseClusterParticles({4,5});
 			}
+#endif
+
 			FinalizeSolver(*Solver);
 
 			if (Frame < 5)
@@ -1528,48 +1546,44 @@ namespace GeometryCollectionExample
 				GeometryCollectionAlgo::GlobalMatrices(DynamicCollection2->Transform, DynamicCollection2->Parent, GlobalTransforms2);
 				for (int i = 0; i < GlobalTransforms2.Num(); i++)
 				{
-					R.ExpectTrue((GlobalTransforms2[i].GetTranslation() - InitialGlobalTransforms2[i].GetTranslation()).Size() < KINDA_SMALL_NUMBER);
+					EXPECT_LT((GlobalTransforms2[i].GetTranslation() - InitialGlobalTransforms2[i].GetTranslation()).Size(),KINDA_SMALL_NUMBER);
 				}
 			}
 			TArray<FTransform> GlobalTransforms;
 			GeometryCollectionAlgo::GlobalMatrices(DynamicCollection->Transform, DynamicCollection->Parent, GlobalTransforms);
 			for (int i = 0; i < GlobalTransforms.Num(); i++)
 			{
-				R.ExpectTrue((GlobalTransforms[i].GetTranslation() - InitialGlobalTransforms[i].GetTranslation()).Size() < KINDA_SMALL_NUMBER);
+				EXPECT_LT((GlobalTransforms[i].GetTranslation() - InitialGlobalTransforms[i].GetTranslation()).Size(),KINDA_SMALL_NUMBER);
 			}
 
 		}
 
-		R.ExpectTrue(ClusterMap.Num() == 1);
-		R.ExpectTrue(ClusterMapContains(ClusterMap, 8, { 1,2 }));
-
+#if TODO_REIMPLEMENT_RIGID_CLUSTERING
+		EXPECT_EQ(ClusterMap.Num(), 1);
+		EXPECT_TRUE(ClusterMapContains(ClusterMap, 8, { 1,2 }));
+#endif
 
 		TArray<int32> Subset({ 1 });
 		TArray<FTransform> GlobalTransform2;
 		GeometryCollectionAlgo::GlobalMatrices(DynamicCollection2->Transform, DynamicCollection2->Parent, Subset, GlobalTransform2);
-		R.ExpectTrue((GlobalTransform2[0].GetTranslation() - InitialGlobalTransforms2[Subset[0]].GetTranslation()).Size() > SMALL_NUMBER);
-		R.ExpectTrue(GlobalTransform2[0].GetTranslation().Z < InitialGlobalTransforms2[Subset[0]].GetTranslation().Z);
+		EXPECT_GT((GlobalTransform2[0].GetTranslation() - InitialGlobalTransforms2[Subset[0]].GetTranslation()).Size(), SMALL_NUMBER);
+		EXPECT_LT(GlobalTransform2[0].GetTranslation().Z, InitialGlobalTransforms2[Subset[0]].GetTranslation().Z);
 
 		FChaosSolversModule::GetModule()->DestroySolver(Solver);
 
 		delete PhysObject;
+		delete PhysObject2;
 
 #endif
-		return !R.HasError();
 	}
-	template bool RigidBodies_ClusterTest_ReleaseClusterParticle_ClusteredNode<float>(ExampleResponse&& R);
+	template void RigidBodies_ClusterTest_ReleaseClusterParticle_ClusteredNode<float>();
 
 	template<class T>
-	bool RigidBodies_ClusterTest_ReleaseClusterParticle_ClusteredKinematicNode(ExampleResponse&& R)
+	void RigidBodies_ClusterTest_ReleaseClusterParticle_ClusteredKinematicNode()
 	{
 #if INCLUDE_CHAOS
 		TUniquePtr<Chaos::TChaosPhysicsMaterial<T>> PhysicalMaterial = MakeUnique<Chaos::TChaosPhysicsMaterial<T>>();
-		PhysicalMaterial->Friction = 0;
-		PhysicalMaterial->Restitution = 0;
-		PhysicalMaterial->SleepingLinearThreshold = 0;
-		PhysicalMaterial->SleepingAngularThreshold = 0;
-		PhysicalMaterial->DisabledLinearThreshold = 0;
-		PhysicalMaterial->DisabledAngularThreshold = 0;
+		InitMaterialToZero(PhysicalMaterial);
 
 		TSharedPtr<FGeometryCollection> RestCollection = CreateClusteredBody(FVector(0, 0, 100));
 		TSharedPtr<FGeometryDynamicCollection> DynamicCollection = GeometryCollectionToGeometryDynamicCollection(RestCollection.Get(), (uint8)EObjectStateTypeEnum::Chaos_Object_Dynamic);
@@ -1618,10 +1632,14 @@ namespace GeometryCollectionExample
 		PhysObject2->Initialize();
 
 		Chaos::FPBDRigidsSolver* Solver = FChaosSolversModule::GetModule()->CreateSolver(true);
+#if TODO_REIMPLEMENT_RIGID_CLUSTERING
 		Chaos::FPBDRigidsSolver::FClusteringType & Clustering = Solver->GetRigidClustering();
-		
+#endif
+
+#if CHAOS_PARTICLEHANDLE_TODO
 		Solver->RegisterObject(PhysObject);
 		Solver->RegisterObject(PhysObject2);
+#endif
 		Solver->SetHasFloor(true);
 		Solver->SetEnabled(true);
 		PhysObject->ActivateBodies();
@@ -1638,17 +1656,24 @@ namespace GeometryCollectionExample
 		TArray<FTransform> InitialGlobalTransforms2;
 		GeometryCollectionAlgo::GlobalMatrices(DynamicCollection2->Transform, DynamicCollection2->Parent, InitialGlobalTransforms2);
 
+#if TODO_REIMPLEMENT_GET_RIGID_PARTICLES
 		const Chaos::TPBDRigidParticles<float, 3>& Particles = PhysObject->GetSolver()->GetRigidParticles();
+#endif
+#if TODO_REIMPLEMENT_RIGID_CLUSTERING
 		const Chaos::TArrayCollectionArray<Chaos::ClusterId> & ClusterIdsArray = Solver->GetRigidClustering().GetClusterIdsArray();
+#endif
 
 		for (int Frame = 1; Frame < 10; Frame++)
 		{
 			Solver->AdvanceSolverBy(1 / 24.);
 
+#if TODO_REIMPLEMENT_RIGID_CLUSTERING
 			if (Frame == 5)
 			{
 				Solver->GetRigidClustering().ReleaseClusterParticles({ 4,5 });
 			}
+#endif
+
 			FinalizeSolver(*Solver);
 
 			// the cluster from DynamicCollection will always be kinematic, and will be released from the union. 
@@ -1656,7 +1681,7 @@ namespace GeometryCollectionExample
 			GeometryCollectionAlgo::GlobalMatrices(DynamicCollection->Transform, DynamicCollection->Parent, GlobalTransforms);
 			for (int i = 0; i < GlobalTransforms.Num(); i++)
 			{
-				R.ExpectTrue((GlobalTransforms[i].GetTranslation() - InitialGlobalTransforms[i].GetTranslation()).Size() < KINDA_SMALL_NUMBER);
+				EXPECT_LT((GlobalTransforms[i].GetTranslation() - InitialGlobalTransforms[i].GetTranslation()).Size(),KINDA_SMALL_NUMBER);
 			}
 
 			// the cluster from DynamicCollection2 will always be dynamic after its released from the union, but should be
@@ -1667,7 +1692,7 @@ namespace GeometryCollectionExample
 				GeometryCollectionAlgo::GlobalMatrices(DynamicCollection2->Transform, DynamicCollection2->Parent, GlobalTransforms2);
 				for (int i = 0; i < GlobalTransforms2.Num(); i++)
 				{
-					R.ExpectTrue((GlobalTransforms2[i].GetTranslation() - InitialGlobalTransforms2[i].GetTranslation()).Size() < KINDA_SMALL_NUMBER);
+					EXPECT_LT((GlobalTransforms2[i].GetTranslation() - InitialGlobalTransforms2[i].GetTranslation()).Size(),KINDA_SMALL_NUMBER);
 				}
 			}
 		}
@@ -1676,30 +1701,25 @@ namespace GeometryCollectionExample
 		TArray<int32> Subset({ 1 });
 		TArray<FTransform> GlobalTransform2;
 		GeometryCollectionAlgo::GlobalMatrices(DynamicCollection2->Transform, DynamicCollection2->Parent, Subset, GlobalTransform2);
-		R.ExpectTrue((GlobalTransform2[0].GetTranslation() - InitialGlobalTransforms2[Subset[0]].GetTranslation()).Size() > SMALL_NUMBER);
-		R.ExpectTrue(GlobalTransform2[0].GetTranslation().Z < InitialGlobalTransforms2[Subset[0]].GetTranslation().Z);
+		EXPECT_GT((GlobalTransform2[0].GetTranslation() - InitialGlobalTransforms2[Subset[0]].GetTranslation()).Size(), SMALL_NUMBER);
+		EXPECT_LT(GlobalTransform2[0].GetTranslation().Z, InitialGlobalTransforms2[Subset[0]].GetTranslation().Z);
 
 		FChaosSolversModule::GetModule()->DestroySolver(Solver);
 
 		delete PhysObject;
-
+		delete PhysObject2;
 #endif
-		return !R.HasError();
+		
 	}
-	template bool RigidBodies_ClusterTest_ReleaseClusterParticle_ClusteredKinematicNode<float>(ExampleResponse&& R);
+	template void RigidBodies_ClusterTest_ReleaseClusterParticle_ClusteredKinematicNode<float>();
 
 
 	template<class T>
-	bool RigidBodies_ClusterTest_ReleaseClusterParticles_AllLeafNodes(ExampleResponse&& R)
+	void RigidBodies_ClusterTest_ReleaseClusterParticles_AllLeafNodes()
 	{
 #if INCLUDE_CHAOS
 		TUniquePtr<Chaos::TChaosPhysicsMaterial<T>> PhysicalMaterial = MakeUnique<Chaos::TChaosPhysicsMaterial<T>>();
-		PhysicalMaterial->Friction = 0;
-		PhysicalMaterial->Restitution = 0;
-		PhysicalMaterial->SleepingLinearThreshold = 0;
-		PhysicalMaterial->SleepingAngularThreshold = 0;
-		PhysicalMaterial->DisabledLinearThreshold = 0;
-		PhysicalMaterial->DisabledAngularThreshold = 0;
+		InitMaterialToZero(PhysicalMaterial);
 
 		TSharedPtr<FGeometryCollection> RestCollection = CreateClusteredBody(FVector(0, 0, 100));
 		TSharedPtr<FGeometryDynamicCollection> DynamicCollection = GeometryCollectionToGeometryDynamicCollection(RestCollection.Get(), (uint8)EObjectStateTypeEnum::Chaos_Object_Dynamic);
@@ -1724,7 +1744,9 @@ namespace GeometryCollectionExample
 		PhysObject->Initialize();
 
 		Chaos::FPBDRigidsSolver* Solver = FChaosSolversModule::GetModule()->CreateSolver(true);
+#if CHAOS_PARTICLEHANDLE_TODO
 		Solver->RegisterObject(PhysObject);
+#endif
 		Solver->SetHasFloor(false);
 		Solver->SetEnabled(true);
 		PhysObject->ActivateBodies();
@@ -1736,16 +1758,21 @@ namespace GeometryCollectionExample
 		GeometryCollectionAlgo::GlobalMatrices(DynamicCollection->Transform, DynamicCollection->Parent, InitialGlobalTransforms);
 		float PreviousHeight = InitialGlobalTransforms[0].GetTranslation().Y;
 
-
+#if TODO_REIMPLEMENT_GET_RIGID_PARTICLES
 		const Chaos::TPBDRigidParticles<float, 3>& Particles = PhysObject->GetSolver()->GetRigidParticles();
+#endif
+#if TODO_REIMPLEMENT_RIGID_CLUSTERING
 		const Chaos::TArrayCollectionArray<Chaos::ClusterId> & ClusterIdsArray = Solver->GetRigidClustering().GetClusterIdsArray();
+#endif
 
 		for (int Frame = 1; Frame < 10; Frame++)
 		{
+#if TODO_REIMPLEMENT_RIGID_CLUSTERING
 			if (Frame == 5)
 			{
 				Solver->GetRigidClustering().ReleaseClusterParticles({ 0,1 });
 			}
+#endif
 
 			Solver->AdvanceSolverBy(1 / 24.);
 			FinalizeSolver(*Solver);
@@ -1758,13 +1785,13 @@ namespace GeometryCollectionExample
 			{
 				for (int i = 0; i < GlobalTransforms.Num(); i++)
 				{
-					R.ExpectTrue((GlobalTransforms[i].GetTranslation() - InitialGlobalTransforms[i].GetTranslation()).Size() < KINDA_SMALL_NUMBER);
+					EXPECT_LT((GlobalTransforms[i].GetTranslation() - InitialGlobalTransforms[i].GetTranslation()).Size(), KINDA_SMALL_NUMBER);
 				}
 			}
 			else if (Frame <= 7)
 			{
-				R.ExpectTrue(PreviousHeight > GlobalTransforms[0].GetTranslation().Z);
-				R.ExpectTrue(FMath::Abs(InitialGlobalTransforms[1].GetTranslation().Z - GlobalTransforms[1].GetTranslation().Z) < KINDA_SMALL_NUMBER);
+				EXPECT_GT(PreviousHeight, GlobalTransforms[0].GetTranslation().Z);
+				EXPECT_LT(FMath::Abs(InitialGlobalTransforms[1].GetTranslation().Z - GlobalTransforms[1].GetTranslation().Z), KINDA_SMALL_NUMBER);
 			}
 			PreviousHeight = GlobalTransforms[0].GetTranslation().Z;
 		}
@@ -1773,21 +1800,16 @@ namespace GeometryCollectionExample
 		delete PhysObject;
 
 #endif
-		return !R.HasError();
+		
 	}
-	template bool RigidBodies_ClusterTest_ReleaseClusterParticles_AllLeafNodes<float>(ExampleResponse&& R);
+	template void RigidBodies_ClusterTest_ReleaseClusterParticles_AllLeafNodes<float>();
 
 	template<class T>
-	bool RigidBodies_ClusterTest_ReleaseClusterParticles_ClusterNodeAndSubClusterNode(ExampleResponse&& R)
+	void RigidBodies_ClusterTest_ReleaseClusterParticles_ClusterNodeAndSubClusterNode()
 	{
 #if INCLUDE_CHAOS
 		TUniquePtr<Chaos::TChaosPhysicsMaterial<T>> PhysicalMaterial = MakeUnique<Chaos::TChaosPhysicsMaterial<T>>();
-		PhysicalMaterial->Friction = 0;
-		PhysicalMaterial->Restitution = 0;
-		PhysicalMaterial->SleepingLinearThreshold = 0;
-		PhysicalMaterial->SleepingAngularThreshold = 0;
-		PhysicalMaterial->DisabledLinearThreshold = 0;
-		PhysicalMaterial->DisabledAngularThreshold = 0;
+		InitMaterialToZero(PhysicalMaterial);
 
 		TSharedPtr<FGeometryCollection> RestCollection = CreateClusteredBody_TwoParents_TwoBodies(FVector(0, 0, 100));
 		TSharedPtr<FGeometryDynamicCollection> DynamicCollection = GeometryCollectionToGeometryDynamicCollection(RestCollection.Get(), (uint8)EObjectStateTypeEnum::Chaos_Object_Dynamic);
@@ -1814,9 +1836,13 @@ namespace GeometryCollectionExample
 
 
 		Chaos::FPBDRigidsSolver* Solver = FChaosSolversModule::GetModule()->CreateSolver(true);
+#if TODO_REIMPLEMENT_RIGID_CLUSTERING
 		Chaos::FPBDRigidsSolver::FClusteringType & Clustering = Solver->GetRigidClustering();
+#endif
 
+#if CHAOS_PARTICLEHANDLE_TODO
 		Solver->RegisterObject(PhysObject);
+#endif
 		Solver->SetHasFloor(false);
 		Solver->SetEnabled(true);
 		PhysObject->ActivateBodies();
@@ -1824,13 +1850,16 @@ namespace GeometryCollectionExample
 		Solver->AdvanceSolverBy(1 / 24.);
 		FinalizeSolver(*Solver);
 
+#if TODO_REIMPLEMENT_GET_RIGID_PARTICLES
 		const Chaos::TPBDRigidParticles<float, 3>& Particles = PhysObject->GetSolver()->GetRigidParticles();
+#if TODO_REIMPLEMENT_RIGID_CLUSTERING
 		const Chaos::TArrayCollectionArray<Chaos::ClusterId> & ClusterIdsArray = Solver->GetRigidClustering().GetClusterIdsArray();
 		const Chaos::FPBDRigidsSolver::FClusteringType::FClusterMap & ClusterMap = Solver->GetRigidClustering().GetChildrenMap();
 
-		R.ExpectTrue(ClusterMap.Num() == 2);
-		R.ExpectTrue(ClusterMapContains(ClusterMap, 2, { 0,1 }));
-		R.ExpectTrue(ClusterMapContains(ClusterMap, 4, { 2 }));
+		EXPECT_EQ(ClusterMap.Num(), 2);
+		EXPECT_TRUE(ClusterMapContains(ClusterMap, 2, { 0,1 }));
+		EXPECT_TRUE(ClusterMapContains(ClusterMap, 4, { 2 }));
+
 
 		for (int Frame = 1; Frame < 10; Frame++)
 		{
@@ -1845,53 +1874,49 @@ namespace GeometryCollectionExample
 
 			if (Frame < 5)
 			{
-				R.ExpectTrue(Particles.Disabled(2));
-				R.ExpectTrue(ClusterIdsArray[2].Id != INDEX_NONE);
-				R.ExpectTrue(ClusterIdsArray[3].Id == INDEX_NONE);
-				R.ExpectTrue(ClusterIdsArray[4].Id == INDEX_NONE);
+				EXPECT_TRUE(Particles.Disabled(2));
+				EXPECT_NE(ClusterIdsArray[2].Id, INDEX_NONE);
+				EXPECT_EQ(ClusterIdsArray[3].Id, INDEX_NONE);
+				EXPECT_EQ(ClusterIdsArray[4].Id, INDEX_NONE);
 			}
 			else
 			{
-				R.ExpectTrue(!Particles.Disabled(2));
-				R.ExpectTrue(ClusterIdsArray[2].Id == INDEX_NONE);
-				R.ExpectTrue(ClusterIdsArray[3].Id == INDEX_NONE);
-				R.ExpectTrue(ClusterIdsArray[4].Id == INDEX_NONE);
+				EXPECT_TRUE(!Particles.Disabled(2));
+				EXPECT_EQ(ClusterIdsArray[2].Id, INDEX_NONE);
+				EXPECT_EQ(ClusterIdsArray[3].Id, INDEX_NONE);
+				EXPECT_EQ(ClusterIdsArray[4].Id, INDEX_NONE);
 
-				R.ExpectTrue(ClusterMap.Num() == 1);
-				R.ExpectTrue(ClusterMapContains(ClusterMap, 2, { 0,1 }));
+				EXPECT_EQ(ClusterMap.Num(), 1);
+				EXPECT_TRUE(ClusterMapContains(ClusterMap, 2, { 0,1 }));
 			}
 		}
+#endif
+#endif
 
 		FChaosSolversModule::GetModule()->DestroySolver(Solver);
 
 		delete PhysObject;
 
 #endif
-		return !R.HasError();
 	}
-	template bool RigidBodies_ClusterTest_ReleaseClusterParticles_ClusterNodeAndSubClusterNode<float>(ExampleResponse&& R);
+	template void RigidBodies_ClusterTest_ReleaseClusterParticles_ClusterNodeAndSubClusterNode<float>();
 
 	template<class T>
-	bool RigidBodies_ClusterTest_RemoveOnFracture(ExampleResponse&& R)
+	void RigidBodies_ClusterTest_RemoveOnFracture()
 	{
 #if INCLUDE_CHAOS
 		TUniquePtr<Chaos::TChaosPhysicsMaterial<T>> PhysicalMaterial = MakeUnique<Chaos::TChaosPhysicsMaterial<T>>();
-		PhysicalMaterial->Friction = 0;
-		PhysicalMaterial->Restitution = 0;
-		PhysicalMaterial->SleepingLinearThreshold = 0;
-		PhysicalMaterial->SleepingAngularThreshold = 0;
-		PhysicalMaterial->DisabledLinearThreshold = 0;
-		PhysicalMaterial->DisabledAngularThreshold = 0;
+		InitMaterialToZero(PhysicalMaterial);
 
 		TSharedPtr<FGeometryCollection> RestCollection = GeometryCollection::MakeCubeElement(FTransform(FQuat::MakeFromEuler(FVector(0, 0, 0.)), FVector(0, -10, 10)), FVector(1.0));
 		RestCollection->AppendGeometry(*GeometryCollection::MakeCubeElement(FTransform(FQuat::MakeFromEuler(FVector(0, 0, 0.)), FVector(0, 10, 10)), FVector(1.0)));
-		R.ExpectTrue(RestCollection->Transform.Num() == 2);
+		EXPECT_EQ(RestCollection->Transform.Num(), 2);
 
 		// this transform should have a zero scale after the simulation has run to the point of fracture
 		RestCollection->SetFlags(1, FGeometryCollection::FS_RemoveOnFracture);
 
 		FGeometryCollectionClusteringUtility::ClusterAllBonesUnderNewRoot(RestCollection.Get());
-		R.ExpectTrue(RestCollection->Transform.Num() == 3);
+		EXPECT_EQ(RestCollection->Transform.Num(), 3);
 		RestCollection->Transform[2] = FTransform(FQuat::MakeFromEuler(FVector(90.f, 0, 0.)), FVector(0, 0, 40));
 
 		TSharedPtr<FGeometryDynamicCollection> DynamicCollection = GeometryCollectionToGeometryDynamicCollection(RestCollection.Get());
@@ -1920,12 +1945,16 @@ namespace GeometryCollectionExample
 
 		FGeometryCollectionPhysicsProxy* PhysObject = new FGeometryCollectionPhysicsProxy(nullptr, DynamicCollection.Get(), InitFunc, nullptr, nullptr);
 		Chaos::FPBDRigidsSolver* Solver = FChaosSolversModule::GetModule()->CreateSolver(true);
+#if TODO_REIMPLEMENT_RIGID_CLUSTERING
 		Chaos::FPBDRigidsSolver::FClusteringType & Clustering = Solver->GetRigidClustering();
+#endif
 
 		PhysObject->Initialize();
 
+#if CHAOS_PARTICLEHANDLE_TODO
 		Solver->RegisterObject(PhysObject);
 		Solver->RegisterObject(FieldObject);
+#endif
 		Solver->SetHasFloor(false);
 		Solver->SetEnabled(true);
 		PhysObject->ActivateBodies();
@@ -1933,7 +1962,9 @@ namespace GeometryCollectionExample
 		TManagedArray<FTransform>& Transform = DynamicCollection->Transform;
 		float StartingRigidDistance = (Transform[1].GetTranslation() - Transform[0].GetTranslation()).Size(), CurrentRigidDistance = 0.f;
 
+#if TODO_REIMPLEMENT_RIGID_CLUSTERING
 		Chaos::TArrayCollectionArray<float>& InternalStrain = Clustering.GetStrainArray();
+#endif
 
 		FName TargetName = GetFieldPhysicsName(EFieldPhysicsType::Field_ExternalClusterStrain);
 		FFieldSystemCommand Command(TargetName, FalloffField->NewCopy());
@@ -1943,9 +1974,9 @@ namespace GeometryCollectionExample
 
 		FVector Scale = Transform[1].GetScale3D();
 
-		R.ExpectTrue(FMath::IsNearlyEqual(Scale.X, 1.0f, SMALL_NUMBER));
-		R.ExpectTrue(FMath::IsNearlyEqual(Scale.Y, 1.0f, SMALL_NUMBER));
-		R.ExpectTrue(FMath::IsNearlyEqual(Scale.Z, 1.0f, SMALL_NUMBER));
+		EXPECT_NEAR(Scale.X, 1.0f, SMALL_NUMBER);
+		EXPECT_NEAR(Scale.Y, 1.0f, SMALL_NUMBER);
+		EXPECT_NEAR(Scale.Z, 1.0f, SMALL_NUMBER);
 
 		Solver->AdvanceSolverBy(1 / 24.);
 		FinalizeSolver(*Solver);
@@ -1957,27 +1988,24 @@ namespace GeometryCollectionExample
 
 		FVector Scale2 = Transform[1].GetScale3D();
 		// geometry hidden by 0 scaling on transform
-		R.ExpectTrue(FMath::IsNearlyEqual(Scale2.X, 0.0f, SMALL_NUMBER));
-		R.ExpectTrue(FMath::IsNearlyEqual(Scale2.Y, 0.0f, SMALL_NUMBER));
-		R.ExpectTrue(FMath::IsNearlyEqual(Scale2.Z, 0.0f, SMALL_NUMBER));
+		EXPECT_NEAR(Scale2.X, 0.0f, SMALL_NUMBER);
+		EXPECT_NEAR(Scale2.Y, 0.0f, SMALL_NUMBER);
+		EXPECT_NEAR(Scale2.Z, 0.0f, SMALL_NUMBER);
 			
 		FChaosSolversModule::GetModule()->DestroySolver(Solver);
 
 		delete PhysObject;
 		delete FalloffField;
-
+		delete FieldObject;
 #endif
-
-
-		return !R.HasError();
 	}
-	template bool RigidBodies_ClusterTest_RemoveOnFracture<float>(ExampleResponse&& R);
+	template void RigidBodies_ClusterTest_RemoveOnFracture<float>();
 
 
 
 	
 	template<class T>
-	bool RigidBodiess_ClusterTest_ParticleImplicitCollisionGeometry(ExampleResponse&& R)
+	void RigidBodiess_ClusterTest_ParticleImplicitCollisionGeometry()
 	{
 #if INCLUDE_CHAOS
 		typename SimulationObjects<T>::FParameters P;
@@ -2002,21 +2030,29 @@ namespace GeometryCollectionExample
 		Chaos::FPBDRigidsSolver* Solver = FChaosSolversModule::GetModule()->CreateSolver(true);
 		Solver->SetHasFloor(false);
 		Solver->SetEnabled(true);
+#if CHAOS_PARTICLEHANDLE_TODO
 		Solver->RegisterObject(Object->PhysicsProxy.Get());
+#endif
+
 		Solver->AdvanceSolverBy(1 / 24.);
+
+#if TODO_REIMPLEMENT_GET_RIGID_PARTICLES
 		Chaos::TPBDRigidParticles<float, 3>& Particles = Solver->GetRigidParticles();
+#endif
 
 		float CollisionParticlesPerObjectFractionDefault = 0.5f;
 		IConsoleVariable*  CVarCollisionParticlesPerObjectFractionDefault = IConsoleManager::Get().FindConsoleVariable(TEXT("p.CollisionParticlesPerObjectFractionDefault"));
-		R.ExpectTrue(CVarCollisionParticlesPerObjectFractionDefault != nullptr);
+		EXPECT_NE(CVarCollisionParticlesPerObjectFractionDefault, nullptr);
 		if (CVarCollisionParticlesPerObjectFractionDefault != nullptr)
 		{
 			CollisionParticlesPerObjectFractionDefault = CVarCollisionParticlesPerObjectFractionDefault->GetFloat();
 		}
 
-		R.ExpectTrue(Particles.CollisionParticles(Object->PhysicsProxy->RigidBodyIDArray_TestingAccess()[10])->Size() == (int)(Simplicials[10]->Size() * CollisionParticlesPerObjectFractionDefault));
-		R.ExpectTrue(Particles.CollisionParticles(Object->PhysicsProxy->RigidBodyIDArray_TestingAccess()[11])->Size() == (int)(Simplicials[11]->Size() * CollisionParticlesPerObjectFractionDefault));
-		R.ExpectTrue(Particles.CollisionParticles(Object->PhysicsProxy->RigidBodyIDArray_TestingAccess()[12])->Size() == (int)(Simplicials[12]->Size() * CollisionParticlesPerObjectFractionDefault));
+#if TODO_REIMPLEMENT_GET_RIGID_PARTICLES
+		EXPECT_EQ(Particles.CollisionParticles(Object->PhysicsProxy->RigidBodyIDArray_TestingAccess()[10])->Size(), (int)(Simplicials[10]->Size() * CollisionParticlesPerObjectFractionDefault));
+		EXPECT_EQ(Particles.CollisionParticles(Object->PhysicsProxy->RigidBodyIDArray_TestingAccess()[11])->Size(), (int)(Simplicials[11]->Size() * CollisionParticlesPerObjectFractionDefault));
+		EXPECT_EQ(Particles.CollisionParticles(Object->PhysicsProxy->RigidBodyIDArray_TestingAccess()[12])->Size(), (int)(Simplicials[12]->Size() * CollisionParticlesPerObjectFractionDefault));
+#endif
 
 		// cleanup
 		//for (auto Obj : Collections) delete Obj;
@@ -2024,8 +2060,6 @@ namespace GeometryCollectionExample
 		
 		delete Object;
 #endif
-		return !R.HasError();
 	}
-	template bool RigidBodiess_ClusterTest_ParticleImplicitCollisionGeometry<float>(ExampleResponse&& R);
+	template void RigidBodiess_ClusterTest_ParticleImplicitCollisionGeometry<float>();
 }
-#endif
