Copying //UE4/Dev-Physics to Dev-Main (//UE4/Dev-Main) [at] 9098278

diff --git a/Engine/Plugins/Experimental/GeometryCollectionPlugin/Source/GeometryCollectionExampleCore/Private/GeometryCollection/GeometryCollectionExampleCreation.cpp b/Engine/Plugins/Experimental/GeometryCollectionPlugin/Source/GeometryCollectionExampleCore/Private/GeometryCollection/GeometryCollectionExampleCreation.cpp
--- a/Engine/Plugins/Experimental/GeometryCollectionPlugin/Source/GeometryCollectionExampleCore/Private/GeometryCollection/GeometryCollectionExampleCreation.cpp
+++ b/Engine/Plugins/Experimental/GeometryCollectionPlugin/Source/GeometryCollectionExampleCore/Private/GeometryCollection/GeometryCollectionExampleCreation.cpp
@@ -16,76 +16,74 @@ namespace GeometryCollectionExample
 {
 
 	template<class T>
-	bool CheckIncrementMask(ExampleResponse&& R)
+	void CheckIncrementMask()
 	{
 		{
 			TArray<int32> Mask;
 			GeometryCollectionAlgo::BuildIncrementMask({ 2 }, 5, Mask);
-			R.ExpectTrue(Mask[2] == 0);
-			R.ExpectTrue(Mask[3] == 1);
+			EXPECT_EQ(Mask[2], 0);
+			EXPECT_EQ(Mask[3], 1);
 		}
 		{
 			TArray<int32> Mask;
 			GeometryCollectionAlgo::BuildIncrementMask({ 0 }, 5, Mask);
-			R.ExpectTrue(Mask[0] == 0);
-			R.ExpectTrue(Mask[1] == 1);
+			EXPECT_EQ(Mask[0], 0);
+			EXPECT_EQ(Mask[1], 1);
 		}
 		{
 			TArray<int32> Mask;
 			GeometryCollectionAlgo::BuildIncrementMask({ 1,2 }, 5, Mask);
-			R.ExpectTrue(Mask[0] == 0);
-			R.ExpectTrue(Mask[1] == 0);
-			R.ExpectTrue(Mask[2] == 1);
-			R.ExpectTrue(Mask[3] == 2);
-			R.ExpectTrue(Mask[4] == 2);
+			EXPECT_EQ(Mask[0], 0);
+			EXPECT_EQ(Mask[1], 0);
+			EXPECT_EQ(Mask[2], 1);
+			EXPECT_EQ(Mask[3], 2);
+			EXPECT_EQ(Mask[4], 2);
 		}
-		return !R.HasError();
 	}
-	template bool CheckIncrementMask<float>(ExampleResponse&& R);
+	template void CheckIncrementMask<float>();
 
 	template<class T>
-	bool Creation(ExampleResponse&& R)
+	void Creation()
 	{
 		TSharedPtr<FGeometryCollection> Collection(new FGeometryCollection());
 
 		GeometryCollection::SetupCubeGridExample(Collection);
 
-		R.ExpectTrue(Collection->HasGroup(FTransformCollection::TransformGroup));
-		R.ExpectTrue(Collection->HasGroup(FGeometryCollection::VerticesGroup));
-		R.ExpectTrue(Collection->HasGroup(FGeometryCollection::FacesGroup));
-		R.ExpectTrue(Collection->HasGroup(FGeometryCollection::MaterialGroup));
-		R.ExpectTrue(Collection->HasGroup(FGeometryCollection::GeometryGroup));
+		EXPECT_TRUE(Collection->HasGroup(FTransformCollection::TransformGroup));
+		EXPECT_TRUE(Collection->HasGroup(FGeometryCollection::VerticesGroup));
+		EXPECT_TRUE(Collection->HasGroup(FGeometryCollection::FacesGroup));
+		EXPECT_TRUE(Collection->HasGroup(FGeometryCollection::MaterialGroup));
+		EXPECT_TRUE(Collection->HasGroup(FGeometryCollection::GeometryGroup));
 
-		R.ExpectTrue(Collection->NumElements(FTransformCollection::TransformGroup) == 1000);
-		R.ExpectTrue(Collection->NumElements(FGeometryCollection::VerticesGroup) == 8000);
-		R.ExpectTrue(Collection->NumElements(FGeometryCollection::FacesGroup) == 12000);
-		R.ExpectTrue(Collection->NumElements(FGeometryCollection::MaterialGroup) == 2);
-		R.ExpectTrue(Collection->NumElements(FGeometryCollection::GeometryGroup) == 1000);
+		EXPECT_EQ(Collection->NumElements(FTransformCollection::TransformGroup), 1000);
+		EXPECT_EQ(Collection->NumElements(FGeometryCollection::VerticesGroup), 8000);
+		EXPECT_EQ(Collection->NumElements(FGeometryCollection::FacesGroup), 12000);
+		EXPECT_EQ(Collection->NumElements(FGeometryCollection::MaterialGroup), 2);
+		EXPECT_EQ(Collection->NumElements(FGeometryCollection::GeometryGroup), 1000);
 
 		int HalfTheFaces = Collection->NumElements(FGeometryCollection::FacesGroup) / 2;
-		R.ExpectTrue((Collection->Sections)[0].MaterialID == 0);
-		R.ExpectTrue((Collection->Sections)[0].FirstIndex == 0);
-		R.ExpectTrue((Collection->Sections)[0].NumTriangles == HalfTheFaces);
-		R.ExpectTrue((Collection->Sections)[0].MinVertexIndex == 0);
-		R.ExpectTrue((Collection->Sections)[0].MaxVertexIndex == Collection->NumElements(FGeometryCollection::VerticesGroup) - 1);
-
-		R.ExpectTrue((Collection->Sections)[1].MaterialID == 1);
-		R.ExpectTrue((Collection->Sections)[1].FirstIndex == HalfTheFaces * 3);
-		R.ExpectTrue((Collection->Sections)[1].NumTriangles == HalfTheFaces);
-		R.ExpectTrue((Collection->Sections)[1].MinVertexIndex == 0);
-		R.ExpectTrue((Collection->Sections)[1].MaxVertexIndex == Collection->NumElements(FGeometryCollection::VerticesGroup) - 1);
-
-		R.ExpectTrue(GeometryCollectionAlgo::HasValidGeometryReferences(Collection.Get()));
-
-		R.ExpectTrue(Collection->HasContiguousFaces());
-		R.ExpectTrue(Collection->HasContiguousVertices());
-		R.ExpectTrue(Collection->HasContiguousRenderFaces());
-		return !R.HasError();
+		EXPECT_EQ((Collection->Sections)[0].MaterialID, 0);
+		EXPECT_EQ((Collection->Sections)[0].FirstIndex, 0);
+		EXPECT_EQ((Collection->Sections)[0].NumTriangles, HalfTheFaces);
+		EXPECT_EQ((Collection->Sections)[0].MinVertexIndex, 0);
+		EXPECT_EQ((Collection->Sections)[0].MaxVertexIndex, Collection->NumElements(FGeometryCollection::VerticesGroup) - 1);
+
+		EXPECT_EQ((Collection->Sections)[1].MaterialID, 1);
+		EXPECT_EQ((Collection->Sections)[1].FirstIndex, HalfTheFaces * 3);
+		EXPECT_EQ((Collection->Sections)[1].NumTriangles, HalfTheFaces);
+		EXPECT_EQ((Collection->Sections)[1].MinVertexIndex, 0);
+		EXPECT_EQ((Collection->Sections)[1].MaxVertexIndex, Collection->NumElements(FGeometryCollection::VerticesGroup) - 1);
+
+		EXPECT_TRUE(GeometryCollectionAlgo::HasValidGeometryReferences(Collection.Get()));
+
+		EXPECT_TRUE(Collection->HasContiguousFaces());
+		EXPECT_TRUE(Collection->HasContiguousVertices());
+		EXPECT_TRUE(Collection->HasContiguousRenderFaces());
 	}
-	template bool Creation<float>(ExampleResponse&& R);
+	template void Creation<float>();
 
 	template<class T>
-	bool AppendTransformHierarchy(ExampleResponse&& R)
+	void AppendTransformHierarchy()
 	{
 		TSharedPtr<FGeometryCollection> Collection = GeometryCollection::MakeCubeElement(FTransform(FQuat::MakeFromEuler(FVector(0, 0, 90.)), FVector(0, 10, 0)), FVector(1.0));
 		Collection->AppendGeometry(*GeometryCollection::MakeCubeElement(FTransform(FQuat::MakeFromEuler(FVector(0, 0, 90.)), FVector(0, 10, 0)), FVector(1.0)));
@@ -114,152 +112,150 @@ namespace GeometryCollectionExample
 		(Collection2->Children)[0].Add(2);
 		(Collection2->Parent)[2] = 0;
 
-		R.ExpectTrue(Collection->NumElements(FTransformCollection::TransformGroup) == 3);
-		R.ExpectTrue(Collection->NumElements(FGeometryCollection::VerticesGroup) == 24);
-		R.ExpectTrue(Collection->NumElements(FGeometryCollection::FacesGroup) == 36);
-		R.ExpectTrue(Collection->NumElements(FGeometryCollection::MaterialGroup) == 2);
-		R.ExpectTrue(Collection->NumElements(FGeometryCollection::GeometryGroup) == 3);
+		EXPECT_EQ(Collection->NumElements(FTransformCollection::TransformGroup), 3);
+		EXPECT_EQ(Collection->NumElements(FGeometryCollection::VerticesGroup), 24);
+		EXPECT_EQ(Collection->NumElements(FGeometryCollection::FacesGroup), 36);
+		EXPECT_EQ(Collection->NumElements(FGeometryCollection::MaterialGroup), 2);
+		EXPECT_EQ(Collection->NumElements(FGeometryCollection::GeometryGroup), 3);
 
-		R.ExpectTrue(Collection2->NumElements(FTransformCollection::TransformGroup) == 3);
-		R.ExpectTrue(Collection2->NumElements(FGeometryCollection::VerticesGroup) == 24);
-		R.ExpectTrue(Collection2->NumElements(FGeometryCollection::FacesGroup) == 36);
-		R.ExpectTrue(Collection2->NumElements(FGeometryCollection::MaterialGroup) == 4);
-		R.ExpectTrue(Collection2->NumElements(FGeometryCollection::GeometryGroup) == 3);
+		EXPECT_EQ(Collection2->NumElements(FTransformCollection::TransformGroup), 3);
+		EXPECT_EQ(Collection2->NumElements(FGeometryCollection::VerticesGroup), 24);
+		EXPECT_EQ(Collection2->NumElements(FGeometryCollection::FacesGroup), 36);
+		EXPECT_EQ(Collection2->NumElements(FGeometryCollection::MaterialGroup), 4);
+		EXPECT_EQ(Collection2->NumElements(FGeometryCollection::GeometryGroup), 3);
 
 		Collection->AppendGeometry(*Collection2);
 
-		R.ExpectTrue(Collection->NumElements(FTransformCollection::TransformGroup) == 6);
-		R.ExpectTrue(Collection->NumElements(FGeometryCollection::VerticesGroup) == 48);
-		R.ExpectTrue(Collection->NumElements(FGeometryCollection::FacesGroup) == 72);
-		R.ExpectTrue(Collection->NumElements(FGeometryCollection::MaterialGroup) == 4); // union of the 2/4 materials
-		R.ExpectTrue(Collection->NumElements(FGeometryCollection::GeometryGroup) == 6);
-
-		R.ExpectTrue(Collection->Parent[0] == -1);
-		R.ExpectTrue(Collection->Parent[1] == 0);
-		R.ExpectTrue(Collection->Parent[2] == 1);
-		R.ExpectTrue(Collection->Parent[3] == -1);
-		R.ExpectTrue(Collection->Parent[4] == 3);
-		R.ExpectTrue(Collection->Parent[5] == 3);
-
-		R.ExpectTrue(Collection->Children[0].Num() == 1);
-		R.ExpectTrue(Collection->Children[1].Num() == 1);
-		R.ExpectTrue(Collection->Children[2].Num() == 0);
-		R.ExpectTrue(Collection->Children[3].Num() == 2);
-		R.ExpectTrue(Collection->Children[4].Num() == 0);
-		R.ExpectTrue(Collection->Children[5].Num() == 0);
-
-		R.ExpectTrue(Collection->Children[0].Array()[0] == 1);
-		R.ExpectTrue(Collection->Children[1].Array()[0] == 2);
-		R.ExpectTrue(Collection->Children[3].Array()[0] == 4);
-		R.ExpectTrue(Collection->Children[3].Array()[1] == 5);
-
-		R.ExpectTrue((Collection->Sections)[0].MaterialID == 0);
-		R.ExpectTrue((Collection->Sections)[0].FirstIndex == 0);
-		R.ExpectTrue((Collection->Sections)[0].NumTriangles == 18 + 9);
-		R.ExpectTrue((Collection->Sections)[0].MinVertexIndex == 0);
-		R.ExpectTrue((Collection->Sections)[0].MaxVertexIndex == Collection->NumElements(FGeometryCollection::VerticesGroup) - 1);
-
-		R.ExpectTrue((Collection->Sections)[1].MaterialID == 1);
-		R.ExpectTrue((Collection->Sections)[1].FirstIndex == (Collection->Sections)[0].FirstIndex + (Collection->Sections)[0].NumTriangles * 3);
-		R.ExpectTrue((Collection->Sections)[1].NumTriangles == 18 + 9);
-		R.ExpectTrue((Collection->Sections)[1].MinVertexIndex == 0);
-		R.ExpectTrue((Collection->Sections)[1].MaxVertexIndex == Collection->NumElements(FGeometryCollection::VerticesGroup) - 1);
-
-		R.ExpectTrue((Collection->Sections)[2].MaterialID == 2);
-		R.ExpectTrue((Collection->Sections)[2].FirstIndex == (Collection->Sections)[1].FirstIndex + (Collection->Sections)[1].NumTriangles * 3);
-		R.ExpectTrue((Collection->Sections)[2].NumTriangles == 9);
-		R.ExpectTrue((Collection->Sections)[2].MinVertexIndex == 0);
-		R.ExpectTrue((Collection->Sections)[2].MaxVertexIndex == Collection->NumElements(FGeometryCollection::VerticesGroup) - 1);
-
-		R.ExpectTrue((Collection->Sections)[3].MaterialID == 3);
-		R.ExpectTrue((Collection->Sections)[3].FirstIndex == (Collection->Sections)[2].FirstIndex + (Collection->Sections)[2].NumTriangles * 3);
-		R.ExpectTrue((Collection->Sections)[3].NumTriangles == 9);
-		R.ExpectTrue((Collection->Sections)[3].MinVertexIndex == 0);
-		R.ExpectTrue((Collection->Sections)[3].MaxVertexIndex == Collection->NumElements(FGeometryCollection::VerticesGroup) - 1);
+		EXPECT_EQ(Collection->NumElements(FTransformCollection::TransformGroup), 6);
+		EXPECT_EQ(Collection->NumElements(FGeometryCollection::VerticesGroup), 48);
+		EXPECT_EQ(Collection->NumElements(FGeometryCollection::FacesGroup), 72);
+		EXPECT_EQ(Collection->NumElements(FGeometryCollection::MaterialGroup), 4); // union of the 2/4 materials
+		EXPECT_EQ(Collection->NumElements(FGeometryCollection::GeometryGroup), 6);
+
+		EXPECT_EQ(Collection->Parent[0], -1);
+		EXPECT_EQ(Collection->Parent[1], 0);
+		EXPECT_EQ(Collection->Parent[2], 1);
+		EXPECT_EQ(Collection->Parent[3], -1);
+		EXPECT_EQ(Collection->Parent[4], 3);
+		EXPECT_EQ(Collection->Parent[5], 3);
+
+		EXPECT_EQ(Collection->Children[0].Num(), 1);
+		EXPECT_EQ(Collection->Children[1].Num(), 1);
+		EXPECT_EQ(Collection->Children[2].Num(), 0);
+		EXPECT_EQ(Collection->Children[3].Num(), 2);
+		EXPECT_EQ(Collection->Children[4].Num(), 0);
+		EXPECT_EQ(Collection->Children[5].Num(), 0);
+
+		EXPECT_EQ(Collection->Children[0].Array()[0], 1);
+		EXPECT_EQ(Collection->Children[1].Array()[0], 2);
+		EXPECT_EQ(Collection->Children[3].Array()[0], 4);
+		EXPECT_EQ(Collection->Children[3].Array()[1], 5);
+
+		EXPECT_EQ((Collection->Sections)[0].MaterialID, 0);
+		EXPECT_EQ((Collection->Sections)[0].FirstIndex, 0);
+		EXPECT_EQ((Collection->Sections)[0].NumTriangles, 18 + 9);
+		EXPECT_EQ((Collection->Sections)[0].MinVertexIndex, 0);
+		EXPECT_EQ((Collection->Sections)[0].MaxVertexIndex, Collection->NumElements(FGeometryCollection::VerticesGroup) - 1);
+
+		EXPECT_EQ((Collection->Sections)[1].MaterialID, 1);
+		EXPECT_EQ((Collection->Sections)[1].FirstIndex, (Collection->Sections)[0].FirstIndex + (Collection->Sections)[0].NumTriangles * 3);
+		EXPECT_EQ((Collection->Sections)[1].NumTriangles, 18 + 9);
+		EXPECT_EQ((Collection->Sections)[1].MinVertexIndex, 0);
+		EXPECT_EQ((Collection->Sections)[1].MaxVertexIndex, Collection->NumElements(FGeometryCollection::VerticesGroup) - 1);
+
+		EXPECT_EQ((Collection->Sections)[2].MaterialID, 2);
+		EXPECT_EQ((Collection->Sections)[2].FirstIndex, (Collection->Sections)[1].FirstIndex + (Collection->Sections)[1].NumTriangles * 3);
+		EXPECT_EQ((Collection->Sections)[2].NumTriangles, 9);
+		EXPECT_EQ((Collection->Sections)[2].MinVertexIndex, 0);
+		EXPECT_EQ((Collection->Sections)[2].MaxVertexIndex, Collection->NumElements(FGeometryCollection::VerticesGroup) - 1);
+
+		EXPECT_EQ((Collection->Sections)[3].MaterialID, 3);
+		EXPECT_EQ((Collection->Sections)[3].FirstIndex, (Collection->Sections)[2].FirstIndex + (Collection->Sections)[2].NumTriangles * 3);
+		EXPECT_EQ((Collection->Sections)[3].NumTriangles, 9);
+		EXPECT_EQ((Collection->Sections)[3].MinVertexIndex, 0);
+		EXPECT_EQ((Collection->Sections)[3].MaxVertexIndex, Collection->NumElements(FGeometryCollection::VerticesGroup) - 1);
 
 		// GeometryGroup Tests
 		{
-			R.ExpectTrue(Collection->NumElements(FGeometryCollection::GeometryGroup) == 6);
-
-			R.ExpectTrue((Collection->TransformIndex)[0] == 0);
-			R.ExpectTrue((Collection->TransformIndex)[1] == 1);
-			R.ExpectTrue((Collection->TransformIndex)[2] == 2);
-			R.ExpectTrue((Collection->TransformIndex)[3] == 3);
-			R.ExpectTrue((Collection->TransformIndex)[4] == 4);
-			R.ExpectTrue((Collection->TransformIndex)[5] == 5);
-
-			R.ExpectTrue((Collection->TransformToGeometryIndex)[0] == 0);
-			R.ExpectTrue((Collection->TransformToGeometryIndex)[1] == 1);
-			R.ExpectTrue((Collection->TransformToGeometryIndex)[2] == 2);
-			R.ExpectTrue((Collection->TransformToGeometryIndex)[3] == 3);
-			R.ExpectTrue((Collection->TransformToGeometryIndex)[4] == 4);
-			R.ExpectTrue((Collection->TransformToGeometryIndex)[5] == 5);
-
-			R.ExpectTrue((Collection->FaceStart)[0] == 0);
-			R.ExpectTrue((Collection->FaceStart)[1] == 12);
-			R.ExpectTrue((Collection->FaceStart)[2] == 24);
-			R.ExpectTrue((Collection->FaceStart)[3] == 36);
-			R.ExpectTrue((Collection->FaceStart)[4] == 48);
-			R.ExpectTrue((Collection->FaceStart)[5] == 60);
-
-			R.ExpectTrue((Collection->FaceCount)[0] == 12);
-			R.ExpectTrue((Collection->FaceCount)[1] == 12);
-			R.ExpectTrue((Collection->FaceCount)[1] == 12);
-			R.ExpectTrue((Collection->FaceCount)[1] == 12);
-			R.ExpectTrue((Collection->FaceCount)[1] == 12);
-			R.ExpectTrue((Collection->FaceCount)[1] == 12);
-			R.ExpectTrue((Collection->Indices).Num() == 72);
-
-			R.ExpectTrue((Collection->VertexStart)[0] == 0);
-			R.ExpectTrue((Collection->VertexStart)[1] == 8);
-			R.ExpectTrue((Collection->VertexStart)[2] == 16);
-			R.ExpectTrue((Collection->VertexStart)[3] == 24);
-			R.ExpectTrue((Collection->VertexStart)[4] == 32);
-			R.ExpectTrue((Collection->VertexStart)[5] == 40);
-
-			R.ExpectTrue((Collection->VertexCount)[0] == 8);
-			R.ExpectTrue((Collection->VertexCount)[1] == 8);
-			R.ExpectTrue((Collection->VertexCount)[2] == 8);
-			R.ExpectTrue((Collection->VertexCount)[3] == 8);
-			R.ExpectTrue((Collection->VertexCount)[4] == 8);
-			R.ExpectTrue((Collection->VertexCount)[5] == 8);
-			R.ExpectTrue((Collection->Vertex).Num() == 48);
+			EXPECT_EQ(Collection->NumElements(FGeometryCollection::GeometryGroup), 6);
+
+			EXPECT_EQ((Collection->TransformIndex)[0], 0);
+			EXPECT_EQ((Collection->TransformIndex)[1], 1);
+			EXPECT_EQ((Collection->TransformIndex)[2], 2);
+			EXPECT_EQ((Collection->TransformIndex)[3], 3);
+			EXPECT_EQ((Collection->TransformIndex)[4], 4);
+			EXPECT_EQ((Collection->TransformIndex)[5], 5);
+
+			EXPECT_EQ((Collection->TransformToGeometryIndex)[0], 0);
+			EXPECT_EQ((Collection->TransformToGeometryIndex)[1], 1);
+			EXPECT_EQ((Collection->TransformToGeometryIndex)[2], 2);
+			EXPECT_EQ((Collection->TransformToGeometryIndex)[3], 3);
+			EXPECT_EQ((Collection->TransformToGeometryIndex)[4], 4);
+			EXPECT_EQ((Collection->TransformToGeometryIndex)[5], 5);
+
+			EXPECT_EQ((Collection->FaceStart)[0], 0);
+			EXPECT_EQ((Collection->FaceStart)[1], 12);
+			EXPECT_EQ((Collection->FaceStart)[2], 24);
+			EXPECT_EQ((Collection->FaceStart)[3], 36);
+			EXPECT_EQ((Collection->FaceStart)[4], 48);
+			EXPECT_EQ((Collection->FaceStart)[5], 60);
+
+			EXPECT_EQ((Collection->FaceCount)[0], 12);
+			EXPECT_EQ((Collection->FaceCount)[1], 12);
+			EXPECT_EQ((Collection->FaceCount)[1], 12);
+			EXPECT_EQ((Collection->FaceCount)[1], 12);
+			EXPECT_EQ((Collection->FaceCount)[1], 12);
+			EXPECT_EQ((Collection->FaceCount)[1], 12);
+			EXPECT_EQ((Collection->Indices).Num(), 72);
+
+			EXPECT_EQ((Collection->VertexStart)[0], 0);
+			EXPECT_EQ((Collection->VertexStart)[1], 8);
+			EXPECT_EQ((Collection->VertexStart)[2], 16);
+			EXPECT_EQ((Collection->VertexStart)[3], 24);
+			EXPECT_EQ((Collection->VertexStart)[4], 32);
+			EXPECT_EQ((Collection->VertexStart)[5], 40);
+
+			EXPECT_EQ((Collection->VertexCount)[0], 8);
+			EXPECT_EQ((Collection->VertexCount)[1], 8);
+			EXPECT_EQ((Collection->VertexCount)[2], 8);
+			EXPECT_EQ((Collection->VertexCount)[3], 8);
+			EXPECT_EQ((Collection->VertexCount)[4], 8);
+			EXPECT_EQ((Collection->VertexCount)[5], 8);
+			EXPECT_EQ((Collection->Vertex).Num(), 48);
 		}
 
-		R.ExpectTrue(Collection->HasContiguousFaces());
-		R.ExpectTrue(Collection->HasContiguousVertices());
-		R.ExpectTrue(Collection->HasContiguousRenderFaces());
-		R.ExpectTrue(GeometryCollectionAlgo::HasValidGeometryReferences(Collection.Get()));
-		return !R.HasError();
+		EXPECT_TRUE(Collection->HasContiguousFaces());
+		EXPECT_TRUE(Collection->HasContiguousVertices());
+		EXPECT_TRUE(Collection->HasContiguousRenderFaces());
+		EXPECT_TRUE(GeometryCollectionAlgo::HasValidGeometryReferences(Collection.Get()));
 	}
-	template bool AppendTransformHierarchy<float>(ExampleResponse&& R);
+	template void AppendTransformHierarchy<float>();
 
 	template<class T>
-	bool ContiguousElementsTest(ExampleResponse&& R)
+	void ContiguousElementsTest()
 	{
 		{
 			TSharedPtr<FGeometryCollection> Collection = GeometryCollection::MakeCubeElement(FTransform::Identity, FVector(1.0));
-			R.ExpectTrue(Collection->HasContiguousFaces());
-			R.ExpectTrue(Collection->HasContiguousVertices());
+			EXPECT_TRUE(Collection->HasContiguousFaces());
+			EXPECT_TRUE(Collection->HasContiguousVertices());
 			Collection->AppendGeometry(*GeometryCollection::MakeCubeElement(FTransform(FVector(0, 0, 10)), FVector(1.0)));
-			R.ExpectTrue(Collection->HasContiguousFaces());
-			R.ExpectTrue(Collection->HasContiguousVertices());
+			EXPECT_TRUE(Collection->HasContiguousFaces());
+			EXPECT_TRUE(Collection->HasContiguousVertices());
 			Collection->AppendGeometry(*GeometryCollection::MakeCubeElement(FTransform(FVector(0, 0, 10)), FVector(1.0)));
-			R.ExpectTrue(Collection->HasContiguousFaces());
-			R.ExpectTrue(Collection->HasContiguousVertices());
+			EXPECT_TRUE(Collection->HasContiguousFaces());
+			EXPECT_TRUE(Collection->HasContiguousVertices());
 		}
 		{
 			TSharedPtr<FGeometryCollection> Collection(new FGeometryCollection());
 			GeometryCollection::SetupCubeGridExample(Collection);
-			R.ExpectTrue(Collection->HasContiguousFaces());
-			R.ExpectTrue(Collection->HasContiguousVertices());
+			EXPECT_TRUE(Collection->HasContiguousFaces());
+			EXPECT_TRUE(Collection->HasContiguousVertices());
 		}
-		return !R.HasError();
 	}
-	template bool ContiguousElementsTest<float>(ExampleResponse&& R);
+	template void ContiguousElementsTest<float>();
 
 	template<class T>
-	bool DeleteFromEnd(ExampleResponse&& R)
+	void DeleteFromEnd()
 	{
 		TSharedPtr<FGeometryCollection> Collection = GeometryCollection::MakeCubeElement(FTransform::Identity, FVector(1.0));
 		Collection->AppendGeometry(*GeometryCollection::MakeCubeElement(FTransform(FVector(0, 0, 10)), FVector(1.0)));
@@ -273,110 +269,109 @@ namespace GeometryCollectionExample
 
 		(Collection->Parent)[2] = 1;
 
-		R.ExpectTrue(Collection->HasGroup(FTransformCollection::TransformGroup));
-		R.ExpectTrue(Collection->HasGroup(FGeometryCollection::VerticesGroup));
-		R.ExpectTrue(Collection->HasGroup(FGeometryCollection::FacesGroup));
-		R.ExpectTrue(Collection->HasGroup(FGeometryCollection::MaterialGroup));
-		R.ExpectTrue(Collection->HasGroup(FGeometryCollection::GeometryGroup));
+		EXPECT_TRUE(Collection->HasGroup(FTransformCollection::TransformGroup));
+		EXPECT_TRUE(Collection->HasGroup(FGeometryCollection::VerticesGroup));
+		EXPECT_TRUE(Collection->HasGroup(FGeometryCollection::FacesGroup));
+		EXPECT_TRUE(Collection->HasGroup(FGeometryCollection::MaterialGroup));
+		EXPECT_TRUE(Collection->HasGroup(FGeometryCollection::GeometryGroup));
 
-		R.ExpectTrue(Collection->NumElements(FTransformCollection::TransformGroup) == 3);
-		R.ExpectTrue(Collection->NumElements(FGeometryCollection::VerticesGroup) == 24);
-		R.ExpectTrue(Collection->NumElements(FGeometryCollection::FacesGroup) == 36);
-		R.ExpectTrue(Collection->NumElements(FGeometryCollection::MaterialGroup) == 2);
-		R.ExpectTrue(Collection->NumElements(FGeometryCollection::GeometryGroup) == 3);
+		EXPECT_EQ(Collection->NumElements(FTransformCollection::TransformGroup), 3);
+		EXPECT_EQ(Collection->NumElements(FGeometryCollection::VerticesGroup), 24);
+		EXPECT_EQ(Collection->NumElements(FGeometryCollection::FacesGroup), 36);
+		EXPECT_EQ(Collection->NumElements(FGeometryCollection::MaterialGroup), 2);
+		EXPECT_EQ(Collection->NumElements(FGeometryCollection::GeometryGroup), 3);
 
 		int HalfTheFaces = Collection->NumElements(FGeometryCollection::FacesGroup) / 2;
-		R.ExpectTrue((Collection->Sections)[0].MaterialID == 0);
-		R.ExpectTrue((Collection->Sections)[0].FirstIndex == 0);
-		R.ExpectTrue((Collection->Sections)[0].NumTriangles == HalfTheFaces);
-		R.ExpectTrue((Collection->Sections)[0].MinVertexIndex == 0);
-		R.ExpectTrue((Collection->Sections)[0].MaxVertexIndex == Collection->NumElements(FGeometryCollection::VerticesGroup) - 1);
+		EXPECT_EQ((Collection->Sections)[0].MaterialID, 0);
+		EXPECT_EQ((Collection->Sections)[0].FirstIndex, 0);
+		EXPECT_EQ((Collection->Sections)[0].NumTriangles, HalfTheFaces);
+		EXPECT_EQ((Collection->Sections)[0].MinVertexIndex, 0);
+		EXPECT_EQ((Collection->Sections)[0].MaxVertexIndex, Collection->NumElements(FGeometryCollection::VerticesGroup) - 1);
 
-		R.ExpectTrue((Collection->Sections)[1].MaterialID == 1);
-		R.ExpectTrue((Collection->Sections)[1].FirstIndex == HalfTheFaces * 3);
-		R.ExpectTrue((Collection->Sections)[1].NumTriangles == HalfTheFaces);
-		R.ExpectTrue((Collection->Sections)[1].MinVertexIndex == 0);
-		R.ExpectTrue((Collection->Sections)[1].MaxVertexIndex == Collection->NumElements(FGeometryCollection::VerticesGroup) - 1);
+		EXPECT_EQ((Collection->Sections)[1].MaterialID, 1);
+		EXPECT_EQ((Collection->Sections)[1].FirstIndex, HalfTheFaces * 3);
+		EXPECT_EQ((Collection->Sections)[1].NumTriangles, HalfTheFaces);
+		EXPECT_EQ((Collection->Sections)[1].MinVertexIndex, 0);
+		EXPECT_EQ((Collection->Sections)[1].MaxVertexIndex, Collection->NumElements(FGeometryCollection::VerticesGroup) - 1);
 
-		R.ExpectTrue(Collection->TransformToGeometryIndex.Num() == 3);
-		R.ExpectTrue(Collection->TransformToGeometryIndex[0] == 0);
-		R.ExpectTrue(Collection->TransformToGeometryIndex[1] == 1);
-		R.ExpectTrue(Collection->TransformToGeometryIndex[2] == 2);
+		EXPECT_EQ(Collection->TransformToGeometryIndex.Num(), 3);
+		EXPECT_EQ(Collection->TransformToGeometryIndex[0], 0);
+		EXPECT_EQ(Collection->TransformToGeometryIndex[1], 1);
+		EXPECT_EQ(Collection->TransformToGeometryIndex[2], 2);
 
-		R.ExpectTrue(GeometryCollectionAlgo::HasValidGeometryReferences(Collection.Get()));
+		EXPECT_TRUE(GeometryCollectionAlgo::HasValidGeometryReferences(Collection.Get()));
 
 		TArray<int32> DelList = { 2 };
 		Collection->RemoveElements(FTransformCollection::TransformGroup, DelList);
 
-		R.ExpectTrue(Collection->NumElements(FTransformCollection::TransformGroup) == 2);
-		R.ExpectTrue(Collection->NumElements(FGeometryCollection::VerticesGroup) == 16);
-		R.ExpectTrue(Collection->NumElements(FGeometryCollection::FacesGroup) == 24);
-		R.ExpectTrue(Collection->NumElements(FGeometryCollection::MaterialGroup) == 2);
+		EXPECT_EQ(Collection->NumElements(FTransformCollection::TransformGroup), 2);
+		EXPECT_EQ(Collection->NumElements(FGeometryCollection::VerticesGroup), 16);
+		EXPECT_EQ(Collection->NumElements(FGeometryCollection::FacesGroup), 24);
+		EXPECT_EQ(Collection->NumElements(FGeometryCollection::MaterialGroup), 2);
 
-		R.ExpectTrue(Collection->TransformToGeometryIndex.Num() == 2);
-		R.ExpectTrue(Collection->TransformToGeometryIndex[0] == 0);
-		R.ExpectTrue(Collection->TransformToGeometryIndex[1] == 1);
+		EXPECT_EQ(Collection->TransformToGeometryIndex.Num(), 2);
+		EXPECT_EQ(Collection->TransformToGeometryIndex[0], 0);
+		EXPECT_EQ(Collection->TransformToGeometryIndex[1], 1);
 
 		for (int Index = 0; Index < Collection->NumElements(FGeometryCollection::VerticesGroup); Index++)
 		{
-			R.ExpectTrue((Collection->BoneMap)[Index] < Collection->NumElements(FGeometryCollection::TransformGroup));
+			EXPECT_TRUE((Collection->BoneMap)[Index] < Collection->NumElements(FGeometryCollection::TransformGroup));
 		}
 		for (int Index = 0; Index < Collection->NumElements(FGeometryCollection::FacesGroup); Index++)
 		{
-			R.ExpectTrue((Collection->Indices)[Index][0] < Collection->NumElements(FGeometryCollection::VerticesGroup));
-			R.ExpectTrue((Collection->Indices)[Index][1] < Collection->NumElements(FGeometryCollection::VerticesGroup));
-			R.ExpectTrue((Collection->Indices)[Index][2] < Collection->NumElements(FGeometryCollection::VerticesGroup));
+			EXPECT_TRUE((Collection->Indices)[Index][0] < Collection->NumElements(FGeometryCollection::VerticesGroup));
+			EXPECT_TRUE((Collection->Indices)[Index][1] < Collection->NumElements(FGeometryCollection::VerticesGroup));
+			EXPECT_TRUE((Collection->Indices)[Index][2] < Collection->NumElements(FGeometryCollection::VerticesGroup));
 		}
 
-		R.ExpectTrue((Collection->Transform)[0].GetTranslation().Z == 0.f);
-		R.ExpectTrue((Collection->Transform)[1].GetTranslation().Z == 10.f);
+		EXPECT_EQ((Collection->Transform)[0].GetTranslation().Z, 0.f);
+		EXPECT_EQ((Collection->Transform)[1].GetTranslation().Z, 10.f);
 
 		HalfTheFaces = Collection->NumElements(FGeometryCollection::FacesGroup) / 2;
-		R.ExpectTrue((Collection->Sections)[0].MaterialID == 0);
-		R.ExpectTrue((Collection->Sections)[0].FirstIndex == 0);
-		R.ExpectTrue((Collection->Sections)[0].NumTriangles == HalfTheFaces);
-		R.ExpectTrue((Collection->Sections)[0].MinVertexIndex == 0);
-		R.ExpectTrue((Collection->Sections)[0].MaxVertexIndex == Collection->NumElements(FGeometryCollection::VerticesGroup) - 1);
-
-		R.ExpectTrue((Collection->Sections)[1].MaterialID == 1);
-		R.ExpectTrue((Collection->Sections)[1].FirstIndex == HalfTheFaces * 3);
-		R.ExpectTrue((Collection->Sections)[1].NumTriangles == HalfTheFaces);
-		R.ExpectTrue((Collection->Sections)[1].MinVertexIndex == 0);
-		R.ExpectTrue((Collection->Sections)[1].MaxVertexIndex == Collection->NumElements(FGeometryCollection::VerticesGroup) - 1);
+		EXPECT_EQ((Collection->Sections)[0].MaterialID, 0);
+		EXPECT_EQ((Collection->Sections)[0].FirstIndex, 0);
+		EXPECT_EQ((Collection->Sections)[0].NumTriangles, HalfTheFaces);
+		EXPECT_EQ((Collection->Sections)[0].MinVertexIndex, 0);
+		EXPECT_EQ((Collection->Sections)[0].MaxVertexIndex, Collection->NumElements(FGeometryCollection::VerticesGroup) - 1);
+
+		EXPECT_EQ((Collection->Sections)[1].MaterialID, 1);
+		EXPECT_EQ((Collection->Sections)[1].FirstIndex, HalfTheFaces * 3);
+		EXPECT_EQ((Collection->Sections)[1].NumTriangles, HalfTheFaces);
+		EXPECT_EQ((Collection->Sections)[1].MinVertexIndex, 0);
+		EXPECT_EQ((Collection->Sections)[1].MaxVertexIndex, Collection->NumElements(FGeometryCollection::VerticesGroup) - 1);
 
 		// GeometryGroup Updated Tests
 		{
-			R.ExpectTrue(Collection->NumElements(FGeometryCollection::GeometryGroup) == 2);
+			EXPECT_EQ(Collection->NumElements(FGeometryCollection::GeometryGroup), 2);
 
-			R.ExpectTrue((Collection->TransformIndex)[0] == 0);
-			R.ExpectTrue((Collection->TransformIndex)[1] == 1);
+			EXPECT_EQ((Collection->TransformIndex)[0], 0);
+			EXPECT_EQ((Collection->TransformIndex)[1], 1);
 
-			R.ExpectTrue((Collection->FaceStart)[0] == 0);
-			R.ExpectTrue((Collection->FaceStart)[1] == 12);
+			EXPECT_EQ((Collection->FaceStart)[0], 0);
+			EXPECT_EQ((Collection->FaceStart)[1], 12);
 
-			R.ExpectTrue((Collection->FaceCount)[0] == 12);
-			R.ExpectTrue((Collection->FaceCount)[1] == 12);
-			R.ExpectTrue((Collection->Indices).Num() == 24);
+			EXPECT_EQ((Collection->FaceCount)[0], 12);
+			EXPECT_EQ((Collection->FaceCount)[1], 12);
+			EXPECT_EQ((Collection->Indices).Num(), 24);
 
-			R.ExpectTrue((Collection->VertexStart)[0] == 0);
-			R.ExpectTrue((Collection->VertexStart)[1] == 8);
+			EXPECT_EQ((Collection->VertexStart)[0], 0);
+			EXPECT_EQ((Collection->VertexStart)[1], 8);
 
-			R.ExpectTrue((Collection->VertexCount)[0] == 8);
-			R.ExpectTrue((Collection->VertexCount)[1] == 8);
-			R.ExpectTrue((Collection->Vertex).Num() == 16);
+			EXPECT_EQ((Collection->VertexCount)[0], 8);
+			EXPECT_EQ((Collection->VertexCount)[1], 8);
+			EXPECT_EQ((Collection->Vertex).Num(), 16);
 		}
 
-		R.ExpectTrue(GeometryCollectionAlgo::HasValidGeometryReferences(Collection.Get()));
-		R.ExpectTrue(Collection->HasContiguousFaces());
-		R.ExpectTrue(Collection->HasContiguousVertices());
-		R.ExpectTrue(Collection->HasContiguousRenderFaces());
-		return !R.HasError();
+		EXPECT_TRUE(GeometryCollectionAlgo::HasValidGeometryReferences(Collection.Get()));
+		EXPECT_TRUE(Collection->HasContiguousFaces());
+		EXPECT_TRUE(Collection->HasContiguousVertices());
+		EXPECT_TRUE(Collection->HasContiguousRenderFaces());
 	}
-	template bool DeleteFromEnd<float>(ExampleResponse&& R);
+	template void DeleteFromEnd<float>();
 
 
 	template<class T>
-	bool DeleteFromStart(ExampleResponse&& R)
+	void DeleteFromStart()
 	{
 		TSharedPtr<FGeometryCollection> Collection = GeometryCollection::MakeCubeElement(FTransform::Identity, FVector(1.0));
 		Collection->AppendGeometry(*GeometryCollection::MakeCubeElement(FTransform(FVector(0, 0, 10)), FVector(1.0)));
@@ -390,109 +385,108 @@ namespace GeometryCollectionExample
 
 		(Collection->Parent)[2] = 1;
 
-		R.ExpectTrue(Collection->HasGroup(FTransformCollection::TransformGroup));
-		R.ExpectTrue(Collection->HasGroup(FGeometryCollection::VerticesGroup));
-		R.ExpectTrue(Collection->HasGroup(FGeometryCollection::FacesGroup));
-		R.ExpectTrue(Collection->HasGroup(FGeometryCollection::MaterialGroup));
-		R.ExpectTrue(Collection->HasGroup(FGeometryCollection::GeometryGroup));
+		EXPECT_TRUE(Collection->HasGroup(FTransformCollection::TransformGroup));
+		EXPECT_TRUE(Collection->HasGroup(FGeometryCollection::VerticesGroup));
+		EXPECT_TRUE(Collection->HasGroup(FGeometryCollection::FacesGroup));
+		EXPECT_TRUE(Collection->HasGroup(FGeometryCollection::MaterialGroup));
+		EXPECT_TRUE(Collection->HasGroup(FGeometryCollection::GeometryGroup));
 
-		R.ExpectTrue(Collection->NumElements(FTransformCollection::TransformGroup) == 3);
-		R.ExpectTrue(Collection->NumElements(FGeometryCollection::VerticesGroup) == 24);
-		R.ExpectTrue(Collection->NumElements(FGeometryCollection::FacesGroup) == 36);
-		R.ExpectTrue(Collection->NumElements(FGeometryCollection::MaterialGroup) == 2);
-		R.ExpectTrue(Collection->NumElements(FGeometryCollection::GeometryGroup) == 3);
+		EXPECT_EQ(Collection->NumElements(FTransformCollection::TransformGroup), 3);
+		EXPECT_EQ(Collection->NumElements(FGeometryCollection::VerticesGroup), 24);
+		EXPECT_EQ(Collection->NumElements(FGeometryCollection::FacesGroup), 36);
+		EXPECT_EQ(Collection->NumElements(FGeometryCollection::MaterialGroup), 2);
+		EXPECT_EQ(Collection->NumElements(FGeometryCollection::GeometryGroup), 3);
 
 		int HalfTheFaces = Collection->NumElements(FGeometryCollection::FacesGroup) / 2;
-		R.ExpectTrue((Collection->Sections)[0].MaterialID == 0);
-		R.ExpectTrue((Collection->Sections)[0].FirstIndex == 0);
-		R.ExpectTrue((Collection->Sections)[0].NumTriangles == HalfTheFaces);
-		R.ExpectTrue((Collection->Sections)[0].MinVertexIndex == 0);
-		R.ExpectTrue((Collection->Sections)[0].MaxVertexIndex == Collection->NumElements(FGeometryCollection::VerticesGroup) - 1);
+		EXPECT_EQ((Collection->Sections)[0].MaterialID, 0);
+		EXPECT_EQ((Collection->Sections)[0].FirstIndex, 0);
+		EXPECT_EQ((Collection->Sections)[0].NumTriangles, HalfTheFaces);
+		EXPECT_EQ((Collection->Sections)[0].MinVertexIndex, 0);
+		EXPECT_EQ((Collection->Sections)[0].MaxVertexIndex, Collection->NumElements(FGeometryCollection::VerticesGroup) - 1);
 
-		R.ExpectTrue((Collection->Sections)[1].MaterialID == 1);
-		R.ExpectTrue((Collection->Sections)[1].FirstIndex == HalfTheFaces * 3);
-		R.ExpectTrue((Collection->Sections)[1].NumTriangles == HalfTheFaces);
-		R.ExpectTrue((Collection->Sections)[1].MinVertexIndex == 0);
-		R.ExpectTrue((Collection->Sections)[1].MaxVertexIndex == Collection->NumElements(FGeometryCollection::VerticesGroup) - 1);
+		EXPECT_EQ((Collection->Sections)[1].MaterialID, 1);
+		EXPECT_EQ((Collection->Sections)[1].FirstIndex, HalfTheFaces * 3);
+		EXPECT_EQ((Collection->Sections)[1].NumTriangles, HalfTheFaces);
+		EXPECT_EQ((Collection->Sections)[1].MinVertexIndex, 0);
+		EXPECT_EQ((Collection->Sections)[1].MaxVertexIndex, Collection->NumElements(FGeometryCollection::VerticesGroup) - 1);
 
-		R.ExpectTrue(Collection->TransformToGeometryIndex.Num() == 3);
-		R.ExpectTrue(Collection->TransformToGeometryIndex[0] == 0);
-		R.ExpectTrue(Collection->TransformToGeometryIndex[1] == 1);
-		R.ExpectTrue(Collection->TransformToGeometryIndex[2] == 2);
+		EXPECT_EQ(Collection->TransformToGeometryIndex.Num(), 3);
+		EXPECT_EQ(Collection->TransformToGeometryIndex[0], 0);
+		EXPECT_EQ(Collection->TransformToGeometryIndex[1], 1);
+		EXPECT_EQ(Collection->TransformToGeometryIndex[2], 2);
 
-		R.ExpectTrue(GeometryCollectionAlgo::HasValidGeometryReferences(Collection.Get()));
+		EXPECT_TRUE(GeometryCollectionAlgo::HasValidGeometryReferences(Collection.Get()));
 
 		TArray<int32> DelList = { 0 };
 		Collection->RemoveElements(FTransformCollection::TransformGroup, DelList);
 
-		R.ExpectTrue(Collection->NumElements(FTransformCollection::TransformGroup) == 2);
-		R.ExpectTrue(Collection->NumElements(FGeometryCollection::VerticesGroup) == 16);
-		R.ExpectTrue(Collection->NumElements(FGeometryCollection::FacesGroup) == 24);
-		R.ExpectTrue(Collection->NumElements(FGeometryCollection::MaterialGroup) == 2);
+		EXPECT_EQ(Collection->NumElements(FTransformCollection::TransformGroup), 2);
+		EXPECT_EQ(Collection->NumElements(FGeometryCollection::VerticesGroup), 16);
+		EXPECT_EQ(Collection->NumElements(FGeometryCollection::FacesGroup), 24);
+		EXPECT_EQ(Collection->NumElements(FGeometryCollection::MaterialGroup), 2);
 
-		R.ExpectTrue(Collection->TransformToGeometryIndex.Num() == 2);
-		R.ExpectTrue(Collection->TransformToGeometryIndex[0] == 0);
-		R.ExpectTrue(Collection->TransformToGeometryIndex[1] == 1);
+		EXPECT_EQ(Collection->TransformToGeometryIndex.Num(), 2);
+		EXPECT_EQ(Collection->TransformToGeometryIndex[0], 0);
+		EXPECT_EQ(Collection->TransformToGeometryIndex[1], 1);
 
 		for (int Index = 0; Index < Collection->NumElements(FGeometryCollection::VerticesGroup); Index++)
 		{
-			R.ExpectTrue((Collection->BoneMap)[Index] < Collection->NumElements(FGeometryCollection::TransformGroup));
+			EXPECT_LT((Collection->BoneMap)[Index], Collection->NumElements(FGeometryCollection::TransformGroup));
 		}
 		for (int Index = 0; Index < Collection->NumElements(FGeometryCollection::FacesGroup); Index++)
 		{
-			R.ExpectTrue((Collection->Indices)[Index][0] < Collection->NumElements(FGeometryCollection::VerticesGroup));
-			R.ExpectTrue((Collection->Indices)[Index][1] < Collection->NumElements(FGeometryCollection::VerticesGroup));
-			R.ExpectTrue((Collection->Indices)[Index][2] < Collection->NumElements(FGeometryCollection::VerticesGroup));
+			EXPECT_LT((Collection->Indices)[Index][0], Collection->NumElements(FGeometryCollection::VerticesGroup));
+			EXPECT_LT((Collection->Indices)[Index][1], Collection->NumElements(FGeometryCollection::VerticesGroup));
+			EXPECT_LT((Collection->Indices)[Index][2], Collection->NumElements(FGeometryCollection::VerticesGroup));
 		}
 
-		R.ExpectTrue((Collection->Transform)[0].GetTranslation().Z == 10.f);
-		R.ExpectTrue((Collection->Transform)[1].GetTranslation().Z == 20.f);
+		EXPECT_EQ((Collection->Transform)[0].GetTranslation().Z, 10.f);
+		EXPECT_EQ((Collection->Transform)[1].GetTranslation().Z, 20.f);
 
 		HalfTheFaces = Collection->NumElements(FGeometryCollection::FacesGroup) / 2;
-		R.ExpectTrue((Collection->Sections)[0].MaterialID == 0);
-		R.ExpectTrue((Collection->Sections)[0].FirstIndex == 0);
-		R.ExpectTrue((Collection->Sections)[0].NumTriangles == HalfTheFaces);
-		R.ExpectTrue((Collection->Sections)[0].MinVertexIndex == 0);
-		R.ExpectTrue((Collection->Sections)[0].MaxVertexIndex == Collection->NumElements(FGeometryCollection::VerticesGroup) - 1);
-
-		R.ExpectTrue((Collection->Sections)[1].MaterialID == 1);
-		R.ExpectTrue((Collection->Sections)[1].FirstIndex == HalfTheFaces * 3);
-		R.ExpectTrue((Collection->Sections)[1].NumTriangles == HalfTheFaces);
-		R.ExpectTrue((Collection->Sections)[1].MinVertexIndex == 0);
-		R.ExpectTrue((Collection->Sections)[1].MaxVertexIndex == Collection->NumElements(FGeometryCollection::VerticesGroup) - 1);
+		EXPECT_EQ((Collection->Sections)[0].MaterialID, 0);
+		EXPECT_EQ((Collection->Sections)[0].FirstIndex, 0);
+		EXPECT_EQ((Collection->Sections)[0].NumTriangles, HalfTheFaces);
+		EXPECT_EQ((Collection->Sections)[0].MinVertexIndex, 0);
+		EXPECT_EQ((Collection->Sections)[0].MaxVertexIndex, Collection->NumElements(FGeometryCollection::VerticesGroup) - 1);
+
+		EXPECT_EQ((Collection->Sections)[1].MaterialID, 1);
+		EXPECT_EQ((Collection->Sections)[1].FirstIndex, HalfTheFaces * 3);
+		EXPECT_EQ((Collection->Sections)[1].NumTriangles, HalfTheFaces);
+		EXPECT_EQ((Collection->Sections)[1].MinVertexIndex, 0);
+		EXPECT_EQ((Collection->Sections)[1].MaxVertexIndex, Collection->NumElements(FGeometryCollection::VerticesGroup) - 1);
 
 		// GeometryGroup Updated Tests
 		{
-			R.ExpectTrue(Collection->NumElements(FGeometryCollection::GeometryGroup) == 2);
+			EXPECT_EQ(Collection->NumElements(FGeometryCollection::GeometryGroup), 2);
 
-			R.ExpectTrue((Collection->TransformIndex)[0] == 0);
-			R.ExpectTrue((Collection->TransformIndex)[1] == 1);
+			EXPECT_EQ((Collection->TransformIndex)[0], 0);
+			EXPECT_EQ((Collection->TransformIndex)[1], 1);
 
-			R.ExpectTrue((Collection->FaceStart)[0] == 0);
-			R.ExpectTrue((Collection->FaceStart)[1] == 12);
+			EXPECT_EQ((Collection->FaceStart)[0], 0);
+			EXPECT_EQ((Collection->FaceStart)[1], 12);
 
-			R.ExpectTrue((Collection->FaceCount)[0] == 12);
-			R.ExpectTrue((Collection->FaceCount)[1] == 12);
-			R.ExpectTrue((Collection->Indices).Num() == 24);
+			EXPECT_EQ((Collection->FaceCount)[0], 12);
+			EXPECT_EQ((Collection->FaceCount)[1], 12);
+			EXPECT_EQ((Collection->Indices).Num(), 24);
 
-			R.ExpectTrue((Collection->VertexStart)[0] == 0);
-			R.ExpectTrue((Collection->VertexStart)[1] == 8);
+			EXPECT_EQ((Collection->VertexStart)[0], 0);
+			EXPECT_EQ((Collection->VertexStart)[1], 8);
 
-			R.ExpectTrue((Collection->VertexCount)[0] == 8);
-			R.ExpectTrue((Collection->VertexCount)[1] == 8);
-			R.ExpectTrue((Collection->Vertex).Num() == 16);
+			EXPECT_EQ((Collection->VertexCount)[0], 8);
+			EXPECT_EQ((Collection->VertexCount)[1], 8);
+			EXPECT_EQ((Collection->Vertex).Num(), 16);
 		}
 
-		R.ExpectTrue(GeometryCollectionAlgo::HasValidGeometryReferences(Collection.Get()));
-		R.ExpectTrue(Collection->HasContiguousFaces());
-		R.ExpectTrue(Collection->HasContiguousVertices());
-		R.ExpectTrue(Collection->HasContiguousRenderFaces());
-		return !R.HasError();
+		EXPECT_TRUE(GeometryCollectionAlgo::HasValidGeometryReferences(Collection.Get()));
+		EXPECT_TRUE(Collection->HasContiguousFaces());
+		EXPECT_TRUE(Collection->HasContiguousVertices());
+		EXPECT_TRUE(Collection->HasContiguousRenderFaces());
 	}
-	template bool DeleteFromStart<float>(ExampleResponse&& R);
+	template void DeleteFromStart<float>();
 
 	template<class T>
-	bool DeleteFromMiddle(ExampleResponse&& R)
+	void DeleteFromMiddle()
 	{
 		TSharedPtr<FGeometryCollection> Collection = GeometryCollection::MakeCubeElement(FTransform::Identity, FVector(1.0));
 		Collection->AppendGeometry(*GeometryCollection::MakeCubeElement(FTransform(FVector(0, 0, 10)), FVector(1.0)));
@@ -506,111 +500,109 @@ namespace GeometryCollectionExample
 
 		(Collection->Parent)[2] = 1;
 
-		R.ExpectTrue(Collection->HasGroup(FTransformCollection::TransformGroup));
-		R.ExpectTrue(Collection->HasGroup(FGeometryCollection::VerticesGroup));
-		R.ExpectTrue(Collection->HasGroup(FGeometryCollection::FacesGroup));
-		R.ExpectTrue(Collection->HasGroup(FGeometryCollection::MaterialGroup));
-		R.ExpectTrue(Collection->HasGroup(FGeometryCollection::GeometryGroup));
+		EXPECT_TRUE(Collection->HasGroup(FTransformCollection::TransformGroup));
+		EXPECT_TRUE(Collection->HasGroup(FGeometryCollection::VerticesGroup));
+		EXPECT_TRUE(Collection->HasGroup(FGeometryCollection::FacesGroup));
+		EXPECT_TRUE(Collection->HasGroup(FGeometryCollection::MaterialGroup));
+		EXPECT_TRUE(Collection->HasGroup(FGeometryCollection::GeometryGroup));
 
-		R.ExpectTrue(Collection->NumElements(FTransformCollection::TransformGroup) == 3);
-		R.ExpectTrue(Collection->NumElements(FGeometryCollection::VerticesGroup) == 24);
-		R.ExpectTrue(Collection->NumElements(FGeometryCollection::FacesGroup) == 36);
-		R.ExpectTrue(Collection->NumElements(FGeometryCollection::MaterialGroup) == 2);
-		R.ExpectTrue(Collection->NumElements(FGeometryCollection::GeometryGroup) == 3);
+		EXPECT_EQ(Collection->NumElements(FTransformCollection::TransformGroup), 3);
+		EXPECT_EQ(Collection->NumElements(FGeometryCollection::VerticesGroup), 24);
+		EXPECT_EQ(Collection->NumElements(FGeometryCollection::FacesGroup), 36);
+		EXPECT_EQ(Collection->NumElements(FGeometryCollection::MaterialGroup), 2);
+		EXPECT_EQ(Collection->NumElements(FGeometryCollection::GeometryGroup), 3);
 
 		int HalfTheFaces = Collection->NumElements(FGeometryCollection::FacesGroup) / 2;
-		R.ExpectTrue((Collection->Sections)[0].MaterialID == 0);
-		R.ExpectTrue((Collection->Sections)[0].FirstIndex == 0);
-		R.ExpectTrue((Collection->Sections)[0].NumTriangles == HalfTheFaces);
-		R.ExpectTrue((Collection->Sections)[0].MinVertexIndex == 0);
-		R.ExpectTrue((Collection->Sections)[0].MaxVertexIndex == Collection->NumElements(FGeometryCollection::VerticesGroup) - 1);
+		EXPECT_EQ((Collection->Sections)[0].MaterialID, 0);
+		EXPECT_EQ((Collection->Sections)[0].FirstIndex, 0);
+		EXPECT_EQ((Collection->Sections)[0].NumTriangles, HalfTheFaces);
+		EXPECT_EQ((Collection->Sections)[0].MinVertexIndex, 0);
+		EXPECT_EQ((Collection->Sections)[0].MaxVertexIndex, Collection->NumElements(FGeometryCollection::VerticesGroup) - 1);
 
-		R.ExpectTrue((Collection->Sections)[1].MaterialID == 1);
-		R.ExpectTrue((Collection->Sections)[1].FirstIndex == HalfTheFaces * 3);
-		R.ExpectTrue((Collection->Sections)[1].NumTriangles == HalfTheFaces);
-		R.ExpectTrue((Collection->Sections)[1].MinVertexIndex == 0);
-		R.ExpectTrue((Collection->Sections)[1].MaxVertexIndex == Collection->NumElements(FGeometryCollection::VerticesGroup) - 1);
+		EXPECT_EQ((Collection->Sections)[1].MaterialID, 1);
+		EXPECT_EQ((Collection->Sections)[1].FirstIndex, HalfTheFaces * 3);
+		EXPECT_EQ((Collection->Sections)[1].NumTriangles, HalfTheFaces);
+		EXPECT_EQ((Collection->Sections)[1].MinVertexIndex, 0);
+		EXPECT_EQ((Collection->Sections)[1].MaxVertexIndex, Collection->NumElements(FGeometryCollection::VerticesGroup) - 1);
 
-		R.ExpectTrue(Collection->TransformToGeometryIndex.Num() == 3);
-		R.ExpectTrue(Collection->TransformToGeometryIndex[0] == 0);
-		R.ExpectTrue(Collection->TransformToGeometryIndex[1] == 1);
-		R.ExpectTrue(Collection->TransformToGeometryIndex[2] == 2);
+		EXPECT_EQ(Collection->TransformToGeometryIndex.Num(), 3);
+		EXPECT_EQ(Collection->TransformToGeometryIndex[0], 0);
+		EXPECT_EQ(Collection->TransformToGeometryIndex[1], 1);
+		EXPECT_EQ(Collection->TransformToGeometryIndex[2], 2);
 
-		R.ExpectTrue(GeometryCollectionAlgo::HasValidGeometryReferences(Collection.Get()));
+		EXPECT_TRUE(GeometryCollectionAlgo::HasValidGeometryReferences(Collection.Get()));
 
 		TArray<int32> DelList = { 1 };
 		Collection->RemoveElements(FTransformCollection::TransformGroup, DelList);
 
-		R.ExpectTrue(Collection->NumElements(FTransformCollection::TransformGroup) == 2);
-		R.ExpectTrue(Collection->NumElements(FGeometryCollection::VerticesGroup) == 16);
-		R.ExpectTrue(Collection->NumElements(FGeometryCollection::FacesGroup) == 24);
-		R.ExpectTrue(Collection->NumElements(FGeometryCollection::MaterialGroup) == 2);
+		EXPECT_EQ(Collection->NumElements(FTransformCollection::TransformGroup), 2);
+		EXPECT_EQ(Collection->NumElements(FGeometryCollection::VerticesGroup), 16);
+		EXPECT_EQ(Collection->NumElements(FGeometryCollection::FacesGroup), 24);
+		EXPECT_EQ(Collection->NumElements(FGeometryCollection::MaterialGroup), 2);
 
-		R.ExpectTrue(Collection->TransformToGeometryIndex.Num() == 2);
-		R.ExpectTrue(Collection->TransformToGeometryIndex[0] == 0);
-		R.ExpectTrue(Collection->TransformToGeometryIndex[1] == 1);
+		EXPECT_EQ(Collection->TransformToGeometryIndex.Num(), 2);
+		EXPECT_EQ(Collection->TransformToGeometryIndex[0], 0);
+		EXPECT_EQ(Collection->TransformToGeometryIndex[1], 1);
 
 		for (int Index = 0; Index < Collection->NumElements(FGeometryCollection::VerticesGroup); Index++)
 		{
-			R.ExpectTrue((Collection->BoneMap)[Index] < Collection->NumElements(FGeometryCollection::TransformGroup));
+			EXPECT_LT((Collection->BoneMap)[Index], Collection->NumElements(FGeometryCollection::TransformGroup));
 		}
 		for (int Index = 0; Index < Collection->NumElements(FGeometryCollection::FacesGroup); Index++)
 		{
-			R.ExpectTrue((Collection->Indices)[Index][0] < Collection->NumElements(FGeometryCollection::VerticesGroup));
-			R.ExpectTrue((Collection->Indices)[Index][1] < Collection->NumElements(FGeometryCollection::VerticesGroup));
-			R.ExpectTrue((Collection->Indices)[Index][2] < Collection->NumElements(FGeometryCollection::VerticesGroup));
+			EXPECT_LT((Collection->Indices)[Index][0], Collection->NumElements(FGeometryCollection::VerticesGroup));
+			EXPECT_LT((Collection->Indices)[Index][1], Collection->NumElements(FGeometryCollection::VerticesGroup));
+			EXPECT_LT((Collection->Indices)[Index][2], Collection->NumElements(FGeometryCollection::VerticesGroup));
 		}
 
-		R.ExpectTrue((Collection->Transform)[0].GetTranslation().Z == 0.f);
-		R.ExpectTrue((Collection->Transform)[1].GetTranslation().Z == 30.f);
+		EXPECT_EQ((Collection->Transform)[0].GetTranslation().Z, 0.f);
+		EXPECT_EQ((Collection->Transform)[1].GetTranslation().Z, 30.f);
 
 		HalfTheFaces = Collection->NumElements(FGeometryCollection::FacesGroup) / 2;
-		R.ExpectTrue((Collection->Sections)[0].MaterialID == 0);
-		R.ExpectTrue((Collection->Sections)[0].FirstIndex == 0);
-		R.ExpectTrue((Collection->Sections)[0].NumTriangles == HalfTheFaces);
-		R.ExpectTrue((Collection->Sections)[0].MinVertexIndex == 0);
-		R.ExpectTrue((Collection->Sections)[0].MaxVertexIndex == Collection->NumElements(FGeometryCollection::VerticesGroup) - 1);
-
-		R.ExpectTrue((Collection->Sections)[1].MaterialID == 1);
-		R.ExpectTrue((Collection->Sections)[1].FirstIndex == HalfTheFaces * 3);
-		R.ExpectTrue((Collection->Sections)[1].NumTriangles == HalfTheFaces);
-		R.ExpectTrue((Collection->Sections)[1].MinVertexIndex == 0);
-		R.ExpectTrue((Collection->Sections)[1].MaxVertexIndex == Collection->NumElements(FGeometryCollection::VerticesGroup) - 1);
+		EXPECT_EQ((Collection->Sections)[0].MaterialID, 0);
+		EXPECT_EQ((Collection->Sections)[0].FirstIndex, 0);
+		EXPECT_EQ((Collection->Sections)[0].NumTriangles, HalfTheFaces);
+		EXPECT_EQ((Collection->Sections)[0].MinVertexIndex, 0);
+		EXPECT_EQ((Collection->Sections)[0].MaxVertexIndex, Collection->NumElements(FGeometryCollection::VerticesGroup) - 1);
+
+		EXPECT_EQ((Collection->Sections)[1].MaterialID, 1);
+		EXPECT_EQ((Collection->Sections)[1].FirstIndex, HalfTheFaces * 3);
+		EXPECT_EQ((Collection->Sections)[1].NumTriangles, HalfTheFaces);
+		EXPECT_EQ((Collection->Sections)[1].MinVertexIndex, 0);
+		EXPECT_EQ((Collection->Sections)[1].MaxVertexIndex, Collection->NumElements(FGeometryCollection::VerticesGroup) - 1);
 
 		// GeometryGroup Updated Tests
 		{
-			R.ExpectTrue(Collection->NumElements(FGeometryCollection::GeometryGroup) == 2);
+			EXPECT_EQ(Collection->NumElements(FGeometryCollection::GeometryGroup), 2);
 
-			R.ExpectTrue((Collection->TransformIndex)[0] == 0);
-			R.ExpectTrue((Collection->TransformIndex)[1] == 1);
+			EXPECT_EQ((Collection->TransformIndex)[0], 0);
+			EXPECT_EQ((Collection->TransformIndex)[1], 1);
 
-			R.ExpectTrue((Collection->FaceStart)[0] == 0);
-			R.ExpectTrue((Collection->FaceStart)[1] == 12);
+			EXPECT_EQ((Collection->FaceStart)[0], 0);
+			EXPECT_EQ((Collection->FaceStart)[1], 12);
 
-			R.ExpectTrue((Collection->FaceCount)[0] == 12);
-			R.ExpectTrue((Collection->FaceCount)[1] == 12);
-			R.ExpectTrue((Collection->Indices).Num() == 24);
+			EXPECT_EQ((Collection->FaceCount)[0], 12);
+			EXPECT_EQ((Collection->FaceCount)[1], 12);
+			EXPECT_EQ((Collection->Indices).Num(), 24);
 
-			R.ExpectTrue((Collection->VertexStart)[0] == 0);
-			R.ExpectTrue((Collection->VertexStart)[1] == 8);
+			EXPECT_EQ((Collection->VertexStart)[0], 0);
+			EXPECT_EQ((Collection->VertexStart)[1], 8);
 
-			R.ExpectTrue((Collection->VertexCount)[0] == 8);
-			R.ExpectTrue((Collection->VertexCount)[1] == 8);
-			R.ExpectTrue((Collection->Vertex).Num() == 16);
+			EXPECT_EQ((Collection->VertexCount)[0], 8);
+			EXPECT_EQ((Collection->VertexCount)[1], 8);
+			EXPECT_EQ((Collection->Vertex).Num(), 16);
 		}
 
-		R.ExpectTrue(GeometryCollectionAlgo::HasValidGeometryReferences(Collection.Get()));
-		R.ExpectTrue(Collection->HasContiguousFaces());
-		R.ExpectTrue(Collection->HasContiguousVertices());
-		R.ExpectTrue(Collection->HasContiguousRenderFaces());
-
-		return !R.HasError();
+		EXPECT_TRUE(GeometryCollectionAlgo::HasValidGeometryReferences(Collection.Get()));
+		EXPECT_TRUE(Collection->HasContiguousFaces());
+		EXPECT_TRUE(Collection->HasContiguousVertices());
+		EXPECT_TRUE(Collection->HasContiguousRenderFaces());
 	}
-	template bool DeleteFromMiddle<float>(ExampleResponse&& R);
+	template void DeleteFromMiddle<float>();
 
 
 	template<class T>
-	bool DeleteBranch(ExampleResponse&& R)
+	void DeleteBranch()
 	{
 		TSharedPtr<FGeometryCollection> Collection = GeometryCollection::MakeCubeElement(FTransform::Identity, FVector(1.0));
 		Collection->AppendGeometry(*GeometryCollection::MakeCubeElement(FTransform(FVector(0, 0, 10)), FVector(1.0)));
@@ -633,38 +625,38 @@ namespace GeometryCollectionExample
 		(Collection->Parent)[3] = 1;
 		(Collection->Parent)[4] = 2;
 
-		R.ExpectTrue(Collection->HasGroup(FTransformCollection::TransformGroup));
-		R.ExpectTrue(Collection->HasGroup(FGeometryCollection::VerticesGroup));
-		R.ExpectTrue(Collection->HasGroup(FGeometryCollection::FacesGroup));
-		R.ExpectTrue(Collection->HasGroup(FGeometryCollection::MaterialGroup));
+		EXPECT_TRUE(Collection->HasGroup(FTransformCollection::TransformGroup));
+		EXPECT_TRUE(Collection->HasGroup(FGeometryCollection::VerticesGroup));
+		EXPECT_TRUE(Collection->HasGroup(FGeometryCollection::FacesGroup));
+		EXPECT_TRUE(Collection->HasGroup(FGeometryCollection::MaterialGroup));
 
-		R.ExpectTrue(Collection->NumElements(FTransformCollection::TransformGroup) == 5);
-		R.ExpectTrue(Collection->NumElements(FGeometryCollection::VerticesGroup) == 40);
-		R.ExpectTrue(Collection->NumElements(FGeometryCollection::FacesGroup) == 60);
-		R.ExpectTrue(Collection->NumElements(FGeometryCollection::MaterialGroup) == 2);
-		R.ExpectTrue(Collection->NumElements(FGeometryCollection::GeometryGroup) == 5);
+		EXPECT_EQ(Collection->NumElements(FTransformCollection::TransformGroup), 5);
+		EXPECT_EQ(Collection->NumElements(FGeometryCollection::VerticesGroup), 40);
+		EXPECT_EQ(Collection->NumElements(FGeometryCollection::FacesGroup), 60);
+		EXPECT_EQ(Collection->NumElements(FGeometryCollection::MaterialGroup), 2);
+		EXPECT_EQ(Collection->NumElements(FGeometryCollection::GeometryGroup), 5);
 
 		int HalfTheFaces = Collection->NumElements(FGeometryCollection::FacesGroup) / 2;
-		R.ExpectTrue((Collection->Sections)[0].MaterialID == 0);
-		R.ExpectTrue((Collection->Sections)[0].FirstIndex == 0);
-		R.ExpectTrue((Collection->Sections)[0].NumTriangles == HalfTheFaces);
-		R.ExpectTrue((Collection->Sections)[0].MinVertexIndex == 0);
-		R.ExpectTrue((Collection->Sections)[0].MaxVertexIndex == Collection->NumElements(FGeometryCollection::VerticesGroup) - 1);
-
-		R.ExpectTrue((Collection->Sections)[1].MaterialID == 1);
-		R.ExpectTrue((Collection->Sections)[1].FirstIndex == HalfTheFaces * 3);
-		R.ExpectTrue((Collection->Sections)[1].NumTriangles == HalfTheFaces);
-		R.ExpectTrue((Collection->Sections)[1].MinVertexIndex == 0);
-		R.ExpectTrue((Collection->Sections)[1].MaxVertexIndex == Collection->NumElements(FGeometryCollection::VerticesGroup) - 1);
-
-		R.ExpectTrue(Collection->TransformToGeometryIndex.Num() == 5);
-		R.ExpectTrue(Collection->TransformToGeometryIndex[0] == 0);
-		R.ExpectTrue(Collection->TransformToGeometryIndex[1] == 1);
-		R.ExpectTrue(Collection->TransformToGeometryIndex[2] == 2);
-		R.ExpectTrue(Collection->TransformToGeometryIndex[3] == 3);
-		R.ExpectTrue(Collection->TransformToGeometryIndex[4] == 4);
-
-		R.ExpectTrue(GeometryCollectionAlgo::HasValidGeometryReferences(Collection.Get()));
+		EXPECT_EQ((Collection->Sections)[0].MaterialID, 0);
+		EXPECT_EQ((Collection->Sections)[0].FirstIndex, 0);
+		EXPECT_EQ((Collection->Sections)[0].NumTriangles, HalfTheFaces);
+		EXPECT_EQ((Collection->Sections)[0].MinVertexIndex, 0);
+		EXPECT_EQ((Collection->Sections)[0].MaxVertexIndex, Collection->NumElements(FGeometryCollection::VerticesGroup) - 1);
+
+		EXPECT_EQ((Collection->Sections)[1].MaterialID, 1);
+		EXPECT_EQ((Collection->Sections)[1].FirstIndex, HalfTheFaces * 3);
+		EXPECT_EQ((Collection->Sections)[1].NumTriangles, HalfTheFaces);
+		EXPECT_EQ((Collection->Sections)[1].MinVertexIndex, 0);
+		EXPECT_EQ((Collection->Sections)[1].MaxVertexIndex, Collection->NumElements(FGeometryCollection::VerticesGroup) - 1);
+
+		EXPECT_EQ(Collection->TransformToGeometryIndex.Num(), 5);
+		EXPECT_EQ(Collection->TransformToGeometryIndex[0], 0);
+		EXPECT_EQ(Collection->TransformToGeometryIndex[1], 1);
+		EXPECT_EQ(Collection->TransformToGeometryIndex[2], 2);
+		EXPECT_EQ(Collection->TransformToGeometryIndex[3], 3);
+		EXPECT_EQ(Collection->TransformToGeometryIndex[4], 4);
+
+		EXPECT_TRUE(GeometryCollectionAlgo::HasValidGeometryReferences(Collection.Get()));
 
 		//  0
 		//  ...2
@@ -672,91 +664,89 @@ namespace GeometryCollectionExample
 		TArray<int32> DelList = { 1, 3 };
 		Collection->RemoveElements(FTransformCollection::TransformGroup, DelList);
 
-		R.ExpectTrue(Collection->NumElements(FTransformCollection::TransformGroup) == 3);
-		R.ExpectTrue(Collection->NumElements(FGeometryCollection::VerticesGroup) == 24);
-		R.ExpectTrue(Collection->NumElements(FGeometryCollection::FacesGroup) == 36);
-		R.ExpectTrue(Collection->NumElements(FGeometryCollection::MaterialGroup) == 2);
-
-		R.ExpectTrue((Collection->Parent)[0] == -1);
-		R.ExpectTrue((Collection->Children)[0].Num() == 1);
-		R.ExpectTrue((Collection->Children)[0].Contains(1));
-		R.ExpectTrue((Collection->Parent)[1] == 0);
-		R.ExpectTrue((Collection->Children)[1].Num() == 1);
-		R.ExpectTrue((Collection->Children)[1].Contains(2));
-		R.ExpectTrue((Collection->Parent)[2] == 1);
-		R.ExpectTrue((Collection->Children)[2].Num() == 0);
-
-		R.ExpectTrue(Collection->TransformToGeometryIndex.Num() == 3);
-		R.ExpectTrue(Collection->TransformToGeometryIndex[0] == 0);
-		R.ExpectTrue(Collection->TransformToGeometryIndex[1] == 1);
-		R.ExpectTrue(Collection->TransformToGeometryIndex[2] == 2);
+		EXPECT_EQ(Collection->NumElements(FTransformCollection::TransformGroup), 3);
+		EXPECT_EQ(Collection->NumElements(FGeometryCollection::VerticesGroup), 24);
+		EXPECT_EQ(Collection->NumElements(FGeometryCollection::FacesGroup), 36);
+		EXPECT_EQ(Collection->NumElements(FGeometryCollection::MaterialGroup), 2);
+
+		EXPECT_EQ((Collection->Parent)[0], -1);
+		EXPECT_EQ((Collection->Children)[0].Num(), 1);
+		EXPECT_TRUE((Collection->Children)[0].Contains(1));
+		EXPECT_EQ((Collection->Parent)[1], 0);
+		EXPECT_EQ((Collection->Children)[1].Num(), 1);
+		EXPECT_TRUE((Collection->Children)[1].Contains(2));
+		EXPECT_EQ((Collection->Parent)[2], 1);
+		EXPECT_EQ((Collection->Children)[2].Num(), 0);
+
+		EXPECT_EQ(Collection->TransformToGeometryIndex.Num(), 3);
+		EXPECT_EQ(Collection->TransformToGeometryIndex[0], 0);
+		EXPECT_EQ(Collection->TransformToGeometryIndex[1], 1);
+		EXPECT_EQ(Collection->TransformToGeometryIndex[2], 2);
 
 		for (int Index = 0; Index < Collection->NumElements(FGeometryCollection::VerticesGroup); Index++)
 		{
-			R.ExpectTrue((Collection->BoneMap)[Index] < Collection->NumElements(FGeometryCollection::TransformGroup));
+			EXPECT_LT((Collection->BoneMap)[Index], Collection->NumElements(FGeometryCollection::TransformGroup));
 		}
 		for (int Index = 0; Index < Collection->NumElements(FGeometryCollection::FacesGroup); Index++)
 		{
-			R.ExpectTrue((Collection->Indices)[Index][0] < Collection->NumElements(FGeometryCollection::VerticesGroup));
-			R.ExpectTrue((Collection->Indices)[Index][1] < Collection->NumElements(FGeometryCollection::VerticesGroup));
-			R.ExpectTrue((Collection->Indices)[Index][2] < Collection->NumElements(FGeometryCollection::VerticesGroup));
+			EXPECT_LT((Collection->Indices)[Index][0], Collection->NumElements(FGeometryCollection::VerticesGroup));
+			EXPECT_LT((Collection->Indices)[Index][1], Collection->NumElements(FGeometryCollection::VerticesGroup));
+			EXPECT_LT((Collection->Indices)[Index][2], Collection->NumElements(FGeometryCollection::VerticesGroup));
 		}
 
-		R.ExpectTrue((Collection->Transform)[0].GetTranslation().Z == 0.f);
-		R.ExpectTrue((Collection->Transform)[1].GetTranslation().Z == 10.f);
+		EXPECT_EQ((Collection->Transform)[0].GetTranslation().Z, 0.f);
+		EXPECT_EQ((Collection->Transform)[1].GetTranslation().Z, 10.f);
 
 		HalfTheFaces = Collection->NumElements(FGeometryCollection::FacesGroup) / 2;
-		R.ExpectTrue((Collection->Sections)[0].MaterialID == 0);
-		R.ExpectTrue((Collection->Sections)[0].FirstIndex == 0);
-		R.ExpectTrue((Collection->Sections)[0].NumTriangles == HalfTheFaces);
-		R.ExpectTrue((Collection->Sections)[0].MinVertexIndex == 0);
-		R.ExpectTrue((Collection->Sections)[0].MaxVertexIndex == Collection->NumElements(FGeometryCollection::VerticesGroup) - 1);
-
-		R.ExpectTrue((Collection->Sections)[1].MaterialID == 1);
-		R.ExpectTrue((Collection->Sections)[1].FirstIndex == HalfTheFaces * 3);
-		R.ExpectTrue((Collection->Sections)[1].NumTriangles == HalfTheFaces);
-		R.ExpectTrue((Collection->Sections)[1].MinVertexIndex == 0);
-		R.ExpectTrue((Collection->Sections)[1].MaxVertexIndex == Collection->NumElements(FGeometryCollection::VerticesGroup) - 1);
+		EXPECT_EQ((Collection->Sections)[0].MaterialID, 0);
+		EXPECT_EQ((Collection->Sections)[0].FirstIndex, 0);
+		EXPECT_EQ((Collection->Sections)[0].NumTriangles, HalfTheFaces);
+		EXPECT_EQ((Collection->Sections)[0].MinVertexIndex, 0);
+		EXPECT_EQ((Collection->Sections)[0].MaxVertexIndex, Collection->NumElements(FGeometryCollection::VerticesGroup) - 1);
+
+		EXPECT_EQ((Collection->Sections)[1].MaterialID, 1);
+		EXPECT_EQ((Collection->Sections)[1].FirstIndex, HalfTheFaces * 3);
+		EXPECT_EQ((Collection->Sections)[1].NumTriangles, HalfTheFaces);
+		EXPECT_EQ((Collection->Sections)[1].MinVertexIndex, 0);
+		EXPECT_EQ((Collection->Sections)[1].MaxVertexIndex, Collection->NumElements(FGeometryCollection::VerticesGroup) - 1);
 
 		// GeometryGroup Updated Tests
 		{
-			R.ExpectTrue(Collection->NumElements(FGeometryCollection::GeometryGroup) == 3);
+			EXPECT_EQ(Collection->NumElements(FGeometryCollection::GeometryGroup), 3);
 
-			R.ExpectTrue((Collection->TransformIndex)[0] == 0);
-			R.ExpectTrue((Collection->TransformIndex)[1] == 1);
-			R.ExpectTrue((Collection->TransformIndex)[2] == 2);
+			EXPECT_EQ((Collection->TransformIndex)[0], 0);
+			EXPECT_EQ((Collection->TransformIndex)[1], 1);
+			EXPECT_EQ((Collection->TransformIndex)[2], 2);
 
-			R.ExpectTrue((Collection->FaceStart)[0] == 0);
-			R.ExpectTrue((Collection->FaceStart)[1] == 12);
-			R.ExpectTrue((Collection->FaceStart)[2] == 24);
+			EXPECT_EQ((Collection->FaceStart)[0], 0);
+			EXPECT_EQ((Collection->FaceStart)[1], 12);
+			EXPECT_EQ((Collection->FaceStart)[2], 24);
 
-			R.ExpectTrue((Collection->FaceCount)[0] == 12);
-			R.ExpectTrue((Collection->FaceCount)[1] == 12);
-			R.ExpectTrue((Collection->FaceCount)[2] == 12);
-			R.ExpectTrue((Collection->Indices).Num() == 36);
+			EXPECT_EQ((Collection->FaceCount)[0], 12);
+			EXPECT_EQ((Collection->FaceCount)[1], 12);
+			EXPECT_EQ((Collection->FaceCount)[2], 12);
+			EXPECT_EQ((Collection->Indices).Num(), 36);
 
-			R.ExpectTrue((Collection->VertexStart)[0] == 0);
-			R.ExpectTrue((Collection->VertexStart)[1] == 8);
-			R.ExpectTrue((Collection->VertexStart)[2] == 16);
+			EXPECT_EQ((Collection->VertexStart)[0], 0);
+			EXPECT_EQ((Collection->VertexStart)[1], 8);
+			EXPECT_EQ((Collection->VertexStart)[2], 16);
 
-			R.ExpectTrue((Collection->VertexCount)[0] == 8);
-			R.ExpectTrue((Collection->VertexCount)[1] == 8);
-			R.ExpectTrue((Collection->VertexCount)[2] == 8);
-			R.ExpectTrue((Collection->Vertex).Num() == 24);
+			EXPECT_EQ((Collection->VertexCount)[0], 8);
+			EXPECT_EQ((Collection->VertexCount)[1], 8);
+			EXPECT_EQ((Collection->VertexCount)[2], 8);
+			EXPECT_EQ((Collection->Vertex).Num(), 24);
 		}
 
-		R.ExpectTrue(GeometryCollectionAlgo::HasValidGeometryReferences(Collection.Get()));
-		R.ExpectTrue(Collection->HasContiguousFaces());
-		R.ExpectTrue(Collection->HasContiguousVertices());
-		R.ExpectTrue(Collection->HasContiguousRenderFaces());
-
-		return !R.HasError();
+		EXPECT_TRUE(GeometryCollectionAlgo::HasValidGeometryReferences(Collection.Get()));
+		EXPECT_TRUE(Collection->HasContiguousFaces());
+		EXPECT_TRUE(Collection->HasContiguousVertices());
+		EXPECT_TRUE(Collection->HasContiguousRenderFaces());
 	}
-	template bool DeleteBranch<float>(ExampleResponse&& R);
+	template void DeleteBranch<float>();
 
 
 	template<class T>
-	bool DeleteRootLeafMiddle(ExampleResponse&& R)
+	void DeleteRootLeafMiddle()
 	{
 		TSharedPtr<FGeometryCollection> Collection = GeometryCollection::MakeCubeElement(FTransform::Identity, FVector(1.0));
 		Collection->AppendGeometry(*GeometryCollection::MakeCubeElement(FTransform(FVector(0, 0, 10)), FVector(1.0)));
@@ -801,42 +791,42 @@ namespace GeometryCollectionExample
 		(Collection->BoneName)[7] = "7";
 
 
-		R.ExpectTrue(Collection->HasGroup(FTransformCollection::TransformGroup));
-		R.ExpectTrue(Collection->HasGroup(FGeometryCollection::VerticesGroup));
-		R.ExpectTrue(Collection->HasGroup(FGeometryCollection::FacesGroup));
-		R.ExpectTrue(Collection->HasGroup(FGeometryCollection::MaterialGroup));
-		R.ExpectTrue(Collection->HasGroup(FGeometryCollection::GeometryGroup));
+		EXPECT_TRUE(Collection->HasGroup(FTransformCollection::TransformGroup));
+		EXPECT_TRUE(Collection->HasGroup(FGeometryCollection::VerticesGroup));
+		EXPECT_TRUE(Collection->HasGroup(FGeometryCollection::FacesGroup));
+		EXPECT_TRUE(Collection->HasGroup(FGeometryCollection::MaterialGroup));
+		EXPECT_TRUE(Collection->HasGroup(FGeometryCollection::GeometryGroup));
 
-		R.ExpectTrue(Collection->NumElements(FTransformCollection::TransformGroup) == 8);
-		R.ExpectTrue(Collection->NumElements(FGeometryCollection::VerticesGroup) == 64);
-		R.ExpectTrue(Collection->NumElements(FGeometryCollection::FacesGroup) == 96);
-		R.ExpectTrue(Collection->NumElements(FGeometryCollection::MaterialGroup) == 2);
-		R.ExpectTrue(Collection->NumElements(FGeometryCollection::GeometryGroup) == 8);
+		EXPECT_EQ(Collection->NumElements(FTransformCollection::TransformGroup), 8);
+		EXPECT_EQ(Collection->NumElements(FGeometryCollection::VerticesGroup), 64);
+		EXPECT_EQ(Collection->NumElements(FGeometryCollection::FacesGroup), 96);
+		EXPECT_EQ(Collection->NumElements(FGeometryCollection::MaterialGroup), 2);
+		EXPECT_EQ(Collection->NumElements(FGeometryCollection::GeometryGroup), 8);
 
 		int HalfTheFaces = Collection->NumElements(FGeometryCollection::FacesGroup) / 2;
-		R.ExpectTrue((Collection->Sections)[0].MaterialID == 0);
-		R.ExpectTrue((Collection->Sections)[0].FirstIndex == 0);
-		R.ExpectTrue((Collection->Sections)[0].NumTriangles == HalfTheFaces);
-		R.ExpectTrue((Collection->Sections)[0].MinVertexIndex == 0);
-		R.ExpectTrue((Collection->Sections)[0].MaxVertexIndex == Collection->NumElements(FGeometryCollection::VerticesGroup) - 1);
-
-		R.ExpectTrue((Collection->Sections)[1].MaterialID == 1);
-		R.ExpectTrue((Collection->Sections)[1].FirstIndex == HalfTheFaces * 3);
-		R.ExpectTrue((Collection->Sections)[1].NumTriangles == HalfTheFaces);
-		R.ExpectTrue((Collection->Sections)[1].MinVertexIndex == 0);
-		R.ExpectTrue((Collection->Sections)[1].MaxVertexIndex == Collection->NumElements(FGeometryCollection::VerticesGroup) - 1);
-
-		R.ExpectTrue(Collection->TransformToGeometryIndex.Num() == 8);
-		R.ExpectTrue(Collection->TransformToGeometryIndex[0] == 0);
-		R.ExpectTrue(Collection->TransformToGeometryIndex[1] == 1);
-		R.ExpectTrue(Collection->TransformToGeometryIndex[2] == 2);
-		R.ExpectTrue(Collection->TransformToGeometryIndex[3] == 3);
-		R.ExpectTrue(Collection->TransformToGeometryIndex[4] == 4);
-		R.ExpectTrue(Collection->TransformToGeometryIndex[5] == 5);
-		R.ExpectTrue(Collection->TransformToGeometryIndex[6] == 6);
-		R.ExpectTrue(Collection->TransformToGeometryIndex[7] == 7);
-
-		R.ExpectTrue(GeometryCollectionAlgo::HasValidGeometryReferences(Collection.Get()));
+		EXPECT_EQ((Collection->Sections)[0].MaterialID, 0);
+		EXPECT_EQ((Collection->Sections)[0].FirstIndex, 0);
+		EXPECT_EQ((Collection->Sections)[0].NumTriangles, HalfTheFaces);
+		EXPECT_EQ((Collection->Sections)[0].MinVertexIndex, 0);
+		EXPECT_EQ((Collection->Sections)[0].MaxVertexIndex, Collection->NumElements(FGeometryCollection::VerticesGroup) - 1);
+
+		EXPECT_EQ((Collection->Sections)[1].MaterialID, 1);
+		EXPECT_EQ((Collection->Sections)[1].FirstIndex, HalfTheFaces * 3);
+		EXPECT_EQ((Collection->Sections)[1].NumTriangles, HalfTheFaces);
+		EXPECT_EQ((Collection->Sections)[1].MinVertexIndex, 0);
+		EXPECT_EQ((Collection->Sections)[1].MaxVertexIndex, Collection->NumElements(FGeometryCollection::VerticesGroup) - 1);
+
+		EXPECT_EQ(Collection->TransformToGeometryIndex.Num(), 8);
+		EXPECT_EQ(Collection->TransformToGeometryIndex[0], 0);
+		EXPECT_EQ(Collection->TransformToGeometryIndex[1], 1);
+		EXPECT_EQ(Collection->TransformToGeometryIndex[2], 2);
+		EXPECT_EQ(Collection->TransformToGeometryIndex[3], 3);
+		EXPECT_EQ(Collection->TransformToGeometryIndex[4], 4);
+		EXPECT_EQ(Collection->TransformToGeometryIndex[5], 5);
+		EXPECT_EQ(Collection->TransformToGeometryIndex[6], 6);
+		EXPECT_EQ(Collection->TransformToGeometryIndex[7], 7);
+
+		EXPECT_TRUE(GeometryCollectionAlgo::HasValidGeometryReferences(Collection.Get()));
 
 		//  1
 		//  6
@@ -846,30 +836,30 @@ namespace GeometryCollectionExample
 		TArray<int32> DelList = { 0,5,7 };
 		Collection->RemoveElements(FTransformCollection::TransformGroup, DelList);
 
-		R.ExpectTrue(Collection->NumElements(FTransformCollection::TransformGroup) == 5);
-		R.ExpectTrue(Collection->NumElements(FGeometryCollection::VerticesGroup) == 40);
-		R.ExpectTrue(Collection->NumElements(FGeometryCollection::FacesGroup) == 60);
-		R.ExpectTrue(Collection->NumElements(FGeometryCollection::MaterialGroup) == 2);
-		R.ExpectTrue(Collection->NumElements(FGeometryCollection::GeometryGroup) == 5);
-
-		R.ExpectTrue((Collection->Parent)[0] == -1);
-		R.ExpectTrue((Collection->Children)[0].Num() == 0);
-		R.ExpectTrue((Collection->Parent)[1] == -1);
-		R.ExpectTrue((Collection->Children)[1].Num() == 1);
-		R.ExpectTrue((Collection->Children)[1].Contains(3));
-		R.ExpectTrue((Collection->Parent)[2] == -1);
-		R.ExpectTrue((Collection->Children)[2].Num() == 0);
-		R.ExpectTrue((Collection->Parent)[3] == 1);
-		R.ExpectTrue((Collection->Children)[3].Num() == 0);
-		R.ExpectTrue((Collection->Parent)[4] == -1);
-		R.ExpectTrue((Collection->Children)[4].Num() == 0);
-
-		R.ExpectTrue(Collection->TransformToGeometryIndex.Num() == 5);
-		R.ExpectTrue(Collection->TransformToGeometryIndex[0] == 0);
-		R.ExpectTrue(Collection->TransformToGeometryIndex[1] == 1);
-		R.ExpectTrue(Collection->TransformToGeometryIndex[2] == 2);
-		R.ExpectTrue(Collection->TransformToGeometryIndex[3] == 3);
-		R.ExpectTrue(Collection->TransformToGeometryIndex[4] == 4);
+		EXPECT_EQ(Collection->NumElements(FTransformCollection::TransformGroup), 5);
+		EXPECT_EQ(Collection->NumElements(FGeometryCollection::VerticesGroup), 40);
+		EXPECT_EQ(Collection->NumElements(FGeometryCollection::FacesGroup), 60);
+		EXPECT_EQ(Collection->NumElements(FGeometryCollection::MaterialGroup), 2);
+		EXPECT_EQ(Collection->NumElements(FGeometryCollection::GeometryGroup), 5);
+
+		EXPECT_EQ((Collection->Parent)[0], -1);
+		EXPECT_EQ((Collection->Children)[0].Num(), 0);
+		EXPECT_EQ((Collection->Parent)[1], -1);
+		EXPECT_EQ((Collection->Children)[1].Num(), 1);
+		EXPECT_TRUE((Collection->Children)[1].Contains(3));
+		EXPECT_EQ((Collection->Parent)[2], -1);
+		EXPECT_EQ((Collection->Children)[2].Num(), 0);
+		EXPECT_EQ((Collection->Parent)[3], 1);
+		EXPECT_EQ((Collection->Children)[3].Num(), 0);
+		EXPECT_EQ((Collection->Parent)[4], -1);
+		EXPECT_EQ((Collection->Children)[4].Num(), 0);
+
+		EXPECT_EQ(Collection->TransformToGeometryIndex.Num(), 5);
+		EXPECT_EQ(Collection->TransformToGeometryIndex[0], 0);
+		EXPECT_EQ(Collection->TransformToGeometryIndex[1], 1);
+		EXPECT_EQ(Collection->TransformToGeometryIndex[2], 2);
+		EXPECT_EQ(Collection->TransformToGeometryIndex[3], 3);
+		EXPECT_EQ(Collection->TransformToGeometryIndex[4], 4);
 
 		int32 Index0 = Collection->BoneName.Find("0");
 		int32 Index1 = Collection->BoneName.Find("1");
@@ -878,97 +868,94 @@ namespace GeometryCollectionExample
 		int32 Index4 = Collection->BoneName.Find("4");
 		int32 Index6 = Collection->BoneName.Find("6");
 
-		R.ExpectTrue(Index0 == -1);
-		R.ExpectTrue(Index6 != -1);
-		R.ExpectTrue((Collection->Parent)[Index1] == -1);
-		R.ExpectTrue((Collection->Parent)[Index2] == -1);
-		R.ExpectTrue((Collection->Children)[Index2].Num() == 1);
-		R.ExpectTrue((Collection->Children)[Index2].Contains(Index4));
-		R.ExpectTrue((Collection->Parent)[Index4] == Index2);
-		R.ExpectTrue((Collection->Children)[Index4].Num() == 0);
+		EXPECT_EQ(Index0, -1);
+		EXPECT_NE(Index6, -1);
+		EXPECT_EQ((Collection->Parent)[Index1], -1);
+		EXPECT_EQ((Collection->Parent)[Index2], -1);
+		EXPECT_EQ((Collection->Children)[Index2].Num(), 1);
+		EXPECT_TRUE((Collection->Children)[Index2].Contains(Index4));
+		EXPECT_EQ((Collection->Parent)[Index4], Index2);
+		EXPECT_EQ((Collection->Children)[Index4].Num(), 0);
 
 		for (int Index = 0; Index < Collection->NumElements(FGeometryCollection::VerticesGroup); Index++)
 		{
-			R.ExpectTrue((Collection->BoneMap)[Index] < Collection->NumElements(FGeometryCollection::TransformGroup));
+			EXPECT_LT((Collection->BoneMap)[Index], Collection->NumElements(FGeometryCollection::TransformGroup));
 		}
 		for (int Index = 0; Index < Collection->NumElements(FGeometryCollection::FacesGroup); Index++)
 		{
-			R.ExpectTrue((Collection->Indices)[Index][0] < Collection->NumElements(FGeometryCollection::VerticesGroup));
-			R.ExpectTrue((Collection->Indices)[Index][1] < Collection->NumElements(FGeometryCollection::VerticesGroup));
-			R.ExpectTrue((Collection->Indices)[Index][2] < Collection->NumElements(FGeometryCollection::VerticesGroup));
+			EXPECT_LT((Collection->Indices)[Index][0], Collection->NumElements(FGeometryCollection::VerticesGroup));
+			EXPECT_LT((Collection->Indices)[Index][1], Collection->NumElements(FGeometryCollection::VerticesGroup));
+			EXPECT_LT((Collection->Indices)[Index][2], Collection->NumElements(FGeometryCollection::VerticesGroup));
 		}
 
-		R.ExpectTrue((Collection->Transform)[Index1].GetTranslation().Z == 10.f);
-		R.ExpectTrue((Collection->Transform)[Index2].GetTranslation().Z == 10.f);
-		R.ExpectTrue((Collection->Transform)[Index3].GetTranslation().Z == 20.f);
-		R.ExpectTrue((Collection->Transform)[Index4].GetTranslation().Z == 20.f);
-		R.ExpectTrue((Collection->Transform)[Index6].GetTranslation().Z == 20.f);
+		EXPECT_EQ((Collection->Transform)[Index1].GetTranslation().Z, 10.f);
+		EXPECT_EQ((Collection->Transform)[Index2].GetTranslation().Z, 10.f);
+		EXPECT_EQ((Collection->Transform)[Index3].GetTranslation().Z, 20.f);
+		EXPECT_EQ((Collection->Transform)[Index4].GetTranslation().Z, 20.f);
+		EXPECT_EQ((Collection->Transform)[Index6].GetTranslation().Z, 20.f);
 
 		HalfTheFaces = Collection->NumElements(FGeometryCollection::FacesGroup) / 2;
-		R.ExpectTrue((Collection->Sections)[0].MaterialID == 0);
-		R.ExpectTrue((Collection->Sections)[0].FirstIndex == 0);
-		R.ExpectTrue((Collection->Sections)[0].NumTriangles == HalfTheFaces);
-		R.ExpectTrue((Collection->Sections)[0].MinVertexIndex == 0);
-		R.ExpectTrue((Collection->Sections)[0].MaxVertexIndex == Collection->NumElements(FGeometryCollection::VerticesGroup) - 1);
+		EXPECT_EQ((Collection->Sections)[0].MaterialID, 0);
+		EXPECT_EQ((Collection->Sections)[0].FirstIndex, 0);
+		EXPECT_EQ((Collection->Sections)[0].NumTriangles, HalfTheFaces);
+		EXPECT_EQ((Collection->Sections)[0].MinVertexIndex, 0);
+		EXPECT_EQ((Collection->Sections)[0].MaxVertexIndex, Collection->NumElements(FGeometryCollection::VerticesGroup) - 1);
 
-		R.ExpectTrue((Collection->Sections)[1].MaterialID == 1);
-		R.ExpectTrue((Collection->Sections)[1].FirstIndex == HalfTheFaces * 3);
-		R.ExpectTrue((Collection->Sections)[1].NumTriangles == HalfTheFaces);
-		R.ExpectTrue((Collection->Sections)[1].MinVertexIndex == 0);
-		R.ExpectTrue((Collection->Sections)[1].MaxVertexIndex == Collection->NumElements(FGeometryCollection::VerticesGroup) - 1);
+		EXPECT_EQ((Collection->Sections)[1].MaterialID, 1);
+		EXPECT_EQ((Collection->Sections)[1].FirstIndex, HalfTheFaces * 3);
+		EXPECT_EQ((Collection->Sections)[1].NumTriangles, HalfTheFaces);
+		EXPECT_EQ((Collection->Sections)[1].MinVertexIndex, 0);
+		EXPECT_EQ((Collection->Sections)[1].MaxVertexIndex, Collection->NumElements(FGeometryCollection::VerticesGroup) - 1);
 
 
 		// GeometryGroup Updated Tests
 		{
-			R.ExpectTrue(Collection->NumElements(FGeometryCollection::GeometryGroup) == 5);
-
-
-			R.ExpectTrue((Collection->TransformIndex)[Index1] == 0);
-			R.ExpectTrue((Collection->TransformIndex)[Index2] == 1);
-			R.ExpectTrue((Collection->TransformIndex)[Index3] == 2);
-			R.ExpectTrue((Collection->TransformIndex)[Index4] == 3);
-			R.ExpectTrue((Collection->TransformIndex)[Index6] == 4);
-
-			R.ExpectTrue((Collection->FaceStart)[Index1] == 0);
-			R.ExpectTrue((Collection->FaceStart)[Index2] == 12);
-			R.ExpectTrue((Collection->FaceStart)[Index3] == 24);
-			R.ExpectTrue((Collection->FaceStart)[Index4] == 36);
-			R.ExpectTrue((Collection->FaceStart)[Index6] == 48);
-
-			R.ExpectTrue((Collection->FaceCount)[Index1] == 12);
-			R.ExpectTrue((Collection->FaceCount)[Index2] == 12);
-			R.ExpectTrue((Collection->FaceCount)[Index3] == 12);
-			R.ExpectTrue((Collection->FaceCount)[Index4] == 12);
-			R.ExpectTrue((Collection->FaceCount)[Index6] == 12);
-			R.ExpectTrue((Collection->Indices).Num() == 60);
-
-			R.ExpectTrue((Collection->VertexStart)[Index1] == 0);
-			R.ExpectTrue((Collection->VertexStart)[Index2] == 8);
-			R.ExpectTrue((Collection->VertexStart)[Index3] == 16);
-			R.ExpectTrue((Collection->VertexStart)[Index4] == 24);
-			R.ExpectTrue((Collection->VertexStart)[Index6] == 32);
-
-			R.ExpectTrue((Collection->VertexCount)[Index1] == 8);
-			R.ExpectTrue((Collection->VertexCount)[Index2] == 8);
-			R.ExpectTrue((Collection->VertexCount)[Index3] == 8);
-			R.ExpectTrue((Collection->VertexCount)[Index4] == 8);
-			R.ExpectTrue((Collection->VertexCount)[Index6] == 8);
-			R.ExpectTrue((Collection->Vertex).Num() == 40);
+			EXPECT_EQ(Collection->NumElements(FGeometryCollection::GeometryGroup), 5);
+
+
+			EXPECT_EQ((Collection->TransformIndex)[Index1], 0);
+			EXPECT_EQ((Collection->TransformIndex)[Index2], 1);
+			EXPECT_EQ((Collection->TransformIndex)[Index3], 2);
+			EXPECT_EQ((Collection->TransformIndex)[Index4], 3);
+			EXPECT_EQ((Collection->TransformIndex)[Index6], 4);
+
+			EXPECT_EQ((Collection->FaceStart)[Index1], 0);
+			EXPECT_EQ((Collection->FaceStart)[Index2], 12);
+			EXPECT_EQ((Collection->FaceStart)[Index3], 24);
+			EXPECT_EQ((Collection->FaceStart)[Index4], 36);
+			EXPECT_EQ((Collection->FaceStart)[Index6], 48);
+
+			EXPECT_EQ((Collection->FaceCount)[Index1], 12);
+			EXPECT_EQ((Collection->FaceCount)[Index2], 12);
+			EXPECT_EQ((Collection->FaceCount)[Index3], 12);
+			EXPECT_EQ((Collection->FaceCount)[Index4], 12);
+			EXPECT_EQ((Collection->FaceCount)[Index6], 12);
+			EXPECT_EQ((Collection->Indices).Num(), 60);
+
+			EXPECT_EQ((Collection->VertexStart)[Index1], 0);
+			EXPECT_EQ((Collection->VertexStart)[Index2], 8);
+			EXPECT_EQ((Collection->VertexStart)[Index3], 16);
+			EXPECT_EQ((Collection->VertexStart)[Index4], 24);
+			EXPECT_EQ((Collection->VertexStart)[Index6], 32);
+
+			EXPECT_EQ((Collection->VertexCount)[Index1], 8);
+			EXPECT_EQ((Collection->VertexCount)[Index2], 8);
+			EXPECT_EQ((Collection->VertexCount)[Index3], 8);
+			EXPECT_EQ((Collection->VertexCount)[Index4], 8);
+			EXPECT_EQ((Collection->VertexCount)[Index6], 8);
+			EXPECT_EQ((Collection->Vertex).Num(), 40);
 		}
 
-		R.ExpectTrue(GeometryCollectionAlgo::HasValidGeometryReferences(Collection.Get()));
-		R.ExpectTrue(Collection->HasContiguousFaces());
-		R.ExpectTrue(Collection->HasContiguousVertices());
-		R.ExpectTrue(Collection->HasContiguousRenderFaces());
-
-		return !R.HasError();
+		EXPECT_TRUE(GeometryCollectionAlgo::HasValidGeometryReferences(Collection.Get()));
+		EXPECT_TRUE(Collection->HasContiguousFaces());
+		EXPECT_TRUE(Collection->HasContiguousVertices());
+		EXPECT_TRUE(Collection->HasContiguousRenderFaces());
 	}
-	template bool DeleteRootLeafMiddle<float>(ExampleResponse&& R);
-
+	template void DeleteRootLeafMiddle<float>();
 
 
 	template<class T>
-	bool DeleteEverything(ExampleResponse&& R)
+	void DeleteEverything()
 	{
 		TSharedPtr<FGeometryCollection> Collection = GeometryCollection::MakeCubeElement(FTransform::Identity, FVector(1.0));
 		Collection->AppendGeometry(*GeometryCollection::MakeCubeElement(FTransform(FVector(0, 0, 10)), FVector(1.0)));
@@ -1003,163 +990,158 @@ namespace GeometryCollectionExample
 		(Collection->Parent)[7] = 2;
 		(Collection->Children)[7].Add(4);
 
-		R.ExpectTrue(Collection->TransformToGeometryIndex.Num() == 8);
-		R.ExpectTrue(Collection->TransformToGeometryIndex[0] == 0);
-		R.ExpectTrue(Collection->TransformToGeometryIndex[1] == 1);
-		R.ExpectTrue(Collection->TransformToGeometryIndex[2] == 2);
-		R.ExpectTrue(Collection->TransformToGeometryIndex[3] == 3);
-		R.ExpectTrue(Collection->TransformToGeometryIndex[4] == 4);
-		R.ExpectTrue(Collection->TransformToGeometryIndex[5] == 5);
-		R.ExpectTrue(Collection->TransformToGeometryIndex[6] == 6);
-		R.ExpectTrue(Collection->TransformToGeometryIndex[7] == 7);
+		EXPECT_EQ(Collection->TransformToGeometryIndex.Num(), 8);
+		EXPECT_EQ(Collection->TransformToGeometryIndex[0], 0);
+		EXPECT_EQ(Collection->TransformToGeometryIndex[1], 1);
+		EXPECT_EQ(Collection->TransformToGeometryIndex[2], 2);
+		EXPECT_EQ(Collection->TransformToGeometryIndex[3], 3);
+		EXPECT_EQ(Collection->TransformToGeometryIndex[4], 4);
+		EXPECT_EQ(Collection->TransformToGeometryIndex[5], 5);
+		EXPECT_EQ(Collection->TransformToGeometryIndex[6], 6);
+		EXPECT_EQ(Collection->TransformToGeometryIndex[7], 7);
 
-		R.ExpectTrue(GeometryCollectionAlgo::HasValidGeometryReferences(Collection.Get()));
+		EXPECT_TRUE(GeometryCollectionAlgo::HasValidGeometryReferences(Collection.Get()));
 
 		TArray<int32> DelList = { 0,1,2,3,4,5,6,7 };
 		Collection->RemoveElements(FTransformCollection::TransformGroup, DelList);
 
-		R.ExpectTrue(Collection->TransformToGeometryIndex.Num() == 0);
-
-		R.ExpectTrue(Collection->HasGroup(FTransformCollection::TransformGroup));
-		R.ExpectTrue(Collection->HasGroup(FGeometryCollection::VerticesGroup));
-		R.ExpectTrue(Collection->HasGroup(FGeometryCollection::FacesGroup));
-		R.ExpectTrue(Collection->HasGroup(FGeometryCollection::MaterialGroup));
+		EXPECT_EQ(Collection->TransformToGeometryIndex.Num(), 0);
 
-		R.ExpectTrue(Collection->NumElements(FTransformCollection::TransformGroup) == 0);
-		R.ExpectTrue(Collection->NumElements(FGeometryCollection::VerticesGroup) == 0);
-		R.ExpectTrue(Collection->NumElements(FGeometryCollection::FacesGroup) == 0);
-		R.ExpectTrue(Collection->NumElements(FGeometryCollection::MaterialGroup) == 0);
+		EXPECT_TRUE(Collection->HasGroup(FTransformCollection::TransformGroup));
+		EXPECT_TRUE(Collection->HasGroup(FGeometryCollection::VerticesGroup));
+		EXPECT_TRUE(Collection->HasGroup(FGeometryCollection::FacesGroup));
+		EXPECT_TRUE(Collection->HasGroup(FGeometryCollection::MaterialGroup));
 
-		R.ExpectTrue(Collection->NumElements(FGeometryCollection::GeometryGroup) == 0);
-		R.ExpectTrue((Collection->Indices).Num() == 0);
-		R.ExpectTrue((Collection->Vertex).Num() == 0);
+		EXPECT_EQ(Collection->NumElements(FTransformCollection::TransformGroup), 0);
+		EXPECT_EQ(Collection->NumElements(FGeometryCollection::VerticesGroup), 0);
+		EXPECT_EQ(Collection->NumElements(FGeometryCollection::FacesGroup), 0);
+		EXPECT_EQ(Collection->NumElements(FGeometryCollection::MaterialGroup), 0);
 
-		R.ExpectTrue(GeometryCollectionAlgo::HasValidGeometryReferences(Collection.Get()));
-		R.ExpectTrue(Collection->HasContiguousFaces());
-		R.ExpectTrue(Collection->HasContiguousVertices());
-		R.ExpectTrue(Collection->HasContiguousRenderFaces());
+		EXPECT_EQ(Collection->NumElements(FGeometryCollection::GeometryGroup), 0);
+		EXPECT_EQ((Collection->Indices).Num(), 0);
+		EXPECT_EQ((Collection->Vertex).Num(), 0);
 
-		return !R.HasError();
+		EXPECT_TRUE(GeometryCollectionAlgo::HasValidGeometryReferences(Collection.Get()));
+		EXPECT_TRUE(Collection->HasContiguousFaces());
+		EXPECT_TRUE(Collection->HasContiguousVertices());
+		EXPECT_TRUE(Collection->HasContiguousRenderFaces());
 	}
-	template bool DeleteEverything<float>(ExampleResponse&& R);
+	template void DeleteEverything<float>();
 
 
 	template<class T>
-	bool ParentTransformTest(ExampleResponse&& R)
+	void ParentTransformTest()
 	{
 		FGeometryCollection* Collection = new FGeometryCollection();
 
 		int32 TransformIndex = Collection->AddElements(1, FGeometryCollection::TransformGroup);
 		(Collection->Transform)[TransformIndex].SetTranslation(FVector(13));
 		(Collection->Parent)[TransformIndex]= -1;
-		R.ExpectTrue(TransformIndex == 0);
+		EXPECT_EQ(TransformIndex, 0);
 
 		TransformIndex = Collection->AddElements(1, FGeometryCollection::TransformGroup);
 		(Collection->Transform)[TransformIndex].SetTranslation(FVector(7));
 		(Collection->Parent)[TransformIndex] = -1;
-		R.ExpectTrue(TransformIndex == 1);
+		EXPECT_EQ(TransformIndex, 1);
 
 		//
 		// Parent a transform
 		//
 		GeometryCollectionAlgo::ParentTransform(Collection, 1, 0);
-		R.ExpectTrue((Collection->Children)[0].Num()==0);
-		R.ExpectTrue((Collection->Parent)[0] == 1);
-		R.ExpectTrue((Collection->Children)[1].Num() == 1);
-		R.ExpectTrue((Collection->Children)[1].Contains(0) );
-		R.ExpectTrue((Collection->Parent)[1] == -1);
+		EXPECT_EQ((Collection->Children)[0].Num(), 0);
+		EXPECT_EQ((Collection->Parent)[0], 1);
+		EXPECT_EQ((Collection->Children)[1].Num(), 1);
+		EXPECT_TRUE((Collection->Children)[1].Contains(0));
+		EXPECT_EQ((Collection->Parent)[1], -1);
 
 		TArray<FTransform> GlobalTransform;
 		GeometryCollectionAlgo::GlobalMatrices(Collection->Transform, Collection->Parent, GlobalTransform);
-		R.ExpectTrue(((Collection->Transform)[0].GetTranslation() - FVector(6)).Size() < KINDA_SMALL_NUMBER);
-		R.ExpectTrue((GlobalTransform[0].GetTranslation()-FVector(13)).Size()<KINDA_SMALL_NUMBER);
+		EXPECT_LT(((Collection->Transform)[0].GetTranslation() - FVector(6)).Size(), KINDA_SMALL_NUMBER);
+		EXPECT_LT((GlobalTransform[0].GetTranslation()-FVector(13)).Size(), KINDA_SMALL_NUMBER);
 
-		R.ExpectTrue(Collection->TransformToGeometryIndex.Num() == 2);
-		R.ExpectTrue(Collection->TransformToGeometryIndex[0] == -1);
-		R.ExpectTrue(Collection->TransformToGeometryIndex[1] == -1);
+		EXPECT_EQ(Collection->TransformToGeometryIndex.Num(), 2);
+		EXPECT_EQ(Collection->TransformToGeometryIndex[0], -1);
+		EXPECT_EQ(Collection->TransformToGeometryIndex[1], -1);
 
 		//
 		// Add some geometry
 		//
 		TransformIndex = Collection->AppendGeometry(*GeometryCollection::MakeCubeElement(FTransform(FVector(3)), FVector(1.0)));
-		R.ExpectTrue(((Collection->Transform)[TransformIndex].GetTranslation() - FVector(3)).Size() < KINDA_SMALL_NUMBER);
-		R.ExpectTrue((Collection->TransformIndex).Num() == 1);
-		R.ExpectTrue((Collection->TransformIndex)[0] == TransformIndex);
-		R.ExpectTrue((Collection->VertexStart)[0] == 0);
-		R.ExpectTrue((Collection->VertexCount)[0] == 8);
+		EXPECT_LT(((Collection->Transform)[TransformIndex].GetTranslation() - FVector(3)).Size(), KINDA_SMALL_NUMBER);
+		EXPECT_EQ((Collection->TransformIndex).Num(), 1);
+		EXPECT_EQ((Collection->TransformIndex)[0], TransformIndex);
+		EXPECT_EQ((Collection->VertexStart)[0], 0);
+		EXPECT_EQ((Collection->VertexCount)[0], 8);
 		for (int i = (Collection->VertexStart)[0]; i < (Collection->VertexCount)[0]; i++)
 		{
-			R.ExpectTrue((Collection->BoneMap)[i] == TransformIndex);
+			EXPECT_EQ((Collection->BoneMap)[i], TransformIndex);
 		}
 
-		R.ExpectTrue(Collection->TransformToGeometryIndex.Num() == 3);
-		R.ExpectTrue(Collection->TransformToGeometryIndex[0] == -1);
-		R.ExpectTrue(Collection->TransformToGeometryIndex[1] == -1);
-		R.ExpectTrue(Collection->TransformToGeometryIndex[2] == 0);
+		EXPECT_EQ(Collection->TransformToGeometryIndex.Num(), 3);
+		EXPECT_EQ(Collection->TransformToGeometryIndex[0], -1);
+		EXPECT_EQ(Collection->TransformToGeometryIndex[1], -1);
+		EXPECT_EQ(Collection->TransformToGeometryIndex[2], 0);
 
 		//
 		// Parent the geometry
 		//
 		GeometryCollectionAlgo::ParentTransform(Collection, 0, TransformIndex);
-		R.ExpectTrue((Collection->Children)[0].Num() == 1);
-		R.ExpectTrue((Collection->Parent)[0] == 1);
-		R.ExpectTrue((Collection->Children)[1].Num() == 1);
-		R.ExpectTrue((Collection->Children)[1].Contains(0));
-		R.ExpectTrue((Collection->Parent)[1] == -1);
-		R.ExpectTrue(((Collection->Transform)[TransformIndex].GetTranslation() - FVector(-10)).Size() < KINDA_SMALL_NUMBER);
-		R.ExpectTrue((Collection->TransformIndex).Num() == 1);
-		R.ExpectTrue((Collection->TransformIndex)[0] == TransformIndex);
-		R.ExpectTrue((Collection->VertexStart)[0] == 0);
-		R.ExpectTrue((Collection->VertexCount)[0] == 8);
+		EXPECT_EQ((Collection->Children)[0].Num(), 1);
+		EXPECT_EQ((Collection->Parent)[0], 1);
+		EXPECT_EQ((Collection->Children)[1].Num(), 1);
+		EXPECT_TRUE((Collection->Children)[1].Contains(0));
+		EXPECT_EQ((Collection->Parent)[1], -1);
+		EXPECT_LT(((Collection->Transform)[TransformIndex].GetTranslation() - FVector(-10)).Size(), KINDA_SMALL_NUMBER);
+		EXPECT_EQ((Collection->TransformIndex).Num(), 1);
+		EXPECT_EQ((Collection->TransformIndex)[0], TransformIndex);
+		EXPECT_EQ((Collection->VertexStart)[0], 0);
+		EXPECT_EQ((Collection->VertexCount)[0], 8);
 		for (int i = (Collection->VertexStart)[0]; i < (Collection->VertexCount)[0]; i++)
 		{
-			R.ExpectTrue((Collection->BoneMap)[i] == TransformIndex);
+			EXPECT_EQ((Collection->BoneMap)[i], TransformIndex);
 		}
 
 		GeometryCollectionAlgo::GlobalMatrices(Collection->Transform, Collection->Parent, GlobalTransform);
-		R.ExpectTrue((GlobalTransform[0].GetTranslation() - FVector(13)).Size() < KINDA_SMALL_NUMBER);
-		R.ExpectTrue((GlobalTransform[2].GetTranslation() - FVector(3)).Size() < KINDA_SMALL_NUMBER);
-
+		EXPECT_LT((GlobalTransform[0].GetTranslation() - FVector(13)).Size(), KINDA_SMALL_NUMBER);
+		EXPECT_LT((GlobalTransform[2].GetTranslation() - FVector(3)).Size(), KINDA_SMALL_NUMBER);
 
 
 		//
 		// Force a circular parent
 		//
-		R.ExpectTrue(!GeometryCollectionAlgo::HasCycle((Collection->Parent), TransformIndex));
+		EXPECT_FALSE(GeometryCollectionAlgo::HasCycle((Collection->Parent), TransformIndex));
 		(Collection->Children)[0].Add(2);
 		(Collection->Parent)[0] = 2;
 		(Collection->Children)[2].Add(0);
 		(Collection->Parent)[2] = 0;
-		R.ExpectTrue(GeometryCollectionAlgo::HasCycle((Collection->Parent), TransformIndex));
+		EXPECT_TRUE(GeometryCollectionAlgo::HasCycle((Collection->Parent), TransformIndex));
 
 		delete Collection;
-
-		return !R.HasError();
 	}
-	template bool ParentTransformTest<float>(ExampleResponse&& R);
+	template void ParentTransformTest<float>();
 
 	template<class T>
-	bool ReindexMaterialsTest(ExampleResponse&& R)
+	void ReindexMaterialsTest()
 	{
 		TSharedPtr<FGeometryCollection> Collection = GeometryCollection::MakeCubeElement(FTransform::Identity, FVector(1.0));
 		
-		R.ExpectTrue(Collection->Sections.Num() == 2);
+		EXPECT_EQ(Collection->Sections.Num(), 2);
 
 		Collection->ReindexMaterials();
 
 		// Reindexing doesn't change the number of sections
-		R.ExpectTrue(Collection->Sections.Num() == 2);
+		EXPECT_EQ(Collection->Sections.Num(), 2);
 
 		// Ensure material selections have correct material ids after reindexing
 		for (int i = 0; i < 12; i++)
 		{
 			if (i < 6)
 			{
-				R.ExpectTrue((Collection->MaterialID)[i] == 0);
+				EXPECT_EQ((Collection->MaterialID)[i], 0);
 			}
 			else
 			{
-				R.ExpectTrue((Collection->MaterialID)[i] == 1);
+				EXPECT_EQ((Collection->MaterialID)[i], 1);
 			}
 		}
 
@@ -1170,9 +1152,9 @@ namespace GeometryCollectionExample
 		Collection->ReindexMaterials();
 
 		// Ensure we now have 1 section
-		R.ExpectTrue(Collection->Sections.Num() == 1);
-		R.ExpectTrue((Collection->Sections)[0].MaterialID == 1);
-		R.ExpectTrue((Collection->Sections)[0].NumTriangles == 6);		
+		EXPECT_EQ(Collection->Sections.Num(), 1);
+		EXPECT_EQ((Collection->Sections)[0].MaterialID, 1);
+		EXPECT_EQ((Collection->Sections)[0].NumTriangles, 6);		
 		 
 		// Add a copy of the geometry and reindex
 		TSharedPtr<FGeometryCollection> Collection2 = GeometryCollection::MakeCubeElement(FTransform::Identity, FVector(1.0));
@@ -1180,18 +1162,16 @@ namespace GeometryCollectionExample
 		Collection->ReindexMaterials();
 
 		// test that sections created are consolidated
-		R.ExpectTrue(Collection->Sections.Num() == 2);
-		R.ExpectTrue((Collection->Sections)[0].MaterialID == 0);
-		R.ExpectTrue((Collection->Sections)[0].NumTriangles == 6);
-		R.ExpectTrue((Collection->Sections)[1].MaterialID == 1);
-		R.ExpectTrue((Collection->Sections)[1].NumTriangles == 12);
-
-		return !R.HasError();
+		EXPECT_EQ(Collection->Sections.Num(), 2);
+		EXPECT_EQ((Collection->Sections)[0].MaterialID, 0);
+		EXPECT_EQ((Collection->Sections)[0].NumTriangles, 6);
+		EXPECT_EQ((Collection->Sections)[1].MaterialID, 1);
+		EXPECT_EQ((Collection->Sections)[1].NumTriangles, 12);
 	}
-	template bool ReindexMaterialsTest<float>(ExampleResponse&& R);
+	template void ReindexMaterialsTest<float>();
 
 	template<class T>
-	bool AttributeTransferTest(ExampleResponse&& R)
+	void AttributeTransferTest()
 	{
 		TSharedPtr<FGeometryCollection> Collection1 = GeometryCollection::MakeCubeElement(FTransform::Identity, FVector(1.0));
 		TSharedPtr<FGeometryCollection> Collection2 = GeometryCollection::MakeCubeElement(FTransform::Identity, FVector(1.0));
@@ -1210,15 +1190,13 @@ namespace GeometryCollectionExample
 		// test color is set correctly on 2
 		for (int i = 0; i < Collection2->NumElements(FGeometryCollection::VerticesGroup); ++i)
 		{
-			R.ExpectTrue((Collection2->Color)[i].Equals(FColor(1, 0, 1, 1)));
+			EXPECT_TRUE((Collection2->Color)[i].Equals(FColor(1, 0, 1, 1)));
 		}
-
-		return !R.HasError();
 	}
-	template bool AttributeTransferTest<float>(ExampleResponse&& R);
+	template void AttributeTransferTest<float>();
 
 	template<class T>
-	bool AttributeDependencyTest(ExampleResponse&& R)
+	void AttributeDependencyTest()
 	{
 		FGeometryCollection* Collection = new FGeometryCollection();
 
@@ -1251,9 +1229,7 @@ namespace GeometryCollectionExample
 		//Collection->SetDependency()<FTransform>("AttributeD", Group1, Transform, Group3);
 
 		delete Collection;
-
-		return !R.HasError();
 	}
-	template bool AttributeDependencyTest<float>(ExampleResponse&& R);
+	template void AttributeDependencyTest<float>();
 
 }
