Copying //UE4/Dev-Physics to Dev-Main (//UE4/Dev-Main) [at] 9098278

diff --git a/Engine/Plugins/Experimental/GeometryCollectionPlugin/Source/GeometryCollectionExampleCore/Private/GeometryCollection/GeometryCollectionExampleDecimation.cpp b/Engine/Plugins/Experimental/GeometryCollectionPlugin/Source/GeometryCollectionExampleCore/Private/GeometryCollection/GeometryCollectionExampleDecimation.cpp
--- a/Engine/Plugins/Experimental/GeometryCollectionPlugin/Source/GeometryCollectionExampleCore/Private/GeometryCollection/GeometryCollectionExampleDecimation.cpp
+++ b/Engine/Plugins/Experimental/GeometryCollectionPlugin/Source/GeometryCollectionExampleCore/Private/GeometryCollection/GeometryCollectionExampleDecimation.cpp
@@ -132,7 +132,7 @@ namespace GeometryCollectionExample
 
 	template <class T>
 	bool
-	RunGeomDecimationTest(FGeometryCollection* TestCollection, ExampleResponse&  R, const char *BaseName, const char *OutputDir, const uint32 ExpectedHash, const bool RestrictToLocalIndexRange = false)
+	RunGeomDecimationTest(FGeometryCollection* TestCollection, const char *BaseName, const char *OutputDir, const uint32 ExpectedHash, const bool RestrictToLocalIndexRange = false)
 	{
 		Chaos::TParticles<T, 3> Particles = BuildParticlesFeomGeomCollection(TestCollection);
 		Chaos::TTriangleMesh<T> TriMesh = BuildTriMeshFromGeomCollection(TestCollection);
@@ -142,9 +142,18 @@ namespace GeometryCollectionExample
 		check(CoincidentVertices.Num() < Importance.Num());
 
 		const int numParticles = Particles.Size();
-		R.ExpectTrue((!RestrictToLocalIndexRange && Importance.Num() == numParticles) ||
-					 (RestrictToLocalIndexRange && Importance.Num() <= numParticles)); // got the right number of indices
-		R.ExpectTrue(TSet<int32>(Importance).Num() == Importance.Num()); // indices were unique
+
+		// got the right number of indices
+		if (RestrictToLocalIndexRange)
+		{
+			EXPECT_LE(Importance.Num(), numParticles);
+		}
+		else 
+		{
+			EXPECT_EQ(Importance.Num(), numParticles);
+		}
+
+		EXPECT_EQ(TSet<int32>(Importance).Num(), Importance.Num()); // indices were unique
 
 		WriteImportanceOrderObjs(
 			TestCollection, Importance, CoincidentVertices, BaseName, OutputDir);
@@ -165,53 +174,51 @@ namespace GeometryCollectionExample
 
 	template <class T, class TGeom>
 	bool
-	RunGeomDecimationTest(ExampleResponse& R, const char *BaseName, const char *OutputDir, const uint32 ExpectedHash, const bool RestrictToLocalIndexRange = false)
+	RunGeomDecimationTest(const char *BaseName, const char *OutputDir, const uint32 ExpectedHash, const bool RestrictToLocalIndexRange = false)
 	{
 		TGeom Geom;
 		FGeometryCollection* TestCollection =
 			FGeometryCollection::NewGeometryCollection(
 				Geom.RawVertexArray,
 				Geom.RawIndicesArray);
-		return RunGeomDecimationTest<T>(TestCollection, R, BaseName, OutputDir, ExpectedHash, RestrictToLocalIndexRange);
+		return RunGeomDecimationTest<T>(TestCollection, BaseName, OutputDir, ExpectedHash, RestrictToLocalIndexRange);
 	}
 
 	template <class T, class TGeom>
 	bool
-	RunGlobalGeomDecimationTest(ExampleResponse& R, const char *BaseName, const char *OutputDir, const uint32 ExpectedHash, const bool RestrictToLocalIndexRange = true)
+	RunGlobalGeomDecimationTest(const char *BaseName, const char *OutputDir, const uint32 ExpectedHash, const bool RestrictToLocalIndexRange = true)
 	{
 		TGeom Geom;
 		FGeometryCollection* TestCollection =
 			FGeometryCollection::NewGeometryCollection(
 				Geom.RawVertexArray,
 				Geom.RawIndicesArray1);
-		return RunGeomDecimationTest<T>(TestCollection, R, BaseName, OutputDir, ExpectedHash, RestrictToLocalIndexRange);
+		return RunGeomDecimationTest<T>(TestCollection, BaseName, OutputDir, ExpectedHash, RestrictToLocalIndexRange);
 	}
 
 	template<class T>
-	bool TestGeometryDecimation(ExampleResponse&& R)
+	void TestGeometryDecimation()
 	{
 		bool Success = true;
 		// If E:\TestGeometry\Decimation doesn't already exist, the files aren't written.
 
 		// Standalone point pools.
-		Success &= RunGeomDecimationTest<T, BoxGeometry>(R, "box", "E:\\TestGeometry\\Decimation\\", 4024338882);
-		Success &= RunGeomDecimationTest<T, CylinderGeometry>(R, "cylinder", "E:\\TestGeometry\\Decimation\\", 2477299646);
-		Success &= RunGeomDecimationTest<T, EllipsoidGeometry>(R, "ellipsoid", "E:\\TestGeometry\\Decimation\\", 1158371240);
-		Success &= RunGeomDecimationTest<T, EllipsoidGeometry2>(R, "ellipsoid2", "E:\\TestGeometry\\Decimation\\", 554754926);
-		Success &= RunGeomDecimationTest<T, EllipsoidGeometry3>(R, "ellipsoid3", "E:\\TestGeometry\\Decimation\\", 2210765036);
-		Success &= RunGeomDecimationTest<T, FracturedGeometry>(R, "fractured", "E:\\TestGeometry\\Decimation\\", 2030682536);
-		Success &= RunGeomDecimationTest<T, SphereGeometry>(R, "sphere", "E:\\TestGeometry\\Decimation\\", 4119721232);
-		Success &= RunGeomDecimationTest<T, TorusGeometry>(R, "torus", "E:\\TestGeometry\\Decimation\\", 2519379615);
+		Success &= RunGeomDecimationTest<T, BoxGeometry>("box", "E:\\TestGeometry\\Decimation\\", 4024338882);
+		Success &= RunGeomDecimationTest<T, CylinderGeometry>("cylinder", "E:\\TestGeometry\\Decimation\\", 2477299646);
+		Success &= RunGeomDecimationTest<T, EllipsoidGeometry>("ellipsoid", "E:\\TestGeometry\\Decimation\\", 1158371240);
+		Success &= RunGeomDecimationTest<T, EllipsoidGeometry2>("ellipsoid2", "E:\\TestGeometry\\Decimation\\", 554754926);
+		Success &= RunGeomDecimationTest<T, EllipsoidGeometry3>("ellipsoid3", "E:\\TestGeometry\\Decimation\\", 2210765036);
+		Success &= RunGeomDecimationTest<T, FracturedGeometry>("fractured", "E:\\TestGeometry\\Decimation\\", 2030682536);
+		Success &= RunGeomDecimationTest<T, SphereGeometry>("sphere", "E:\\TestGeometry\\Decimation\\", 4119721232);
+		Success &= RunGeomDecimationTest<T, TorusGeometry>("torus", "E:\\TestGeometry\\Decimation\\", 2519379615);
 		
 		// Geometry in a global point pool.
-		Success &= RunGeomDecimationTest<T, GlobalFracturedGeometry>(R, "globalFractured", "E:\\TestGeometry\\Decimation\\", 4227374796, true);
-		Success &= RunGlobalGeomDecimationTest<T, GlobalFracturedGeometry>(R, "globalFracturedMerged", "E:\\TestGeometry\\Decimation\\", 4227374796, true);
-
-		R.ExpectTrue(Success);
+		Success &= RunGeomDecimationTest<T, GlobalFracturedGeometry>("globalFractured", "E:\\TestGeometry\\Decimation\\", 4227374796, true);
+		Success &= RunGlobalGeomDecimationTest<T, GlobalFracturedGeometry>("globalFracturedMerged", "E:\\TestGeometry\\Decimation\\", 4227374796, true);
 
-		return !R.HasError();
+		EXPECT_TRUE(Success);
 	}
-	template bool TestGeometryDecimation<float>(ExampleResponse&& R);
+	template void TestGeometryDecimation<float>();
 	
 }
 
diff --git a/Engine/Plugins/Experimental/GeometryCollectionPlugin/Source/GeometryCollectionExampleCore/Private/GeometryCollection/GeometryCollectionExampleEvents.cpp b/Engine/Plugins/Experimental/GeometryCollectionPlugin/Source/GeometryCollectionExampleCore/Private/GeometryCollection/GeometryCollectionExampleEvents.cpp
--- a/Engine/Plugins/Experimental/GeometryCollectionPlugin/Source/GeometryCollectionExampleCore/Private/GeometryCollection/GeometryCollectionExampleEvents.cpp
+++ b/Engine/Plugins/Experimental/GeometryCollectionPlugin/Source/GeometryCollectionExampleCore/Private/GeometryCollection/GeometryCollectionExampleEvents.cpp
@@ -80,7 +80,7 @@ namespace GeometryCollectionExample
 
 
 	template<class T>
-	bool Event_Handler(ExampleResponse&& R)
+	void Event_Handler()
 	{
 #if INCLUDE_CHAOS
 		Chaos::FEventManager EventManager(Chaos::EMultiBufferMode::Single);
@@ -117,7 +117,7 @@ namespace GeometryCollectionExample
 		EventManager.FillProducerData(Solver);
 		EventManager.FlipBuffersIfRequired();
 		EventManager.DispatchEvents();
-		R.ExpectTrue(HandlerTest.ResultFromHandler == TestData);
+		EXPECT_EQ(HandlerTest.ResultFromHandler, TestData);
 
 		TestData.Data1 = 789;
 		TestData.Data2 = FVector(7, 8, 9);
@@ -125,7 +125,7 @@ namespace GeometryCollectionExample
 		EventManager.FillProducerData(Solver);
 		EventManager.FlipBuffersIfRequired();
 		EventManager.DispatchEvents();
-		R.ExpectTrue(HandlerTest.ResultFromHandler == TestData);
+		EXPECT_EQ(HandlerTest.ResultFromHandler, TestData);
 
 		// Unregister - data should no longer update
 		HandlerTest.UnregisterHandler1();
@@ -137,7 +137,7 @@ namespace GeometryCollectionExample
 		EventManager.FillProducerData(Solver);
 		EventManager.FlipBuffersIfRequired();
 		EventManager.DispatchEvents();
-		R.ExpectTrue(HandlerTest.ResultFromHandler == OriginalTestData);
+		EXPECT_EQ(HandlerTest.ResultFromHandler, OriginalTestData);
 
 		EventTestData Data;
 		TestArrayData.Push(EventTestData(123, FVector(1, 2, 3)));
@@ -148,8 +148,8 @@ namespace GeometryCollectionExample
 		EventManager.FlipBuffersIfRequired();
 		EventManager.DispatchEvents();
 		// dispatched to multiple handlers
-		R.ExpectTrue(HandlerTest.ResultFromHandler2 == TestArrayData);
-		R.ExpectTrue(AnotherHandlerTest.ResultFromHandler2 == TestArrayData);
+		EXPECT_EQ(HandlerTest.ResultFromHandler2, TestArrayData);
+		EXPECT_EQ(AnotherHandlerTest.ResultFromHandler2, TestArrayData);
 
 		// Unregister one of the handlers from the events
 		HandlerTest.UnregisterHandler2();
@@ -161,21 +161,19 @@ namespace GeometryCollectionExample
 		EventManager.FillProducerData(Solver);
 		EventManager.FlipBuffersIfRequired();
 		EventManager.DispatchEvents();
-		R.ExpectTrue(HandlerTest.ResultFromHandler2 == OriginalTestArrayData); // Unregistered - data should no longer update
-		R.ExpectTrue(AnotherHandlerTest.ResultFromHandler2 == TestArrayData); // Still registered so should get updates
+		EXPECT_EQ(HandlerTest.ResultFromHandler2, OriginalTestArrayData); // Unregistered - data should no longer update
+		EXPECT_EQ(AnotherHandlerTest.ResultFromHandler2, TestArrayData); // Still registered so should get updates
 
 		HandlerTest.RegisterHandler2();
 
 		EventManager.FillProducerData(Solver);
 		EventManager.FlipBuffersIfRequired();
 		EventManager.DispatchEvents();
-		R.ExpectTrue(HandlerTest.ResultFromHandler2 == TestArrayData);
-		R.ExpectTrue(AnotherHandlerTest.ResultFromHandler2 == TestArrayData);
+		EXPECT_EQ(HandlerTest.ResultFromHandler2, TestArrayData);
+		EXPECT_EQ(AnotherHandlerTest.ResultFromHandler2, TestArrayData);
 #endif
-		return !R.HasError();
-
 	}
-	template bool Event_Handler<float>(ExampleResponse&& R);
+	template void Event_Handler<float>();
 
 }
 
diff --git a/Engine/Plugins/Experimental/GeometryCollectionPlugin/Source/GeometryCollectionExampleCore/Private/GeometryCollection/GeometryCollectionExampleFields.cpp b/Engine/Plugins/Experimental/GeometryCollectionPlugin/Source/GeometryCollectionExampleCore/Private/GeometryCollection/GeometryCollectionExampleFields.cpp
--- a/Engine/Plugins/Experimental/GeometryCollectionPlugin/Source/GeometryCollectionExampleCore/Private/GeometryCollection/GeometryCollectionExampleFields.cpp
+++ b/Engine/Plugins/Experimental/GeometryCollectionPlugin/Source/GeometryCollectionExampleCore/Private/GeometryCollection/GeometryCollectionExampleFields.cpp
@@ -20,7 +20,7 @@ DEFINE_LOG_CATEGORY_STATIC(GCTF_Log, Verbose, All);
 
 namespace GeometryCollectionExample
 {
-	bool Fields_NoiseSample(ExampleResponse&& R)
+	void Fields_NoiseSample()
 	{
 		int Bounds = 100;
 		TArray<ContextIndex> IndicesArray;
@@ -68,17 +68,14 @@ namespace GeometryCollectionExample
 			//std::cout << std::endl;
 		}
 
-		R.ExpectTrue(min >= MinDoimain);
-		R.ExpectTrue(max <= MaxDomain);
-		R.ExpectTrue(min < max);
+		EXPECT_GE(min, MinDoimain);
+		EXPECT_LE(max, MaxDomain);
+		EXPECT_LT(min, max);
 
 		delete NoiseField;
-
-		return !R.HasError();
-
 	}
 
-	bool Fields_RadialIntMask(ExampleResponse&& R)
+	void Fields_RadialIntMask()
 	{
 		TArray<ContextIndex> IndicesArray;
 		ContextIndex::ContiguousIndices(IndicesArray, 10);
@@ -113,20 +110,18 @@ namespace GeometryCollectionExample
 		{
 			if (Index <= 2)
 			{
-				R.ExpectTrue(!!ResultsView[Index]);
+				EXPECT_TRUE(ResultsView[Index]);
 			}
 			else
 			{
-				R.ExpectTrue(!ResultsView[Index]);
+				EXPECT_FALSE(ResultsView[Index]);
 			}
 			//UE_LOG(GCTF_Log, Error, TEXT("[%d] %d"), Index, ResultsView[Index]);
 		}
 		delete RadialMask;
-
-		return !R.HasError();
 	}
 
-	bool Fields_RadialFalloff(ExampleResponse&& R)
+	void Fields_RadialFalloff()
 	{
 		TArray<ContextIndex> IndicesArray;
 		ContextIndex::ContiguousIndices(IndicesArray, 10);
@@ -163,11 +158,11 @@ namespace GeometryCollectionExample
 
 			if (Index <= 5)
 			{
-				R.ExpectTrue(FMath::Abs(ResultsView[Index]) - ExpectedVal < KINDA_SMALL_NUMBER);
+				EXPECT_LT(FMath::Abs(ResultsView[Index]) - ExpectedVal, KINDA_SMALL_NUMBER);
 			}
 			else
 			{
-				R.ExpectTrue(!ResultsView[Index]);
+				EXPECT_FALSE(ResultsView[Index]);
 			}
 			//UE_LOG(GCTF_Log, Error, TEXT("[%d] sample:%3.5f (%3.5f,%3.5f,%3.5f) %3.5f"), Index,
 			//	ExpectedVal ,
@@ -175,11 +170,9 @@ namespace GeometryCollectionExample
 			//	ResultsView[Index]);
 		}
 		delete RadialFalloff;
-
-		return !R.HasError();
 	}
 
-	bool Fields_PlaneFalloff(ExampleResponse&& R)
+	void Fields_PlaneFalloff()
 	{
 		TArray<ContextIndex> IndicesArray;
 		ContextIndex::ContiguousIndices(IndicesArray, 10);
@@ -220,14 +213,12 @@ namespace GeometryCollectionExample
 				ExpectedVal = -PlaneFalloff->Magnitude * Distance;
 			}
 
-			R.ExpectTrue(FMath::Abs(ResultsView[Index]) - ExpectedVal < KINDA_SMALL_NUMBER);
+			EXPECT_LT(FMath::Abs(ResultsView[Index]) - ExpectedVal, KINDA_SMALL_NUMBER);
 		}
 		delete PlaneFalloff;
-
-		return !R.HasError();
 	}
 
-	bool Fields_UniformVector(ExampleResponse&& R)
+	void Fields_UniformVector()
 	{
 		TArray<ContextIndex> IndicesArray;
 		ContextIndex::ContiguousIndices(IndicesArray, 10);
@@ -260,18 +251,16 @@ namespace GeometryCollectionExample
 		for (int32 Index = 0; Index < 10; Index++)
 		{
 			FVector ExpectedVal = 10.0*FVector(3, 5, 7);
-			R.ExpectTrue((ResultsView[Index] - ExpectedVal).Size() < KINDA_SMALL_NUMBER);
+			EXPECT_LT((ResultsView[Index] - ExpectedVal).Size(), KINDA_SMALL_NUMBER);
 			//UE_LOG(GCTF_Log, Error, TEXT("[%d] sample:(%3.5f,%3.5f,%3.5f) (%3.5f,%3.5f,%3.5f) (%3.5f,%3.5f,%3.5f)"), Index,
 			//	ExpectedVal.X, ExpectedVal.Y, ExpectedVal.Z,
 			//	SamplesArray[Index].X, SamplesArray[Index].Y, SamplesArray[Index].Z,
 			//	ResultsView[Index].X, ResultsView[Index].Y, ResultsView[Index].Z);
 		}
 		delete UniformVector;
-
-		return !R.HasError();
 	}
 
-	bool Fields_RaidalVector(ExampleResponse&& R)
+	void Fields_RaidalVector()
 	{
 		TArray<ContextIndex> IndicesArray;
 		ContextIndex::ContiguousIndices(IndicesArray, 10);
@@ -304,18 +293,16 @@ namespace GeometryCollectionExample
 		for (int32 Index = 0; Index < 10; Index++)
 		{
 			FVector ExpectedVal = RadialVector->Magnitude * (SamplesArray[Index] - RadialVector->Position).GetSafeNormal();
-			R.ExpectTrue((ResultsView[Index] - ExpectedVal).Size() < KINDA_SMALL_NUMBER);
+			EXPECT_LT((ResultsView[Index] - ExpectedVal).Size(), KINDA_SMALL_NUMBER);
 			//UE_LOG(GCTF_Log, Error, TEXT("[%d] sample:(%3.5f,%3.5f,%3.5f) (%3.5f,%3.5f,%3.5f) (%3.5f,%3.5f,%3.5f)"), Index,
 			//	ExpectedVal.X, ExpectedVal.Y, ExpectedVal.Z,
 			//	SamplesArray[Index].X, SamplesArray[Index].Y, SamplesArray[Index].Z,
 			//	ResultsView[Index].X, ResultsView[Index].Y, ResultsView[Index].Z);
 		}
 		delete RadialVector;
-
-		return !R.HasError();
 	}
 
-	bool Fields_SumVectorFullMult(ExampleResponse&& R)
+	void Fields_SumVectorFullMult()
 	{
 		TArray<ContextIndex> IndicesArray;
 		ContextIndex::ContiguousIndices(IndicesArray, 10);
@@ -381,14 +368,12 @@ namespace GeometryCollectionExample
 			//	SamplesArray[Index].X, SamplesArray[Index].Y, SamplesArray[Index].Z,
 			//	ResultsView[Index].X, ResultsView[Index].Y, ResultsView[Index].Z);
 
-			R.ExpectTrue((ResultsView[Index] - ExpectedVal).Size() < KINDA_SMALL_NUMBER);
+			EXPECT_LT((ResultsView[Index] - ExpectedVal).Size(), KINDA_SMALL_NUMBER);
 		}
 		delete SumVector;
-
-		return !R.HasError();
 	}
 
-	bool Fields_SumVectorFullDiv(ExampleResponse&& R)
+	void Fields_SumVectorFullDiv()
 	{
 		TArray<ContextIndex> IndicesArray;
 		ContextIndex::ContiguousIndices(IndicesArray, 10);
@@ -454,15 +439,12 @@ namespace GeometryCollectionExample
 			//	SamplesArray[Index].X, SamplesArray[Index].Y, SamplesArray[Index].Z,
 			//	ResultsView[Index].X, ResultsView[Index].Y, ResultsView[Index].Z);
 
-			R.ExpectTrue((ResultsView[Index] - ExpectedVal).Size() < KINDA_SMALL_NUMBER);
+			EXPECT_LT((ResultsView[Index] - ExpectedVal).Size(), KINDA_SMALL_NUMBER);
 		}
 		delete SumVector;
-
-		return !R.HasError();
 	}
 
-
-	bool Fields_SumVectorFullAdd(ExampleResponse&& R)
+	void Fields_SumVectorFullAdd()
 	{
 		TArray<ContextIndex> IndicesArray;
 		ContextIndex::ContiguousIndices(IndicesArray, 10);
@@ -529,15 +511,12 @@ namespace GeometryCollectionExample
 			//	SamplesArray[Index].X, SamplesArray[Index].Y, SamplesArray[Index].Z,
 			//	ResultsView[Index].X, ResultsView[Index].Y, ResultsView[Index].Z);
 
-			R.ExpectTrue((ResultsView[Index] - ExpectedVal).Size() < KINDA_SMALL_NUMBER);
+			EXPECT_LT((ResultsView[Index] - ExpectedVal).Size(), KINDA_SMALL_NUMBER);
 		}
 		delete SumVector;
-
-		return !R.HasError();
 	}
 
-
-	bool Fields_SumVectorFullSub(ExampleResponse&& R)
+	void Fields_SumVectorFullSub()
 	{
 		TArray<ContextIndex> IndicesArray;
 		ContextIndex::ContiguousIndices(IndicesArray, 10);
@@ -604,15 +583,12 @@ namespace GeometryCollectionExample
 			//	SamplesArray[Index].X, SamplesArray[Index].Y, SamplesArray[Index].Z,
 			//	ResultsView[Index].X, ResultsView[Index].Y, ResultsView[Index].Z);
 
-			R.ExpectTrue((ResultsView[Index] - ExpectedVal).Size() < KINDA_SMALL_NUMBER);
+			EXPECT_LT((ResultsView[Index] - ExpectedVal).Size(), KINDA_SMALL_NUMBER);
 		}
 		delete SumVector;
-
-		return !R.HasError();
 	}
 
-
-	bool Fields_SumVectorLeftSide(ExampleResponse&& R)
+	void Fields_SumVectorLeftSide()
 	{
 		TArray<ContextIndex> IndicesArray;
 		ContextIndex::ContiguousIndices(IndicesArray, 10);
@@ -678,15 +654,15 @@ namespace GeometryCollectionExample
 			//	SamplesArray[Index].X, SamplesArray[Index].Y, SamplesArray[Index].Z,
 			//	ResultsView[Index].X, ResultsView[Index].Y, ResultsView[Index].Z);
 
-			R.ExpectTrue((ResultsView[Index] - ExpectedVal).Size() < KINDA_SMALL_NUMBER);
+			EXPECT_LT((ResultsView[Index] - ExpectedVal).Size(), KINDA_SMALL_NUMBER);
 		}
 		delete SumVector;
 		delete UniformVector;
 
-		return !R.HasError();
+		
 	}
 
-	bool Fields_SumVectorRightSide(ExampleResponse&& R)
+	void Fields_SumVectorRightSide()
 	{
 		TArray<ContextIndex> IndicesArray;
 		ContextIndex::ContiguousIndices(IndicesArray, 10);
@@ -754,16 +730,15 @@ namespace GeometryCollectionExample
 			//	SamplesArray[Index].X, SamplesArray[Index].Y, SamplesArray[Index].Z,
 			//	ResultsView[Index].X, ResultsView[Index].Y, ResultsView[Index].Z);
 
-			R.ExpectTrue((ResultsView[Index] - ExpectedVal).Size() < KINDA_SMALL_NUMBER);
+			EXPECT_LT((ResultsView[Index] - ExpectedVal).Size(), KINDA_SMALL_NUMBER);
 		}
 		delete SumVector;
 		delete RadialVector;
 
-		return !R.HasError();
+		
 	}
 
-
-	bool Fields_SumScalar(ExampleResponse&& R)
+	void Fields_SumScalar()
 	{
 		int32 NumPoints = 20;
 		TArray<ContextIndex> IndicesArray;
@@ -833,15 +808,14 @@ namespace GeometryCollectionExample
 			//	SamplesArray[Index].X, SamplesArray[Index].Y, SamplesArray[Index].Z,
 			//	ExpectedVal, ResultsView[Index]);
 
-			R.ExpectTrue((ResultsView[Index] - ExpectedVal) < KINDA_SMALL_NUMBER);
+			EXPECT_LT((ResultsView[Index] - ExpectedVal), KINDA_SMALL_NUMBER);
 		}
 		delete SumScalar;
 
-		return !R.HasError();
+		
 	}
 
-
-	bool Fields_SumScalarRightSide(ExampleResponse&& R)
+	void Fields_SumScalarRightSide()
 	{
 		int32 NumPoints = 20;
 		TArray<ContextIndex> IndicesArray;
@@ -902,16 +876,15 @@ namespace GeometryCollectionExample
 			//	SamplesArray[Index].X, SamplesArray[Index].Y, SamplesArray[Index].Z,
 			//	ExpectedVal, ResultsView[Index]);
 
-			R.ExpectTrue((ResultsView[Index] - ExpectedVal) < KINDA_SMALL_NUMBER);
+			EXPECT_LT((ResultsView[Index] - ExpectedVal), KINDA_SMALL_NUMBER);
 		}
 		delete SumScalar;
 		delete RadialFalloff2;
 
-		return !R.HasError();
+		
 	}
 
-
-	bool Fields_SumScalarLeftSide(ExampleResponse&& R)
+	void Fields_SumScalarLeftSide()
 	{
 		int32 NumPoints = 20;
 		TArray<ContextIndex> IndicesArray;
@@ -973,18 +946,15 @@ namespace GeometryCollectionExample
 			//	SamplesArray[Index].X, SamplesArray[Index].Y, SamplesArray[Index].Z,
 			//	ExpectedVal, ResultsView[Index]);
 
-			R.ExpectTrue((ResultsView[Index] - ExpectedVal) < KINDA_SMALL_NUMBER);
+			EXPECT_LT((ResultsView[Index] - ExpectedVal), KINDA_SMALL_NUMBER);
 		}
 		delete SumScalar;
 		delete RadialFalloff;
 
-		return !R.HasError();
+		
 	}
 
-
-
-
-	bool Fields_Culling(ExampleResponse&& R)
+	void Fields_Culling()
 	{
 		int32 NumPoints = 20;
 		TArray<ContextIndex> IndicesArray;
@@ -1047,14 +1017,11 @@ namespace GeometryCollectionExample
 			//	SamplesArray[Index].X, SamplesArray[Index].Y, SamplesArray[Index].Z,
 			//	ExpectedVal, ResultsView[Index]);
 
-			R.ExpectTrue((ResultsView[Index] - ExpectedVal) < KINDA_SMALL_NUMBER);
+			EXPECT_LT((ResultsView[Index] - ExpectedVal), KINDA_SMALL_NUMBER);
 		}
 		delete CullingField;
-
-		return !R.HasError();
 	}
 
-
 	FFieldSystemCommand SaveAndLoad(FFieldSystemCommand & CommandOut)
 	{
 		FBufferArchive Ar;
@@ -1078,67 +1045,98 @@ namespace GeometryCollectionExample
 		return CommandIn;
 	}
 
+	// Helper for testing FFieldSystemCommand equality. Defined very similarly to the == operator but with more granular
+	// checks at each stage of testing. 
+	bool TestFFieldSystemEquality(const FFieldSystemCommand& CommandIn, const FFieldSystemCommand& CommandOut)
+	{
+		EXPECT_TRUE(CommandOut.TargetAttribute.IsEqual(CommandIn.TargetAttribute));
+		if (CommandOut.TargetAttribute.IsEqual(CommandIn.TargetAttribute))
+		{
+			EXPECT_EQ(CommandOut.RootNode.IsValid(), CommandIn.RootNode.IsValid());
+			if (CommandOut.RootNode.IsValid() == CommandIn.RootNode.IsValid())
+			{
+				if (CommandOut.RootNode.IsValid())
+				{
+					EXPECT_EQ(CommandOut.RootNode->SerializationType(), CommandIn.RootNode->SerializationType());
+					if (CommandOut.RootNode->SerializationType() == CommandIn.RootNode->SerializationType())
+					{
+						EXPECT_TRUE(CommandOut.RootNode->operator==(*CommandIn.RootNode));
+						if (CommandOut.RootNode->operator==(*CommandIn.RootNode))
+						{
+							return true;
+						}
+					}
+				}
+				else
+				{
+					return true;
+				}
+			}
+		}
+		return false;
+	}
+
 	template<class T>
-	bool Fields_SerializeAPI(ExampleResponse&& R)
+	void Fields_SerializeAPI()
 	{
 		{
 			FFieldSystemCommand CommandOut("FUniformInteger", new FUniformInteger(3));
 			FFieldSystemCommand CommandIn = SaveAndLoad(CommandOut);
-			R.ExpectTrue(CommandIn == CommandOut);
+			EXPECT_TRUE(TestFFieldSystemEquality(CommandIn, CommandOut));
 		}
 
 		{
 			FFieldSystemCommand CommandOut("FRadialIntMask", new FRadialIntMask(1.f,FVector(3,5,7),11,13) );
 			FFieldSystemCommand CommandIn = SaveAndLoad(CommandOut);
-			R.ExpectTrue(CommandIn == CommandOut);
+			EXPECT_TRUE(TestFFieldSystemEquality(CommandIn, CommandOut));
 		}
 
 		{
 			FFieldSystemCommand CommandOut("FUniformScalar", new FUniformScalar(13));
 			FFieldSystemCommand CommandIn = SaveAndLoad(CommandOut);
-			R.ExpectTrue(CommandIn == CommandOut);
+			EXPECT_TRUE(TestFFieldSystemEquality(CommandIn, CommandOut));
 		}
 
 		{
 			FFieldSystemCommand CommandOut("FRadialFalloff", new FRadialFalloff(1.f,3.f,5.f,7.f,11.f,FVector(13, 17, 19)));
 			FFieldSystemCommand CommandIn = SaveAndLoad(CommandOut);
-			R.ExpectTrue(CommandIn == CommandOut);
+			EXPECT_TRUE(TestFFieldSystemEquality(CommandIn, CommandOut));
 		}
 
 		{
 			FFieldSystemCommand CommandOut("FPlaneFalloff", new FPlaneFalloff(1.f, 3.f, 5.f, 7.f, 100.f, FVector(9, 11, 13), FVector(17, 19, 23)));
 			FFieldSystemCommand CommandIn = SaveAndLoad(CommandOut);
-			R.ExpectTrue(CommandIn == CommandOut);
+			EXPECT_TRUE(TestFFieldSystemEquality(CommandIn, CommandOut));
 		}
 		
 		{
 			FFieldSystemCommand CommandOut("FBoxFalloff", new FBoxFalloff(1.f, 7.f, 9.f, 13.f, FTransform::Identity));
 			FFieldSystemCommand CommandIn = SaveAndLoad(CommandOut);
-			R.ExpectTrue(CommandIn == CommandOut);
+			EXPECT_TRUE(TestFFieldSystemEquality(CommandIn, CommandOut));
 		}
 		
 		{
 			FFieldSystemCommand CommandOut("FNoiseField", new	FNoiseField(1.f, 3.f));
 			FFieldSystemCommand CommandIn = SaveAndLoad(CommandOut);
-			R.ExpectTrue(CommandIn == CommandOut);
+			EXPECT_TRUE(TestFFieldSystemEquality(CommandIn, CommandOut));
 		}
 
 		{
 			FFieldSystemCommand CommandOut("FUniformVector", new FUniformVector(1.f, FVector(3, 5, 7)));
 			FFieldSystemCommand CommandIn = SaveAndLoad(CommandOut);
-			R.ExpectTrue(CommandIn == CommandOut);
+			EXPECT_TRUE(TestFFieldSystemEquality(CommandIn, CommandOut));
 		}
 
 		{
 			FFieldSystemCommand CommandOut("FRadialVector", new FRadialVector(1.f,FVector(3, 5, 7)));
 			FFieldSystemCommand CommandIn = SaveAndLoad(CommandOut);
-			R.ExpectTrue(CommandIn == CommandOut);
+			EXPECT_TRUE(TestFFieldSystemEquality(CommandIn, CommandOut));
 		}
 
 		{
 			FFieldSystemCommand CommandOut("FRandomVector", new FRandomVector(1.f));
 			FFieldSystemCommand CommandIn = SaveAndLoad(CommandOut);
-			R.ExpectTrue(CommandIn == CommandOut);
+			EXPECT_TRUE(TestFFieldSystemEquality(CommandIn, CommandOut));
 		}
 
 		{
@@ -1146,7 +1144,7 @@ namespace GeometryCollectionExample
 			FRadialFalloff* RadialScalar = new FRadialFalloff(1.f, 3.f, 5.f, 7.f, 11.f, FVector(13, 17, 19));
 			FFieldSystemCommand CommandOut("FSumScalar", new FSumScalar(1.f, UniformScalar, RadialScalar, EFieldOperationType::Field_Substract));
 			FFieldSystemCommand CommandIn = SaveAndLoad(CommandOut);
-			R.ExpectTrue(CommandIn == CommandOut);
+			EXPECT_TRUE(TestFFieldSystemEquality(CommandIn, CommandOut));
 		}
 
 		{
@@ -1155,7 +1153,7 @@ namespace GeometryCollectionExample
 			FRadialVector* RadialVector = new FRadialVector(21.f, FVector(3, 5, 7));
 			FFieldSystemCommand CommandOut("FSumVector", new FSumVector(1.f, UniformScalar, UniformVector, RadialVector, EFieldOperationType::Field_Divide));
 			FFieldSystemCommand CommandIn = SaveAndLoad(CommandOut);
-			R.ExpectTrue(CommandIn == CommandOut);
+			EXPECT_TRUE(TestFFieldSystemEquality(CommandIn, CommandOut));
 		}
 
 		// conversion fields
@@ -1163,13 +1161,13 @@ namespace GeometryCollectionExample
 			FUniformScalar* UniformScalar = new FUniformScalar(41.f);
 			FFieldSystemCommand CommandOut("FConversionField", new FConversionField<float,int32>(UniformScalar));
 			FFieldSystemCommand CommandIn = SaveAndLoad(CommandOut);
-			R.ExpectTrue(CommandIn == CommandOut);
+			EXPECT_TRUE(TestFFieldSystemEquality(CommandIn, CommandOut));
 		}
 		{
 			FUniformInteger* UniformInteger = new FUniformInteger(3);
 			FFieldSystemCommand CommandOut("FConversionField", new FConversionField<int32, float>(UniformInteger));
 			FFieldSystemCommand CommandIn = SaveAndLoad(CommandOut);
-			R.ExpectTrue(CommandIn == CommandOut);
+			EXPECT_TRUE(TestFFieldSystemEquality(CommandIn, CommandOut));
 		}
 
 		// culling fields
@@ -1178,38 +1176,38 @@ namespace GeometryCollectionExample
 			FRadialFalloff* RadialScalar = new FRadialFalloff(1.f, 3.f, 5.f, 7.f, 11.f, FVector(13, 17, 19));
 			FFieldSystemCommand CommandOut("FCullingField", new FCullingField<int32>(RadialScalar,UniformInteger));
 			FFieldSystemCommand CommandIn = SaveAndLoad(CommandOut);
-			R.ExpectTrue(CommandIn == CommandOut);
+			EXPECT_TRUE(TestFFieldSystemEquality(CommandIn, CommandOut));
 		}
 		{
 			FUniformScalar* UniformScalar = new FUniformScalar(3.f);
 			FRadialFalloff* RadialScalar = new FRadialFalloff(1.f, 3.f, 5.f, 7.f, 11.f, FVector(13, 17, 19));
 			FFieldSystemCommand CommandOut("FCullingField", new FCullingField<float>(RadialScalar, UniformScalar));
 			FFieldSystemCommand CommandIn = SaveAndLoad(CommandOut);
-			R.ExpectTrue(CommandIn == CommandOut);
+			EXPECT_TRUE(TestFFieldSystemEquality(CommandIn, CommandOut));
 		}
 		{
 			FUniformVector* UniformVector = new FUniformVector(3.f);
 			FRadialFalloff* RadialScalar = new FRadialFalloff(1.f, 3.f, 5.f, 7.f, 11.f, FVector(13, 17, 19));
 			FFieldSystemCommand CommandOut("FCullingField", new FCullingField<FVector>(RadialScalar, UniformVector));
 			FFieldSystemCommand CommandIn = SaveAndLoad(CommandOut);
-			R.ExpectTrue(CommandIn == CommandOut);
+			EXPECT_TRUE(TestFFieldSystemEquality(CommandIn, CommandOut));
 		}
 
 		// terminals
 		{
 			FFieldSystemCommand CommandOut("FReturnResultsTerminal", new FReturnResultsTerminal<int32>());
 			FFieldSystemCommand CommandIn = SaveAndLoad(CommandOut);
-			R.ExpectTrue(CommandIn == CommandOut);
+			EXPECT_TRUE(TestFFieldSystemEquality(CommandIn, CommandOut));
 		}
 		{
 			FFieldSystemCommand CommandOut("FReturnResultsTerminal", new FReturnResultsTerminal<float>());
 			FFieldSystemCommand CommandIn = SaveAndLoad(CommandOut);
-			R.ExpectTrue(CommandIn == CommandOut);
+			EXPECT_TRUE(TestFFieldSystemEquality(CommandIn, CommandOut));
 		}
 		{
 			FFieldSystemCommand CommandOut("FReturnResultsTerminal", new FReturnResultsTerminal<FVector>());
 			FFieldSystemCommand CommandIn = SaveAndLoad(CommandOut);
-			R.ExpectTrue(CommandIn == CommandOut);
+			EXPECT_TRUE(TestFFieldSystemEquality(CommandIn, CommandOut));
 		}
 
 
@@ -1253,13 +1251,11 @@ namespace GeometryCollectionExample
 
 			FFieldSystemCommand CommandOut("DeepTreeOfEverything", SumVector3);
 			FFieldSystemCommand CommandIn = SaveAndLoad(CommandOut);
-			R.ExpectTrue(CommandIn == CommandOut);
+			EXPECT_TRUE(TestFFieldSystemEquality(CommandIn, CommandOut));
 		}
 
-		return !R.HasError();
-
 	}
-	template bool Fields_SerializeAPI<float>(ExampleResponse&& R);
+	template void Fields_SerializeAPI<float>();
 
 
 }
diff --git a/Engine/Plugins/Experimental/GeometryCollectionPlugin/Source/GeometryCollectionExampleCore/Private/GeometryCollection/GeometryCollectionExampleImplicitCapsule.cpp b/Engine/Plugins/Experimental/GeometryCollectionPlugin/Source/GeometryCollectionExampleCore/Private/GeometryCollection/GeometryCollectionExampleImplicitCapsule.cpp
--- a/Engine/Plugins/Experimental/GeometryCollectionPlugin/Source/GeometryCollectionExampleCore/Private/GeometryCollection/GeometryCollectionExampleImplicitCapsule.cpp
+++ b/Engine/Plugins/Experimental/GeometryCollectionPlugin/Source/GeometryCollectionExampleCore/Private/GeometryCollection/GeometryCollectionExampleImplicitCapsule.cpp
@@ -7,58 +7,58 @@
 namespace GeometryCollectionExample
 {
 	template <class T>
-	void RunTestComputeSamplePoints(ExampleResponse& R, const Chaos::TCapsule<T> &Capsule)
+	void RunTestComputeSamplePoints(const Chaos::TCapsule<T> &Capsule)
 	{
 		const Chaos::TCapsule<T> OACapsule = Chaos::TCapsule<T>::NewFromOriginAndAxis(Capsule.GetOrigin(), Capsule.GetAxis(), Capsule.GetHeight(), Capsule.GetRadius());
-		R.ExpectTrue((Capsule.GetOrigin() - OACapsule.GetOrigin()).Size() < KINDA_SMALL_NUMBER, FString("Capsule != OACapsule, origin."));
-		R.ExpectTrue((Capsule.GetInsertion() - OACapsule.GetInsertion()).Size() < KINDA_SMALL_NUMBER, FString("Capsule != OACapsule, insertion."));
-		R.ExpectTrue((Capsule.GetAxis() - OACapsule.GetAxis()).Size() < KINDA_SMALL_NUMBER, FString("Capsule != OACapsule, axis."));
-		R.ExpectTrue(Capsule.GetHeight() - OACapsule.GetHeight() < KINDA_SMALL_NUMBER, FString("Capsule != OACapsule, height."));
-		R.ExpectTrue(Capsule.GetRadius() - OACapsule.GetRadius() < KINDA_SMALL_NUMBER, FString("Capsule != OACapsule, radius."));
+		EXPECT_NEAR((Capsule.GetOrigin() - OACapsule.GetOrigin()).Size(), 0, KINDA_SMALL_NUMBER) << *FString("Capsule != OACapsule, origin.");
+		EXPECT_NEAR((Capsule.GetInsertion() - OACapsule.GetInsertion()).Size(), 0, KINDA_SMALL_NUMBER) << *FString("Capsule != OACapsule, insertion.");
+		EXPECT_NEAR((Capsule.GetAxis() - OACapsule.GetAxis()).Size(), 0, KINDA_SMALL_NUMBER) << *FString("Capsule != OACapsule, axis.");
+		EXPECT_NEAR(Capsule.GetHeight() - OACapsule.GetHeight(), 0, KINDA_SMALL_NUMBER) << *FString("Capsule != OACapsule, height.");
+		EXPECT_NEAR(Capsule.GetRadius() - OACapsule.GetRadius(), 0, KINDA_SMALL_NUMBER) << *FString("Capsule != OACapsule, radius.");
 
 		Chaos::TVector<T, 3> Point;
 		T Phi;
-		R.ExpectTrue(Capsule.GetType() == Chaos::ImplicitObjectType::Capsule, FString("Implicit object type is not 'capsule'."));
+		EXPECT_EQ(Capsule.GetType(), Chaos::ImplicitObjectType::Capsule) << *FString("Implicit object type is not 'capsule'.");
 
 		Point = Capsule.GetAxis();
-		R.ExpectTrue(FMath::Abs(Point.Size() - 1.0) < KINDA_SMALL_NUMBER, FString("Capsule axis is not unit length."));
+		EXPECT_LT(FMath::Abs(Point.Size() - 1.0), KINDA_SMALL_NUMBER) << *FString("Capsule axis is not unit length.");
 		
 		Point = Capsule.GetOrigin() + Capsule.GetAxis() * (Capsule.GetHeight() + Capsule.GetRadius() * 2);
-		R.ExpectTrue(((Point - Capsule.GetInsertion()).Size() < KINDA_SMALL_NUMBER), FString("Capsule is broken."));
+		EXPECT_LT((Point - Capsule.GetInsertion()).Size(), KINDA_SMALL_NUMBER) << *FString("Capsule is broken.");
 
 		Point = Capsule.GetInsertion();// Capsule.GetOrigin() + Capsule.GetAxis() * Capsule.GetHeight();
 		Phi = Capsule.SignedDistance(Point);
-		R.ExpectTrue(Phi <= KINDA_SMALL_NUMBER, FString("Capsule failed phi surface (insertion) sanity test."));
+		EXPECT_NEAR(Phi, 0, KINDA_SMALL_NUMBER) << *FString("Capsule failed phi surface (insertion) sanity test.");
 
 		Point = Capsule.GetOrigin() + Capsule.GetAxis() * (Capsule.GetHeight() + Capsule.GetRadius() * 2) * 0.25;
 		Phi = Capsule.SignedDistance(Point);
-		R.ExpectTrue(Phi <= (T)0.0, FString("Capsule failed phi depth (1/4 origin) sanity test."));
+		EXPECT_LE(Phi, (T)0.0) << *FString("Capsule failed phi depth (1/4 origin) sanity test.");
 
 		Point = Capsule.GetOrigin() + Capsule.GetAxis() * (Capsule.GetHeight() + Capsule.GetRadius() * 2) * 0.75;
 		Phi = Capsule.SignedDistance(Point);
-		R.ExpectTrue(Phi <= (T)0.0, FString("Capsule failed phi depth (3/4 origin) sanity test."));
+		EXPECT_LE(Phi, (T)0.0) << *FString("Capsule failed phi depth (3/4 origin) sanity test.");
 
-		R.ExpectTrue((Capsule.GetCenter() - Chaos::TVector<T,3>(Capsule.GetOrigin() + Capsule.GetAxis() * (Capsule.GetHeight() + Capsule.GetRadius() * 2) * 0.5)).Size() <= KINDA_SMALL_NUMBER, FString("Capsule center is off mid axis."));
+		EXPECT_NEAR((Capsule.GetCenter() - Chaos::TVector<T,3>(Capsule.GetOrigin() + Capsule.GetAxis() * (Capsule.GetHeight() + Capsule.GetRadius() * 2) * 0.5)).Size(), 0, KINDA_SMALL_NUMBER) << *FString("Capsule center is off mid axis.");
 
 		Point = Capsule.GetCenter();
 		Phi = Capsule.SignedDistance(Point);
-		R.ExpectTrue(Phi < (T)0.0, FString("Capsule failed phi depth sanity test."));
+		EXPECT_LT(Phi, (T)0.0) << *FString("Capsule failed phi depth sanity test.");
 		
 		Point = Capsule.GetOrigin();
 		Phi = Capsule.SignedDistance(Point);
-		R.ExpectTrue(FMath::Abs(Phi) <= KINDA_SMALL_NUMBER, FString("Capsule failed phi surface (origin) sanity test."));
+		EXPECT_NEAR(FMath::Abs(Phi), 0, KINDA_SMALL_NUMBER) << *FString("Capsule failed phi surface (origin) sanity test.");
 
 		Point = Capsule.GetInsertion();
 		Phi = Capsule.SignedDistance(Point);
-		R.ExpectTrue(FMath::Abs(Phi) <= KINDA_SMALL_NUMBER, FString("Capsule failed phi surface (origin+axis*height) sanity test."));
+		EXPECT_NEAR(FMath::Abs(Phi), 0, KINDA_SMALL_NUMBER) << *FString("Capsule failed phi surface (origin+axis*height) sanity test.");
 
 		Point = Capsule.GetOrigin() + Capsule.GetAxis() * Capsule.GetRadius() + Capsule.GetAxis().GetOrthogonalVector().GetSafeNormal() * Capsule.GetRadius();
 		Phi = Capsule.SignedDistance(Point);
-		R.ExpectTrue(FMath::Abs(Phi) <= KINDA_SMALL_NUMBER, FString("Capsule failed phi surface (origin+orthogonalAxis*radius) sanity test."));
+		EXPECT_NEAR(FMath::Abs(Phi), 0, KINDA_SMALL_NUMBER) << *FString("Capsule failed phi surface (origin+orthogonalAxis*radius) sanity test.");
 
 		Point = Capsule.GetCenter() + Capsule.GetAxis().GetOrthogonalVector().GetSafeNormal() * Capsule.GetRadius();
 		Phi = Capsule.SignedDistance(Point);
-		R.ExpectTrue(FMath::Abs(Phi) <= KINDA_SMALL_NUMBER, FString("Capsule failed phi surface (center+orthogonalAxis*radius) sanity test."));
+		EXPECT_NEAR(FMath::Abs(Phi), 0, KINDA_SMALL_NUMBER) << *FString("Capsule failed phi surface (center+orthogonalAxis*radius) sanity test.");
 
 		TArray<Chaos::TVector<T, 3>> Points = Capsule.ComputeSamplePoints(100);
 		check(Points.Num() == 100);
@@ -75,17 +75,17 @@ namespace GeometryCollectionExample
 			
 			const bool Differs = Pt != Point;
 			check(Differs);
-			R.ExpectTrue(Differs, FString("Produced a redundant value."));
+			EXPECT_TRUE(Differs) << *FString("Produced a redundant value.");
 			Point = Pt;
 		}
 
 		const bool OnSurface = FMath::Abs(MinPhi) <= KINDA_SMALL_NUMBER && FMath::Abs(MaxPhi) <= KINDA_SMALL_NUMBER;
 		check(OnSurface);
-		R.ExpectTrue(OnSurface, FString("Produced a point not on the surface of the capsule."));
+		EXPECT_TRUE(OnSurface) << *FString("Produced a point not on the surface of the capsule.");
 	}
 
 	template <class T>
-	void TestComputeSamplePoints_Capsule(ExampleResponse& R)
+	void TestComputeSamplePoints_Capsule()
 	{
 		//
 		// Height == 1
@@ -94,32 +94,32 @@ namespace GeometryCollectionExample
 		// At the origin with radius 1
 		{
 			Chaos::TCapsule<T> Capsule(Chaos::TVector<T, 3>(0,0,0), Chaos::TVector<T, 3>(0,0,1), (T)1.0);
-			RunTestComputeSamplePoints(R, Capsule);
+			RunTestComputeSamplePoints(Capsule);
 		}
 		// At the origin with radius > 1
 		{
 			Chaos::TCapsule<T> Capsule(Chaos::TVector<T, 3>(0,0,0), Chaos::TVector<T, 3>(0,0,1), (T)10.0);
-			RunTestComputeSamplePoints(R, Capsule);
+			RunTestComputeSamplePoints(Capsule);
 		}
 		// At the origin with radius < 1
 		{
 			Chaos::TCapsule<T> Capsule(Chaos::TVector<T, 3>(0,0,0), Chaos::TVector<T, 3>(0,0,1), (T)0.1);
-			RunTestComputeSamplePoints(R, Capsule);
+			RunTestComputeSamplePoints(Capsule);
 		}
 		// Off the origin with radius 1
 		{
 			Chaos::TCapsule<T> Capsule(Chaos::TVector<T, 3>(10,10,10), Chaos::TVector<T, 3>(10,10,11), (T)1.0);
-			RunTestComputeSamplePoints(R, Capsule);
+			RunTestComputeSamplePoints(Capsule);
 		}
 		// Off the origin with radius > 1
 		{
 			Chaos::TCapsule<T> Capsule(Chaos::TVector<T, 3>(10,10,10), Chaos::TVector<T, 3>(10,10,11), (T)10.0);
-			RunTestComputeSamplePoints(R, Capsule);
+			RunTestComputeSamplePoints(Capsule);
 		}
 		// Off the origin with radius < 1
 		{
 			Chaos::TCapsule<T> Capsule(Chaos::TVector<T, 3>(10,10,10), Chaos::TVector<T, 3>(10,10,11), (T)0.1);
-			RunTestComputeSamplePoints(R, Capsule);
+			RunTestComputeSamplePoints(Capsule);
 		}
 
 		//
@@ -129,32 +129,32 @@ namespace GeometryCollectionExample
 		// At the origin with radius 1
 		{
 			Chaos::TCapsule<T> Capsule(Chaos::TVector<T, 3>(0,0,0), Chaos::TVector<T, 3>(0,0,10), (T)1.0);
-			RunTestComputeSamplePoints(R, Capsule);
+			RunTestComputeSamplePoints(Capsule);
 		}
 		// At the origin with radius > 1
 		{
 			Chaos::TCapsule<T> Capsule(Chaos::TVector<T, 3>(0,0,0), Chaos::TVector<T, 3>(0,0,10), (T)10.0);
-			RunTestComputeSamplePoints(R, Capsule);
+			RunTestComputeSamplePoints(Capsule);
 		}
 		// At the origin with radius < 1
 		{
 			Chaos::TCapsule<T> Capsule(Chaos::TVector<T, 3>(0,0,0), Chaos::TVector<T, 3>(0,0,10), (T)0.1);
-			RunTestComputeSamplePoints(R, Capsule);
+			RunTestComputeSamplePoints(Capsule);
 		}
 		// Off the origin with radius 1
 		{
 			Chaos::TCapsule<T> Capsule(Chaos::TVector<T, 3>(10,10,10), Chaos::TVector<T, 3>(10,10,21), (T)1.0);
-			RunTestComputeSamplePoints(R, Capsule);
+			RunTestComputeSamplePoints(Capsule);
 		}
 		// Off the origin with radius > 1
 		{
 			Chaos::TCapsule<T> Capsule(Chaos::TVector<T, 3>(10,10,10), Chaos::TVector<T, 3>(10,10,21), (T)10.0);
-			RunTestComputeSamplePoints(R, Capsule);
+			RunTestComputeSamplePoints(Capsule);
 		}
 		// Off the origin with radius < 1
 		{
 			Chaos::TCapsule<T> Capsule(Chaos::TVector<T, 3>(10,10,10), Chaos::TVector<T, 3>(10,10,21), (T)0.1);
-			RunTestComputeSamplePoints(R, Capsule);
+			RunTestComputeSamplePoints(Capsule);
 		}
 
 		// 
@@ -164,40 +164,39 @@ namespace GeometryCollectionExample
 		// At the origin with radius 1
 		{
 			Chaos::TCapsule<T> Capsule(Chaos::TVector<T, 3>(0,0,0), Chaos::TVector<T, 3>(1,1,1), (T)1.0);
-			RunTestComputeSamplePoints(R, Capsule);
+			RunTestComputeSamplePoints(Capsule);
 		}
 		// At the origin with radius > 1
 		{
 			Chaos::TCapsule<T> Capsule(Chaos::TVector<T, 3>(0,0,0), Chaos::TVector<T, 3>(1,1,1), (T)10.0);
-			RunTestComputeSamplePoints(R, Capsule);
+			RunTestComputeSamplePoints(Capsule);
 		}
 		// At the origin with radius < 1
 		{
 			Chaos::TCapsule<T> Capsule(Chaos::TVector<T, 3>(0,0,0), Chaos::TVector<T, 3>(1,1,1), (T)0.1);
-			RunTestComputeSamplePoints(R, Capsule);
+			RunTestComputeSamplePoints(Capsule);
 		}
 		// Off the origin with radius 1
 		{
 			Chaos::TCapsule<T> Capsule(Chaos::TVector<T, 3>(10,10,10), Chaos::TVector<T, 3>(11,11,11), (T)1.0);
-			RunTestComputeSamplePoints(R, Capsule);
+			RunTestComputeSamplePoints(Capsule);
 		}
 		// Off the origin with radius > 1
 		{
 			Chaos::TCapsule<T> Capsule(Chaos::TVector<T, 3>(10,10,10), Chaos::TVector<T, 3>(11,11,11), (T)10.0);
-			RunTestComputeSamplePoints(R, Capsule);
+			RunTestComputeSamplePoints(Capsule);
 		}
 		// Off the origin with radius < 1
 		{
 			Chaos::TCapsule<T> Capsule(Chaos::TVector<T, 3>(10,10,10), Chaos::TVector<T, 3>(11,11,11), (T)0.1);
-			RunTestComputeSamplePoints(R, Capsule);
+			RunTestComputeSamplePoints(Capsule);
 		}
 	}
 
 	template <class T>
-	bool TestImplicitCapsule(ExampleResponse&& R)
+	void TestImplicitCapsule()
 	{
-		TestComputeSamplePoints_Capsule<T>(R);
-		return !R.HasError();
+		TestComputeSamplePoints_Capsule<T>();
 	}
-	template bool TestImplicitCapsule<float>(ExampleResponse&& R);
+	template void TestImplicitCapsule<float>();
 } // namespace GeometryCollectionExample
\ No newline at end of file
diff --git a/Engine/Plugins/Experimental/GeometryCollectionPlugin/Source/GeometryCollectionExampleCore/Private/GeometryCollection/GeometryCollectionExampleImplicitCylinder.cpp b/Engine/Plugins/Experimental/GeometryCollectionPlugin/Source/GeometryCollectionExampleCore/Private/GeometryCollection/GeometryCollectionExampleImplicitCylinder.cpp
--- a/Engine/Plugins/Experimental/GeometryCollectionPlugin/Source/GeometryCollectionExampleCore/Private/GeometryCollection/GeometryCollectionExampleImplicitCylinder.cpp
+++ b/Engine/Plugins/Experimental/GeometryCollectionPlugin/Source/GeometryCollectionExampleCore/Private/GeometryCollection/GeometryCollectionExampleImplicitCylinder.cpp
@@ -7,51 +7,51 @@
 namespace GeometryCollectionExample
 {
 	template <class T>
-	void RunTestComputeSamplePoints(ExampleResponse& R, const Chaos::TCylinder<T> &Cylinder)
+	void RunTestComputeSamplePoints(const Chaos::TCylinder<T> &Cylinder)
 	{
 		Chaos::TVector<T, 3> Point;
 		T Phi;
-		R.ExpectTrue(Cylinder.GetType() == Chaos::ImplicitObjectType::Cylinder, FString("Implicit object type is not 'cylinder'."));
+		EXPECT_EQ(Cylinder.GetType(), Chaos::ImplicitObjectType::Cylinder) << *FString("Implicit object type is not 'cylinder'.");
 
 		Point = Cylinder.GetAxis();
-		R.ExpectTrue(FMath::Abs(Point.Size() - 1.0) < KINDA_SMALL_NUMBER, FString("Cylinder axis is not unit length."));
+		EXPECT_LT(FMath::Abs(Point.Size() - 1.0), KINDA_SMALL_NUMBER) << *FString("Cylinder axis is not unit length.");
 		
 		Point = Cylinder.GetOrigin() + Cylinder.GetAxis() * Cylinder.GetHeight();
-		R.ExpectTrue(((Point - Cylinder.GetInsertion()).Size() < KINDA_SMALL_NUMBER), FString("Cylinder is broken."));
+		EXPECT_LT((Point - Cylinder.GetInsertion()).Size(), KINDA_SMALL_NUMBER) << *FString("Cylinder is broken.");
 
 		Point = Cylinder.GetInsertion();// Cylinder.GetOrigin() + Cylinder.GetAxis() * Cylinder.GetHeight();
 		Phi = Cylinder.SignedDistance(Point);
-		R.ExpectTrue(Phi <= KINDA_SMALL_NUMBER, FString("Cylinder failed phi surface (insertion) sanity test."));
+		EXPECT_LE(Phi, KINDA_SMALL_NUMBER) << *FString("Cylinder failed phi surface (insertion) sanity test.");
 
 		Point = Cylinder.GetOrigin() + Cylinder.GetAxis() * Cylinder.GetHeight() * 0.25;
 		Phi = Cylinder.SignedDistance(Point);
-		R.ExpectTrue(Phi <= (T)0.0, FString("Cylinder failed phi depth (1/4 origin) sanity test."));
+		EXPECT_LE(Phi, (T)0.0) << *FString("Cylinder failed phi depth (1/4 origin) sanity test.");
 
 		Point = Cylinder.GetOrigin() + Cylinder.GetAxis() * Cylinder.GetHeight() * 0.75;
 		Phi = Cylinder.SignedDistance(Point);
-		R.ExpectTrue(Phi <= (T)0.0, FString("Cylinder failed phi depth (3/4 origin) sanity test."));
+		EXPECT_LE(Phi, (T)0.0) << *FString("Cylinder failed phi depth (3/4 origin) sanity test.");
 
-		R.ExpectTrue((Cylinder.GetCenter() - Chaos::TVector<T,3>(Cylinder.GetOrigin() + Cylinder.GetAxis() * Cylinder.GetHeight() * 0.5)).Size() <= KINDA_SMALL_NUMBER, FString("Cylinder center is off mid axis."));
+		EXPECT_LE((Cylinder.GetCenter() - Chaos::TVector<T,3>(Cylinder.GetOrigin() + Cylinder.GetAxis() * Cylinder.GetHeight() * 0.5)).Size(), KINDA_SMALL_NUMBER) << *FString("Cylinder center is off mid axis.");
 
 		Point = Cylinder.GetCenter();
 		Phi = Cylinder.SignedDistance(Point);
-		R.ExpectTrue(Phi < (T)0.0, FString("Cylinder failed phi depth sanity test."));
+		EXPECT_LT(Phi, (T)0.0) << *FString("Cylinder failed phi depth sanity test.");
 		
 		Point = Cylinder.GetOrigin();
 		Phi = Cylinder.SignedDistance(Point);
-		R.ExpectTrue(FMath::Abs(Phi) <= KINDA_SMALL_NUMBER, FString("Cylinder failed phi surface (origin) sanity test."));
+		EXPECT_LE(FMath::Abs(Phi), KINDA_SMALL_NUMBER) << *FString("Cylinder failed phi surface (origin) sanity test.");
 
 		Point = Cylinder.GetOrigin() + Cylinder.GetAxis() * Cylinder.GetHeight();
 		Phi = Cylinder.SignedDistance(Point);
-		R.ExpectTrue(FMath::Abs(Phi) <= KINDA_SMALL_NUMBER, FString("Cylinder failed phi surface (origin+axis*height) sanity test."));
+		EXPECT_LE(FMath::Abs(Phi), KINDA_SMALL_NUMBER) << *FString("Cylinder failed phi surface (origin+axis*height) sanity test.");
 
 		Point = Cylinder.GetOrigin() + Cylinder.GetAxis().GetOrthogonalVector().GetSafeNormal() * Cylinder.GetRadius();
 		Phi = Cylinder.SignedDistance(Point);
-		R.ExpectTrue(FMath::Abs(Phi) <= KINDA_SMALL_NUMBER, FString("Cylinder failed phi surface (origin+orthogonalAxis*radius) sanity test."));
+		EXPECT_LE(FMath::Abs(Phi), KINDA_SMALL_NUMBER) << *FString("Cylinder failed phi surface (origin+orthogonalAxis*radius) sanity test.");
 
 		Point = Cylinder.GetCenter() + Cylinder.GetAxis().GetOrthogonalVector().GetSafeNormal() * Cylinder.GetRadius();
 		Phi = Cylinder.SignedDistance(Point);
-		R.ExpectTrue(FMath::Abs(Phi) <= KINDA_SMALL_NUMBER, FString("Cylinder failed phi surface (center+orthogonalAxis*radius) sanity test."));
+		EXPECT_LE(FMath::Abs(Phi), KINDA_SMALL_NUMBER) << *FString("Cylinder failed phi surface (center+orthogonalAxis*radius) sanity test.");
 
 		TArray<Chaos::TVector<T, 3>> Points = Cylinder.ComputeSamplePoints(100);
 		check(Points.Num() == 100);
@@ -68,17 +68,17 @@ namespace GeometryCollectionExample
 			
 			const bool Differs = Pt != Point;
 			check(Differs);
-			R.ExpectTrue(Differs, FString("Produced a redundant value."));
+			EXPECT_TRUE(Differs) << *FString("Produced a redundant value.");
 			Point = Pt;
 		}
 
 		const bool OnSurface = FMath::Abs(MinPhi) <= KINDA_SMALL_NUMBER && FMath::Abs(MaxPhi) <= KINDA_SMALL_NUMBER;
 		check(OnSurface);
-		R.ExpectTrue(OnSurface, FString("Produced a point not on the surface of the cylinder."));
+		EXPECT_TRUE(OnSurface) << *FString("Produced a point not on the surface of the cylinder.");
 	}
 
 	template <class T>
-	void TestComputeSamplePoints_Cylinder(ExampleResponse& R)
+	void TestComputeSamplePoints_Cylinder()
 	{
 		//
 		// Height == 1
@@ -87,32 +87,32 @@ namespace GeometryCollectionExample
 		// At the origin with radius 1
 		{
 			Chaos::TCylinder<T> Cylinder(Chaos::TVector<T, 3>(0,0,0), Chaos::TVector<T, 3>(0,0,1), (T)1.0);
-			RunTestComputeSamplePoints(R, Cylinder);
+			RunTestComputeSamplePoints(Cylinder);
 		}
 		// At the origin with radius > 1
 		{
 			Chaos::TCylinder<T> Cylinder(Chaos::TVector<T, 3>(0,0,0), Chaos::TVector<T, 3>(0,0,1), (T)10.0);
-			RunTestComputeSamplePoints(R, Cylinder);
+			RunTestComputeSamplePoints(Cylinder);
 		}
 		// At the origin with radius < 1
 		{
 			Chaos::TCylinder<T> Cylinder(Chaos::TVector<T, 3>(0,0,0), Chaos::TVector<T, 3>(0,0,1), (T)0.1);
-			RunTestComputeSamplePoints(R, Cylinder);
+			RunTestComputeSamplePoints(Cylinder);
 		}
 		// Off the origin with radius 1
 		{
 			Chaos::TCylinder<T> Cylinder(Chaos::TVector<T, 3>(10,10,10), Chaos::TVector<T, 3>(10,10,11), (T)1.0);
-			RunTestComputeSamplePoints(R, Cylinder);
+			RunTestComputeSamplePoints(Cylinder);
 		}
 		// Off the origin with radius > 1
 		{
 			Chaos::TCylinder<T> Cylinder(Chaos::TVector<T, 3>(10,10,10), Chaos::TVector<T, 3>(10,10,11), (T)10.0);
-			RunTestComputeSamplePoints(R, Cylinder);
+			RunTestComputeSamplePoints(Cylinder);
 		}
 		// Off the origin with radius < 1
 		{
 			Chaos::TCylinder<T> Cylinder(Chaos::TVector<T, 3>(10,10,10), Chaos::TVector<T, 3>(10,10,11), (T)0.1);
-			RunTestComputeSamplePoints(R, Cylinder);
+			RunTestComputeSamplePoints(Cylinder);
 		}
 
 		//
@@ -122,32 +122,32 @@ namespace GeometryCollectionExample
 		// At the origin with radius 1
 		{
 			Chaos::TCylinder<T> Cylinder(Chaos::TVector<T, 3>(0,0,0), Chaos::TVector<T, 3>(0,0,10), (T)1.0);
-			RunTestComputeSamplePoints(R, Cylinder);
+			RunTestComputeSamplePoints(Cylinder);
 		}
 		// At the origin with radius > 1
 		{
 			Chaos::TCylinder<T> Cylinder(Chaos::TVector<T, 3>(0,0,0), Chaos::TVector<T, 3>(0,0,10), (T)10.0);
-			RunTestComputeSamplePoints(R, Cylinder);
+			RunTestComputeSamplePoints(Cylinder);
 		}
 		// At the origin with radius < 1
 		{
 			Chaos::TCylinder<T> Cylinder(Chaos::TVector<T, 3>(0,0,0), Chaos::TVector<T, 3>(0,0,10), (T)0.1);
-			RunTestComputeSamplePoints(R, Cylinder);
+			RunTestComputeSamplePoints(Cylinder);
 		}
 		// Off the origin with radius 1
 		{
 			Chaos::TCylinder<T> Cylinder(Chaos::TVector<T, 3>(10,10,10), Chaos::TVector<T, 3>(10,10,21), (T)1.0);
-			RunTestComputeSamplePoints(R, Cylinder);
+			RunTestComputeSamplePoints(Cylinder);
 		}
 		// Off the origin with radius > 1
 		{
 			Chaos::TCylinder<T> Cylinder(Chaos::TVector<T, 3>(10,10,10), Chaos::TVector<T, 3>(10,10,21), (T)10.0);
-			RunTestComputeSamplePoints(R, Cylinder);
+			RunTestComputeSamplePoints(Cylinder);
 		}
 		// Off the origin with radius < 1
 		{
 			Chaos::TCylinder<T> Cylinder(Chaos::TVector<T, 3>(10,10,10), Chaos::TVector<T, 3>(10,10,21), (T)0.1);
-			RunTestComputeSamplePoints(R, Cylinder);
+			RunTestComputeSamplePoints(Cylinder);
 		}
 
 		// 
@@ -157,40 +157,39 @@ namespace GeometryCollectionExample
 		// At the origin with radius 1
 		{
 			Chaos::TCylinder<T> Cylinder(Chaos::TVector<T, 3>(0,0,0), Chaos::TVector<T, 3>(1,1,1), (T)1.0);
-			RunTestComputeSamplePoints(R, Cylinder);
+			RunTestComputeSamplePoints(Cylinder);
 		}
 		// At the origin with radius > 1
 		{
 			Chaos::TCylinder<T> Cylinder(Chaos::TVector<T, 3>(0,0,0), Chaos::TVector<T, 3>(1,1,1), (T)10.0);
-			RunTestComputeSamplePoints(R, Cylinder);
+			RunTestComputeSamplePoints(Cylinder);
 		}
 		// At the origin with radius < 1
 		{
 			Chaos::TCylinder<T> Cylinder(Chaos::TVector<T, 3>(0,0,0), Chaos::TVector<T, 3>(1,1,1), (T)0.1);
-			RunTestComputeSamplePoints(R, Cylinder);
+			RunTestComputeSamplePoints(Cylinder);
 		}
 		// Off the origin with radius 1
 		{
 			Chaos::TCylinder<T> Cylinder(Chaos::TVector<T, 3>(10,10,10), Chaos::TVector<T, 3>(11,11,11), (T)1.0);
-			RunTestComputeSamplePoints(R, Cylinder);
+			RunTestComputeSamplePoints(Cylinder);
 		}
 		// Off the origin with radius > 1
 		{
 			Chaos::TCylinder<T> Cylinder(Chaos::TVector<T, 3>(10,10,10), Chaos::TVector<T, 3>(11,11,11), (T)10.0);
-			RunTestComputeSamplePoints(R, Cylinder);
+			RunTestComputeSamplePoints(Cylinder);
 		}
 		// Off the origin with radius < 1
 		{
 			Chaos::TCylinder<T> Cylinder(Chaos::TVector<T, 3>(10,10,10), Chaos::TVector<T, 3>(11,11,11), (T)0.1);
-			RunTestComputeSamplePoints(R, Cylinder);
+			RunTestComputeSamplePoints(Cylinder);
 		}
 	}
 
 	template <class T>
-	bool TestImplicitCylinder(ExampleResponse&& R)
+	void TestImplicitCylinder()
 	{
-		TestComputeSamplePoints_Cylinder<T>(R);
-		return !R.HasError();
+		TestComputeSamplePoints_Cylinder<T>();
 	}
-	template bool TestImplicitCylinder<float>(ExampleResponse&& R);
+	template void TestImplicitCylinder<float>();
 } // namespace GeometryCollectionExample
\ No newline at end of file
diff --git a/Engine/Plugins/Experimental/GeometryCollectionPlugin/Source/GeometryCollectionExampleCore/Private/GeometryCollection/GeometryCollectionExampleImplicitSphere.cpp b/Engine/Plugins/Experimental/GeometryCollectionPlugin/Source/GeometryCollectionExampleCore/Private/GeometryCollection/GeometryCollectionExampleImplicitSphere.cpp
--- a/Engine/Plugins/Experimental/GeometryCollectionPlugin/Source/GeometryCollectionExampleCore/Private/GeometryCollection/GeometryCollectionExampleImplicitSphere.cpp
+++ b/Engine/Plugins/Experimental/GeometryCollectionPlugin/Source/GeometryCollectionExampleCore/Private/GeometryCollection/GeometryCollectionExampleImplicitSphere.cpp
@@ -17,7 +17,6 @@ namespace GeometryCollectionExample
 
 	template <class T>
 	void RunTestFindClosestIntersectionHelper(
-		ExampleResponse& R, 
 		const FString TestName,
 		const Chaos::TVector<T, 3>& Center, 
 		const T Radius,
@@ -28,10 +27,10 @@ namespace GeometryCollectionExample
 	{
 		Chaos::TSphere<T, 3> Sphere(Center, Radius);
 		IntersectionResult Res = Sphere.FindClosestIntersection(StartPt, EndPt, Thickness);
-		R.ExpectTrue(
+		EXPECT_TRUE(
 			(!Res.Second && !ExpectedRes.Second) ||
-			(Res.Second && ExpectedRes.Second && (Res.First - ExpectedRes.First).Size() < SMALL_NUMBER),
-			FString::Printf(
+			(Res.Second && ExpectedRes.Second && (Res.First - ExpectedRes.First).Size() < SMALL_NUMBER)) <<
+			*FString::Printf(
 				TEXT(
 					"%s - Sphere((%g, %g, %g), %g).FindClosestIntersection("
 						"StartPt(%g, %g, %g), "
@@ -52,12 +51,11 @@ namespace GeometryCollectionExample
 				
 				ExpectedRes.Second ? 1 : 0,
 				ExpectedRes.First[0], ExpectedRes.First[1], ExpectedRes.First[2]
-			));
+			);
 	}
 
 	template <class T>
 	void RunTestFindClosestIntersection_Tangent(
-		ExampleResponse& R,
 		const Chaos::TVector<T, 3>& Center,
 		const T Radius,
 		const T Thickness)
@@ -69,13 +67,13 @@ namespace GeometryCollectionExample
 		// Fails with no collision:
 		// "RunTestFindClosestIntersection_Tangent1 - Sphere((0, 0, 0), 1).FindClosestIntersection(StartPt(0, 0, 1), EndPt(0, 0, 1), Thickness:0.1) = RESULT: 0, Point: (0, 0, 0), EXPECTED: 1 (0, 0, 1.1)."	FString
 		StartPt = EndPt = Center + Chaos::TVector<T,3>(0, 0, Radius);
-		//RunTestFindClosestIntersectionHelper(R, "RunTestFindClosestIntersection_Tangent1", Center, Radius, Thickness,
+		//RunTestFindClosestIntersectionHelper("RunTestFindClosestIntersection_Tangent1", Center, Radius, Thickness,
 		//	StartPt, EndPt, 
 		//	IntersectionResult(Center+Chaos::TVector<T,3>(0,0,Radius+Thickness), true));
 
 		StartPt = Center + Chaos::TVector<T, 3>(Radius, 0, Radius+Thickness);
 		EndPt = Center + Chaos::TVector<T, 3>(-Radius, 0, Radius+Thickness);
-		RunTestFindClosestIntersectionHelper(R, "RunTestFindClosestIntersection_Tangent2", Center, Radius, Thickness,
+		RunTestFindClosestIntersectionHelper("RunTestFindClosestIntersection_Tangent2", Center, Radius, Thickness,
 			StartPt, EndPt, 
 			IntersectionResult(Center+Chaos::TVector<T,3>(0,0,Radius+Thickness), true));
 		
@@ -86,32 +84,31 @@ namespace GeometryCollectionExample
 		// "RunTestFindClosestIntersection_Tangent3 - Sphere((0, 0, 0), 1).FindClosestIntersection(StartPt(0, 0, 2), EndPt(0, 0, 2), Thickness:1) = RESULT: 1, Point: (0, 0, 1), EXPECTED: 1 (0, 0, 2)."	FString
 		// "RunTestFindClosestIntersection_Tangent3 - Sphere((0, 0, 0), 1).FindClosestIntersection(StartPt(0, 0, 3), EndPt(0, 0, 3), Thickness:2) = RESULT: 1, Point: (0, 0, 1), EXPECTED: 1 (0, 0, 3)."	FString
 		StartPt = EndPt = Center + Chaos::TVector<T,3>(0, 0, Radius+Thickness);
-		//RunTestFindClosestIntersectionHelper(R, "RunTestFindClosestIntersection_Tangent3", Center, Radius, Thickness,
+		//RunTestFindClosestIntersectionHelper("RunTestFindClosestIntersection_Tangent3", Center, Radius, Thickness,
 		//	StartPt, EndPt, 
 		//	IntersectionResult(Center+Chaos::TVector<T,3>(0,0,Radius+Thickness), true));
 
 		StartPt = Center + Chaos::TVector<T, 3>(Radius, 0, Radius+Thickness);
 		EndPt = Center + Chaos::TVector<T, 3>(-Radius, 0, Radius+Thickness);
-		RunTestFindClosestIntersectionHelper(R, "RunTestFindClosestIntersection_Tangent4", Center, Radius, Thickness,
+		RunTestFindClosestIntersectionHelper("RunTestFindClosestIntersection_Tangent4", Center, Radius, Thickness,
 			StartPt, EndPt, 
 			IntersectionResult(Center+Chaos::TVector<T,3>(0,0,Radius+Thickness), true));
 
 		// Radius + Thickness + 1 (miss)
 		StartPt = EndPt = Center + Chaos::TVector<T,3>(0, 0, Radius+Thickness+1.0);
-		RunTestFindClosestIntersectionHelper(R, "RunTestFindClosestIntersection_Tangent5", Center, Radius, Thickness,
+		RunTestFindClosestIntersectionHelper("RunTestFindClosestIntersection_Tangent5", Center, Radius, Thickness,
 			StartPt, EndPt, 
 			IntersectionResult(Center+Chaos::TVector<T,3>(Radius+Thickness,0,0), false));
 
 		StartPt = Center + Chaos::TVector<T, 3>(Radius, 0, Radius+Thickness+1.0);
 		EndPt = Center + Chaos::TVector<T, 3>(-Radius, 0, Radius+Thickness+1.0);
-		RunTestFindClosestIntersectionHelper(R, "RunTestFindClosestIntersection_Tangent6", Center, Radius, Thickness,
+		RunTestFindClosestIntersectionHelper("RunTestFindClosestIntersection_Tangent6", Center, Radius, Thickness,
 			StartPt, EndPt, 
 			IntersectionResult(Center+Chaos::TVector<T,3>(0,0,Radius+Thickness), false));
 	}
 
 	template <class T>
 	void RunTestFindClosestIntersection_Shallow(
-		ExampleResponse& R,
 		const Chaos::TVector<T, 3>& Center,
 		const T Radius,
 		const T Thickness)
@@ -121,26 +118,26 @@ namespace GeometryCollectionExample
 
 		// Fails w/no collision:
 		// "RunTestFindClosestIntersection_Shallow1 - Sphere((0, 0, 0), 1).FindClosestIntersection(StartPt(0, 0, 0.999999), EndPt(0, 0, 0.999999), Thickness:0.1) = RESULT: 0, Point: (0, 0, 0), EXPECTED: 1 (1.1, 0...	FString
-		//RunTestFindClosestIntersectionHelper(R, "RunTestFindClosestIntersection_Shallow1", Center, Radius, Thickness,
+		//RunTestFindClosestIntersectionHelper("RunTestFindClosestIntersection_Shallow1", Center, Radius, Thickness,
 		//	StartPt, EndPt, 
 		//	IntersectionResult(Center+Chaos::TVector<T,3>(Radius+Thickness,0,0), true));
 
 		// Fails w/no collision:
 		// "RunTestFindClosestIntersection_Shallow2 - Sphere((0, 0, 0), 1).FindClosestIntersection(StartPt(0, 0, 0.999999), EndPt(0, 0, 1), Thickness:0.1) = RESULT: 0, Point: (0, 0, 0), EXPECTED: 1 (0, 0, 1.1)."	FString
 		EndPt += Chaos::TVector<T, 3>(0, 0, Radius * .000001);
-		//RunTestFindClosestIntersectionHelper(R, "RunTestFindClosestIntersection_Shallow2", Center, Radius, Thickness,
+		//RunTestFindClosestIntersectionHelper("RunTestFindClosestIntersection_Shallow2", Center, Radius, Thickness,
 		//	StartPt, EndPt, 
 		//	IntersectionResult(Center+Chaos::TVector<T,3>(0,0,Radius+Thickness), true));
 
 		// Fails w/no collision:
 		// "RunTestFindClosestIntersection_Shallow3 - Sphere((0, 0, 0), 1).FindClosestIntersection(StartPt(0, 0, 0.999999), EndPt(0, 0, 1), Thickness:0.1) = RESULT: 0, Point: (0, 0, 0), EXPECTED: 1 (0, 0, 1.1)."	FString
 		EndPt += Chaos::TVector<T, 3>(0, 0, Radius * .000001);
-		//RunTestFindClosestIntersectionHelper(R, "RunTestFindClosestIntersection_Shallow3", Center, Radius, Thickness,
+		//RunTestFindClosestIntersectionHelper("RunTestFindClosestIntersection_Shallow3", Center, Radius, Thickness,
 		//	StartPt, EndPt, 
 		//	IntersectionResult(Center+Chaos::TVector<T,3>(0,0,Radius+Thickness), true));
 
 		EndPt += Chaos::TVector<T, 3>(0, 0, Thickness);
-		RunTestFindClosestIntersectionHelper(R, "RunTestFindClosestIntersection_Shallow3.1", Center, Radius, Thickness,
+		RunTestFindClosestIntersectionHelper("RunTestFindClosestIntersection_Shallow3.1", Center, Radius, Thickness,
 			StartPt, EndPt, 
 			IntersectionResult(Center+Chaos::TVector<T,3>(0,0,Radius+Thickness), true));
 
@@ -148,26 +145,25 @@ namespace GeometryCollectionExample
 		// "RunTestFindClosestIntersection_Shallow4 - Sphere((0, 0, 0), 1).FindClosestIntersection(StartPt(0, 0, 1), EndPt(0, 0, 0.999999), Thickness:0.1) = RESULT: 0, Point: (0, 0, 0), EXPECTED: 1 (0, 0, 1.1)."	FString
 		EndPt = StartPt;
 		StartPt = EndPt + Chaos::TVector<T, 3>(0, 0, Radius * .000001);
-		//RunTestFindClosestIntersectionHelper(R, "RunTestFindClosestIntersection_Shallow4", Center, Radius, Thickness,
+		//RunTestFindClosestIntersectionHelper("RunTestFindClosestIntersection_Shallow4", Center, Radius, Thickness,
 		//	StartPt, EndPt, 
 		//	IntersectionResult(Center+Chaos::TVector<T,3>(0,0,Radius+Thickness), true));
 
 		// Fails w/no collision:
 		// "RunTestFindClosestIntersection_Shallow5 - Sphere((0, 0, 0), 1).FindClosestIntersection(StartPt(0, 0, 1), EndPt(0, 0, 0.999999), Thickness:0.1) = RESULT: 0, Point: (0, 0, 0), EXPECTED: 1 (0, 0, 1.1)."	FString
 		StartPt += Chaos::TVector<T, 3>(0, 0, Radius * .000001);
-		//RunTestFindClosestIntersectionHelper(R, "RunTestFindClosestIntersection_Shallow5", Center, Radius, Thickness,
+		//RunTestFindClosestIntersectionHelper("RunTestFindClosestIntersection_Shallow5", Center, Radius, Thickness,
 		//	StartPt, EndPt, 
 		//	IntersectionResult(Center+Chaos::TVector<T,3>(0,0,Radius+Thickness), true));
 
 		StartPt += Chaos::TVector<T, 3>(0, 0, Thickness);
-		RunTestFindClosestIntersectionHelper(R, "RunTestFindClosestIntersection_Shallow5.1", Center, Radius, Thickness,
+		RunTestFindClosestIntersectionHelper("RunTestFindClosestIntersection_Shallow5.1", Center, Radius, Thickness,
 			StartPt, EndPt, 
 			IntersectionResult(Center+Chaos::TVector<T,3>(0,0,Radius+Thickness), true));
 	}
 
 	template <class T>
 	void RunTestFindClosestIntersection_Mid(
-		ExampleResponse& R, 
 		const Chaos::TVector<T, 3>& Center, 
 		const T Radius,
 		const T Thickness)
@@ -177,21 +173,21 @@ namespace GeometryCollectionExample
 
 		// Fails w/no collision:
 		// "RunTestFindClosestIntersection_Mid1 - Sphere((0, 0, 0), 1).FindClosestIntersection(StartPt(0, 0, 0.5), EndPt(0, 0, 0.5), Thickness:0.1) = RESULT: 0, Point: (0, 0, 0), EXPECTED: 1 (1.1, 0, 0)."	FString
-		//RunTestFindClosestIntersectionHelper(R, "RunTestFindClosestIntersection_Mid1", Center, Radius, Thickness,
+		//RunTestFindClosestIntersectionHelper("RunTestFindClosestIntersection_Mid1", Center, Radius, Thickness,
 		//	StartPt, EndPt, 
 		//	IntersectionResult(Center+Chaos::TVector<T,3>(Radius+Thickness,0,0), true));
 
 		// Fails w/no collision:
 		// "RunTestFindClosestIntersection_Mid2 - Sphere((0, 0, 0), 1).FindClosestIntersection(StartPt(0, 0, 0.5), EndPt(0, 0, 1), Thickness:0.1) = RESULT: 0, Point: (0, 0, 0), EXPECTED: 1 (0, 0, 1.1)."	FString
 		//EndPt += Chaos::TVector<T, 3>(0, 0, Radius * .5);
-		//RunTestFindClosestIntersectionHelper(R, "RunTestFindClosestIntersection_Mid2", Center, Radius, Thickness,
+		//RunTestFindClosestIntersectionHelper("RunTestFindClosestIntersection_Mid2", Center, Radius, Thickness,
 		//	StartPt, EndPt, 
 		//	IntersectionResult(Center+Chaos::TVector<T,3>(0,0,Radius+Thickness), true));
 		
 		// Fails w/no collision:
 		// "RunTestFindClosestIntersection_Mid3 - Sphere((0, 0, 0), 1).FindClosestIntersection(StartPt(0, 0, 0.5), EndPt(0, 0, 1), Thickness:0.1) = RESULT: 0, Point: (0, 0, 0), EXPECTED: 1 (0, 0, 1.1)."	FString
 		EndPt += Chaos::TVector<T, 3>(0, 0, Radius * .5);
-		//RunTestFindClosestIntersectionHelper(R, "RunTestFindClosestIntersection_Mid3", Center, Radius, Thickness,
+		//RunTestFindClosestIntersectionHelper("RunTestFindClosestIntersection_Mid3", Center, Radius, Thickness,
 		//	StartPt, EndPt, 
 		//	IntersectionResult(Center+Chaos::TVector<T,3>(0,0,Radius+Thickness), true));
 
@@ -199,19 +195,18 @@ namespace GeometryCollectionExample
 		// "RunTestFindClosestIntersection_Mid4 - Sphere((0, 0, 0), 1).FindClosestIntersection(StartPt(0, 0, 1), EndPt(0, 0, 0.5), Thickness:0.1) = RESULT: 0, Point: (0, 0, 0), EXPECTED: 1 (0, 0, 1.1)."	FString
 		EndPt = StartPt;
 		StartPt = EndPt + Chaos::TVector<T, 3>(0, 0, Radius * .5);
-		//RunTestFindClosestIntersectionHelper(R, "RunTestFindClosestIntersection_Mid4", Center, Radius, Thickness,
+		//RunTestFindClosestIntersectionHelper("RunTestFindClosestIntersection_Mid4", Center, Radius, Thickness,
 		//	StartPt, EndPt, 
 		//	IntersectionResult(Center+Chaos::TVector<T,3>(0,0,Radius+Thickness), true));
 
 		StartPt += Chaos::TVector<T, 3>(0, 0, Radius * .5);
-		RunTestFindClosestIntersectionHelper(R, "RunTestFindClosestIntersection_Mid5", Center, Radius, Thickness,
+		RunTestFindClosestIntersectionHelper("RunTestFindClosestIntersection_Mid5", Center, Radius, Thickness,
 			StartPt, EndPt, 
 			IntersectionResult(Center+Chaos::TVector<T,3>(0,0,Radius+Thickness), true));
 	}
 
 	template <class T>
 	void RunTestFindClosestIntersection_Deep(
-		ExampleResponse& R, 
 		const Chaos::TVector<T, 3>& Center, 
 		const T Radius,
 		const T Thickness)
@@ -221,19 +216,19 @@ namespace GeometryCollectionExample
 
 		// Fails w/no collision
 		// "RunTestFindClosestIntersection_Deep1 - Sphere((0, 0, 0), 1).FindClosestIntersection(StartPt(0, 0, 0), EndPt(0, 0, 0), Thickness:0.1) = RESULT: 0, Point: (0, 0, 0), EXPECTED: 1 (1.1, 0, 0)."	FString
-		//RunTestFindClosestIntersectionHelper(R, "RunTestFindClosestIntersection_Deep1", Center, Radius, Thickness,
+		//RunTestFindClosestIntersectionHelper("RunTestFindClosestIntersection_Deep1", Center, Radius, Thickness,
 		//	StartPt, EndPt, 
 		//	IntersectionResult(Center+Chaos::TVector<T,3>(Radius+Thickness,0,0), true));
 
 		// Fails w/no collision
 		// "RunTestFindClosestIntersection_Deep2 - Sphere((0, 0, 0), 1).FindClosestIntersection(StartPt(0, 0, 0), EndPt(0, 0, 0.5), Thickness:0.1) = RESULT: 0, Point: (0, 0, 0), EXPECTED: 1 (0, 0, 1.1)."	FString
 		EndPt += Chaos::TVector<T, 3>(0, 0, Radius*.5);
-		//RunTestFindClosestIntersectionHelper(R, "RunTestFindClosestIntersection_Deep2", Center, Radius, Thickness,
+		//RunTestFindClosestIntersectionHelper("RunTestFindClosestIntersection_Deep2", Center, Radius, Thickness,
 		//	StartPt, EndPt, 
 		//	IntersectionResult(Center+Chaos::TVector<T, 3>(0, 0, Radius+Thickness), true));
 
 		EndPt += Chaos::TVector<T, 3>(0, 0, Radius);
-		RunTestFindClosestIntersectionHelper(R, "RunTestFindClosestIntersection_Deep3", Center, Radius, Thickness,
+		RunTestFindClosestIntersectionHelper("RunTestFindClosestIntersection_Deep3", Center, Radius, Thickness,
 			StartPt, EndPt, 
 			IntersectionResult(Center+Chaos::TVector<T, 3>(0, 0, Radius+Thickness), true));
 
@@ -241,79 +236,79 @@ namespace GeometryCollectionExample
 		// "RunTestFindClosestIntersection_Deep4 - Sphere((0, 0, 0), 1).FindClosestIntersection(StartPt(0, 0, 0.5), EndPt(0, 0, 0), Thickness:0.1) = RESULT: 0, Point: (0, 0, 0), EXPECTED: 1 (0, 0, 1.1)."	FString
 		EndPt = Center;
 		StartPt += Chaos::TVector<T, 3>(0, 0, Radius*.5);
-		//RunTestFindClosestIntersectionHelper(R, "RunTestFindClosestIntersection_Deep4", Center, Radius, Thickness,
+		//RunTestFindClosestIntersectionHelper("RunTestFindClosestIntersection_Deep4", Center, Radius, Thickness,
 		//	StartPt, EndPt, 
 		//	IntersectionResult(Center+Chaos::TVector<T, 3>(0, 0, Radius+Thickness), true));
 
 		StartPt += Chaos::TVector<T, 3>(0, 0, Radius);
-		RunTestFindClosestIntersectionHelper(R, "RunTestFindClosestIntersection_Deep5", Center, Radius, Thickness,
+		RunTestFindClosestIntersectionHelper("RunTestFindClosestIntersection_Deep5", Center, Radius, Thickness,
 			StartPt, EndPt, 
 			IntersectionResult(Center+Chaos::TVector<T, 3>(0, 0, Radius+Thickness), true));
 
 		// Miss
 		StartPt = Center + Chaos::TVector<T, 3>(Radius + Thickness + 1);
 		EndPt = StartPt + Chaos::TVector<T, 3>(Radius + Thickness + 1);
-		RunTestFindClosestIntersectionHelper(R, "RunTestFindClosestIntersection_Deep6", Center, Radius, Thickness,
+		RunTestFindClosestIntersectionHelper("RunTestFindClosestIntersection_Deep6", Center, Radius, Thickness,
 			StartPt, EndPt, 
 			IntersectionResult(StartPt, false));
 	}
 
 	template <class T>
 	void RunTestFindClosestIntersection(
-		ExampleResponse& R, const Chaos::TVector<T,3>& Center, const T Radius)
+		const Chaos::TVector<T,3>& Center, const T Radius)
 	{
 		//
 		// Run different test configurations with varying thicknesses.
 		//
 
 		// Deep intersection
-		RunTestFindClosestIntersection_Deep<T>(R, Center, Radius, 0.0);
-		RunTestFindClosestIntersection_Deep<T>(R, Center, Radius, Radius*.1);
-//		RunTestFindClosestIntersection_Deep<T>(R, Center, Radius, Radius*.5); // fails with incorrect collision points
-//		RunTestFindClosestIntersection_Deep<T>(R, Center, Radius, Radius);
-//		RunTestFindClosestIntersection_Deep<T>(R, Center, Radius, Radius*2);
+		RunTestFindClosestIntersection_Deep<T>(Center, Radius, 0.0);
+		RunTestFindClosestIntersection_Deep<T>(Center, Radius, Radius*.1);
+//		RunTestFindClosestIntersection_Deep<T>(Center, Radius, Radius*.5); // fails with incorrect collision points
+//		RunTestFindClosestIntersection_Deep<T>(Center, Radius, Radius);
+//		RunTestFindClosestIntersection_Deep<T>(Center, Radius, Radius*2);
 
 		// Mid intersection
-		RunTestFindClosestIntersection_Mid<T>(R, Center, Radius, 0.0);
-		RunTestFindClosestIntersection_Mid<T>(R, Center, Radius, Radius*.1);
-//		RunTestFindClosestIntersection_Mid<T>(R, Center, Radius, Radius*.5); // fails with incorrect collision points
-//		RunTestFindClosestIntersection_Mid<T>(R, Center, Radius, Radius);
-//		RunTestFindClosestIntersection_Mid<T>(R, Center, Radius, Radius*2);
+		RunTestFindClosestIntersection_Mid<T>(Center, Radius, 0.0);
+		RunTestFindClosestIntersection_Mid<T>(Center, Radius, Radius*.1);
+//		RunTestFindClosestIntersection_Mid<T>(Center, Radius, Radius*.5); // fails with incorrect collision points
+//		RunTestFindClosestIntersection_Mid<T>(Center, Radius, Radius);
+//		RunTestFindClosestIntersection_Mid<T>(Center, Radius, Radius*2);
 
 		// Grazing intersection - all fail
-		RunTestFindClosestIntersection_Shallow<T>(R, Center, Radius, 0.0);
-//		RunTestFindClosestIntersection_Shallow<T>(R, Center, Radius, Radius*.1); // fails with incorrect collision points
-//		RunTestFindClosestIntersection_Shallow<T>(R, Center, Radius, Radius*.5);
-//		RunTestFindClosestIntersection_Shallow<T>(R, Center, Radius, Radius);
-//		RunTestFindClosestIntersection_Shallow<T>(R, Center, Radius, Radius*2);
+		RunTestFindClosestIntersection_Shallow<T>(Center, Radius, 0.0);
+//		RunTestFindClosestIntersection_Shallow<T>(Center, Radius, Radius*.1); // fails with incorrect collision points
+//		RunTestFindClosestIntersection_Shallow<T>(Center, Radius, Radius*.5);
+//		RunTestFindClosestIntersection_Shallow<T>(Center, Radius, Radius);
+//		RunTestFindClosestIntersection_Shallow<T>(Center, Radius, Radius*2);
 
 		// Tangent intersection
-		RunTestFindClosestIntersection_Tangent<T>(R, Center, Radius, 0.0);
-		RunTestFindClosestIntersection_Tangent<T>(R, Center, Radius, Radius*.1);
-		RunTestFindClosestIntersection_Tangent<T>(R, Center, Radius, Radius*.5);
-		RunTestFindClosestIntersection_Tangent<T>(R, Center, Radius, Radius);
-		RunTestFindClosestIntersection_Tangent<T>(R, Center, Radius, Radius*2);
+		RunTestFindClosestIntersection_Tangent<T>(Center, Radius, 0.0);
+		RunTestFindClosestIntersection_Tangent<T>(Center, Radius, Radius*.1);
+		RunTestFindClosestIntersection_Tangent<T>(Center, Radius, Radius*.5);
+		RunTestFindClosestIntersection_Tangent<T>(Center, Radius, Radius);
+		RunTestFindClosestIntersection_Tangent<T>(Center, Radius, Radius*2);
 	}
 	template <class T>
-	void TestIntersections(ExampleResponse& R)
+	void TestIntersections()
 	{
 		// At the origin
-		RunTestFindClosestIntersection(R, Chaos::TVector<T, 3>((T)0.0), (T)1.0);
-//		RunTestFindClosestIntersection(R, Chaos::TVector<T, 3>((T)0.0), (T)10.0);
-//		RunTestFindClosestIntersection(R, Chaos::TVector<T, 3>((T)0.0), (T)100.0);
+		RunTestFindClosestIntersection(Chaos::TVector<T, 3>((T)0.0), (T)1.0);
+//		RunTestFindClosestIntersection(Chaos::TVector<T, 3>((T)0.0), (T)10.0);
+//		RunTestFindClosestIntersection(Chaos::TVector<T, 3>((T)0.0), (T)100.0);
 
 		// Off origin
-//		RunTestFindClosestIntersection(R, Chaos::TVector<T, 3>((T)1.0), (T)1.0);
-//		RunTestFindClosestIntersection(R, Chaos::TVector<T, 3>((T)10.0), (T)1.0);
-//		RunTestFindClosestIntersection(R, Chaos::TVector<T, 3>((T)100.0), (T)1.0);
+//		RunTestFindClosestIntersection(Chaos::TVector<T, 3>((T)1.0), (T)1.0);
+//		RunTestFindClosestIntersection(Chaos::TVector<T, 3>((T)10.0), (T)1.0);
+//		RunTestFindClosestIntersection(Chaos::TVector<T, 3>((T)100.0), (T)1.0);
 
-//		RunTestFindClosestIntersection(R, Chaos::TVector<T, 3>((T)1.0), (T)10.0);
-//		RunTestFindClosestIntersection(R, Chaos::TVector<T, 3>((T)10.0), (T)10.0);
-//		RunTestFindClosestIntersection(R, Chaos::TVector<T, 3>((T)100.0), (T)10.0);
+//		RunTestFindClosestIntersection(Chaos::TVector<T, 3>((T)1.0), (T)10.0);
+//		RunTestFindClosestIntersection(Chaos::TVector<T, 3>((T)10.0), (T)10.0);
+//		RunTestFindClosestIntersection(Chaos::TVector<T, 3>((T)100.0), (T)10.0);
 
-//		RunTestFindClosestIntersection(R, Chaos::TVector<T, 3>((T)1.0), (T)100.0);
-//		RunTestFindClosestIntersection(R, Chaos::TVector<T, 3>((T)10.0), (T)100.0);
-//		RunTestFindClosestIntersection(R, Chaos::TVector<T, 3>((T)100.0), (T)100.0);
+//		RunTestFindClosestIntersection(Chaos::TVector<T, 3>((T)1.0), (T)100.0);
+//		RunTestFindClosestIntersection(Chaos::TVector<T, 3>((T)10.0), (T)100.0);
+//		RunTestFindClosestIntersection(Chaos::TVector<T, 3>((T)100.0), (T)100.0);
 	}
 
 	//==========================================================================
@@ -321,17 +316,17 @@ namespace GeometryCollectionExample
 	//==========================================================================
 
 	template <class T, int d>
-	void RunTestComputeSamplePoints(ExampleResponse& R, const Chaos::TSphere<T, d> &Sphere)
+	void RunTestComputeSamplePoints(const Chaos::TSphere<T, d> &Sphere)
 	{
-		R.ExpectTrue(Sphere.GetType() == Chaos::ImplicitObjectType::Sphere, FString("Implicit object type is not 'sphere'."));
+		EXPECT_EQ(Sphere.GetType(), Chaos::ImplicitObjectType::Sphere) << *FString("Implicit object type is not 'sphere'.");
 
 		Chaos::TVector<T, d> Point = Sphere.GetCenter();
 		T Phi = Sphere.SignedDistance(Point);
-		R.ExpectTrue(FMath::Abs(Phi + Sphere.GetRadius()) <= SMALL_NUMBER, FString("Sphere failed phi depth sanity test."));
+		EXPECT_LE(FMath::Abs(Phi + Sphere.GetRadius()), SMALL_NUMBER) << *FString("Sphere failed phi depth sanity test.");
 		
 		Point[0] += Sphere.GetRadius();
 		Phi = Sphere.SignedDistance(Point);
-		R.ExpectTrue(FMath::Abs(Phi) <= KINDA_SMALL_NUMBER, FString("Sphere failed phi surface sanity test."));
+		EXPECT_LE(FMath::Abs(Phi), KINDA_SMALL_NUMBER) << *FString("Sphere failed phi surface sanity test.");
 
 		TArray<Chaos::TVector<T, d>> Points = Sphere.ComputeSamplePoints(100);
 		check(Points.Num() == 100);
@@ -341,16 +336,16 @@ namespace GeometryCollectionExample
 			Phi = Sphere.SignedDistance(Pt);
 			const bool OnSurface = FMath::Abs(Phi) <= KINDA_SMALL_NUMBER;
 			check(OnSurface);
-			R.ExpectTrue(OnSurface, FString("Produced a point not on the surface of the sphere."));
+			EXPECT_TRUE(OnSurface) << *FString("Produced a point not on the surface of the sphere.");
 			const bool Differs = Pt != Point;
 			check(Differs);
-			R.ExpectTrue(Differs, FString("Produced a redundant value."));
+			EXPECT_TRUE(Differs) << *FString("Produced a redundant value.");
 			Point = Pt;
 		}
 	}
 
 	template <class T>
-	void RunTestComputeSemispherePoints(ExampleResponse& R, const Chaos::TSphere<T, 3> &Sphere)
+	void RunTestComputeSemispherePoints(const Chaos::TSphere<T, 3> &Sphere)
 	{
 		TArray<Chaos::TVector<T, 3>> Points;
 		Chaos::TSphereSpecializeSamplingHelper<T, 3>::ComputeBottomHalfSemiSphere(Points, Sphere, 100);
@@ -359,11 +354,11 @@ namespace GeometryCollectionExample
 			T Phi = Sphere.SignedDistance(Pt);
 			const bool OnSurface = FMath::Abs(Phi) <= KINDA_SMALL_NUMBER;
 			check(OnSurface);
-			R.ExpectTrue(OnSurface, FString("Produced a point not on the surface of the sphere."));
+			EXPECT_TRUE(OnSurface) << *FString("Produced a point not on the surface of the sphere.");
 
 			const bool BelowCenter = Pt[2] < Sphere.GetCenter()[2] + KINDA_SMALL_NUMBER;
 			check(BelowCenter);
-			R.ExpectTrue(BelowCenter, FString("Bottom semisphere produced a point above midline."));
+			EXPECT_TRUE(BelowCenter) << *FString("Bottom semisphere produced a point above midline.");
 		}
 
 		Points.Reset();
@@ -373,71 +368,66 @@ namespace GeometryCollectionExample
 			T Phi = Sphere.SignedDistance(Pt);
 			const bool OnSurface = FMath::Abs(Phi) <= KINDA_SMALL_NUMBER;
 			check(OnSurface);
-			R.ExpectTrue(OnSurface, FString("Produced a point not on the surface of the sphere."));
+			EXPECT_TRUE(OnSurface) << *FString("Produced a point not on the surface of the sphere.");
 
 			const bool BelowCenter = Pt[2] > Sphere.GetCenter()[2] - KINDA_SMALL_NUMBER;
 			check(BelowCenter);
-			R.ExpectTrue(BelowCenter, FString("Top semisphere produced a point above midline."));
+			EXPECT_TRUE(BelowCenter) << *FString("Top semisphere produced a point above midline.");
 		}
 	}
 
 	template <class T>
-	void TestComputeSamplePoints_SemiSphere(ExampleResponse& R)
+	void TestComputeSamplePoints_SemiSphere()
 	{
 		// At the origin with radius 1
 		{
 			Chaos::TSphere<T, 3> Sphere(Chaos::TVector<T, 3>((T)0.0), (T)1.0);
-			RunTestComputeSemispherePoints(R, Sphere);
+			RunTestComputeSemispherePoints(Sphere);
 		}	
 	}
 
 	template <class T, int d>
-	void TestComputeSamplePoints_Sphere(ExampleResponse& R)
+	void TestComputeSamplePoints_Sphere()
 	{
 		// At the origin with radius 1
 		{
 			Chaos::TSphere<T, d> Sphere(Chaos::TVector<T, d>((T)0.0), (T)1.0);
-			RunTestComputeSamplePoints(R, Sphere);
+			RunTestComputeSamplePoints(Sphere);
 		}
 		// At the origin with radius > 1
 		{
 			Chaos::TSphere<T, d> Sphere(Chaos::TVector<T, d>(0.0), (T)10.0);
-			RunTestComputeSamplePoints(R, Sphere);
+			RunTestComputeSamplePoints(Sphere);
 		}
 		// At the origin with radius < 1
 		{
 			Chaos::TSphere<T, d> Sphere(Chaos::TVector<T, d>(0.0), (T).1);
-			RunTestComputeSamplePoints(R, Sphere);
+			RunTestComputeSamplePoints(Sphere);
 		}
 		// Off the origin with radius 1
 		{
 			Chaos::TSphere<T, d> Sphere(Chaos::TVector<T, d>(10.0), (T)1.0);
-			RunTestComputeSamplePoints(R, Sphere);
+			RunTestComputeSamplePoints(Sphere);
 		}
 		// Off the origin with radius > 1
 		{
 			Chaos::TSphere<T, d> Sphere(Chaos::TVector<T, d>(10.0), (T)10.0);
-			RunTestComputeSamplePoints(R, Sphere);
+			RunTestComputeSamplePoints(Sphere);
 		}
 		// Off the origin with radius < 1
 		{
 			Chaos::TSphere<T, d> Sphere(Chaos::TVector<T, d>(10.0), (T).1);
-			RunTestComputeSamplePoints(R, Sphere);
+			RunTestComputeSamplePoints(Sphere);
 		}
 	}
 
 	template<class T>
-	bool TestImplicitSphere(ExampleResponse&& R)
+	void TestImplicitSphere()
 	{
 		//TestComputeSamplePoints<T, 2>(R);
-		TestComputeSamplePoints_Sphere<T, 3>(R);
-		TestComputeSamplePoints_SemiSphere<T>(R);
-		check(!R.HasError());
-
-		TestIntersections<T>(R);
-		check(!R.HasError());
-
-		return !R.HasError();
+		TestComputeSamplePoints_Sphere<T, 3>();
+		TestComputeSamplePoints_SemiSphere<T>();
+		TestIntersections<T>();
 	}
-	template bool TestImplicitSphere<float>(ExampleResponse&& R);
+	template void TestImplicitSphere<float>();
 } // namespace GeometryCollectionExample
\ No newline at end of file
diff --git a/Engine/Plugins/Experimental/GeometryCollectionPlugin/Source/GeometryCollectionExampleCore/Private/GeometryCollection/GeometryCollectionExampleMatrices.cpp b/Engine/Plugins/Experimental/GeometryCollectionPlugin/Source/GeometryCollectionExampleCore/Private/GeometryCollection/GeometryCollectionExampleMatrices.cpp
--- a/Engine/Plugins/Experimental/GeometryCollectionPlugin/Source/GeometryCollectionExampleCore/Private/GeometryCollection/GeometryCollectionExampleMatrices.cpp
+++ b/Engine/Plugins/Experimental/GeometryCollectionPlugin/Source/GeometryCollectionExampleCore/Private/GeometryCollection/GeometryCollectionExampleMatrices.cpp
@@ -5,6 +5,7 @@
 #include "GeometryCollection/GeometryCollection.h"
 #include "GeometryCollection/GeometryCollectionAlgo.h"
 #include "GeometryCollection/GeometryCollectionUtility.h"
+#include "GeometryCollection/GeometryCollectionExampleUtility.h"
 #include "GeometryCollection/TransformCollection.h"
 #include "UObject/Package.h"
 #include "UObject/UObjectGlobals.h"
@@ -13,8 +14,11 @@ namespace GeometryCollectionExample
 {
 	DEFINE_LOG_CATEGORY_STATIC(GCTM_Log, Verbose, All);
 
+	MATCHER_P2(VectorNear, V, Tolerance, "") { return arg.Equals(V, Tolerance); }
+#define EXPECT_FVECTOR_NEAR(A,B,T) EXPECT_THAT(A, VectorNear(B, T)) << *FString("Expected: " + B.ToString() + "\nActual:   " + A.ToString());
+
 	template<class T>
-	bool BasicGlobalMatrices(ExampleResponse&& R)
+	void BasicGlobalMatrices()
 	{ 
 		TSharedPtr<FGeometryCollection> Collection = GeometryCollection::MakeCubeElement(FTransform(FQuat::MakeFromEuler(FVector(0, 0, 90.)), FVector(0, 10, 0)), FVector(1.0));
 		Collection->AppendGeometry(*GeometryCollection::MakeCubeElement(FTransform(FQuat::MakeFromEuler(FVector(0, 0, 90.)), FVector(0, 10, 0)), FVector(1.0)));
@@ -33,37 +37,35 @@ namespace GeometryCollectionExample
 		GeometryCollectionAlgo::GlobalMatrices(Collection->Transform, Collection->Parent, GlobalTransform);
 
 		FVector Rot0 = GlobalTransform[0].GetRotation().Euler();
-		R.ExpectTrue(Rot0.Equals(FVector(0, 0, 90), 0.0001));
+		EXPECT_FVECTOR_NEAR(Rot0, FVector(0, 0, 90), 0.0001);
 
 		FVector Rot1 = GlobalTransform[1].GetRotation().Euler();
-		R.ExpectTrue(Rot1.Equals(FVector(0, 0, -180), 0.0001));
+		EXPECT_FVECTOR_NEAR(Rot1, FVector(0, 0, -180), 0.0001);
 
 		FVector Rot2 = GlobalTransform[2].GetRotation().Euler();
-		R.ExpectTrue(Rot2.Equals(FVector(0, 0, -90), 0.0001));
+		EXPECT_FVECTOR_NEAR(Rot2, FVector(0, 0, -90), 0.0001);
 
 		FVector Pos0 = GlobalTransform[0].GetTranslation();
-		R.ExpectTrue(Pos0.Equals(FVector(0, 10, 0), 0.0001));
+		EXPECT_FVECTOR_NEAR(Pos0, FVector(0, 10, 0), 0.0001);
 
 		FVector Pos1 = GlobalTransform[1].GetTranslation();
-		R.ExpectTrue(Pos1.Equals(FVector(-10, 10, 0), 0.0001));
+		EXPECT_FVECTOR_NEAR(Pos1, FVector(-10, 10, 0), 0.0001);
 
 		FVector Pos2 = GlobalTransform[2].GetTranslation();
-		R.ExpectTrue(Pos2.Equals(FVector(-10, 0, 0), 0.0001));
+		EXPECT_FVECTOR_NEAR(Pos2, FVector(-10, 0, 0), 0.0001);
 
 		FTransform Frame = GeometryCollectionAlgo::GlobalMatrix(Collection->Transform, Collection->Parent, 2);
-		R.ExpectTrue(Frame.GetRotation().Euler().Equals(FVector(0, 0, -90), 0.0001));
-		R.ExpectTrue(Frame.GetTranslation().Equals(FVector(-10, 0, 0), 0.0001));
+		EXPECT_FVECTOR_NEAR(Frame.GetRotation().Euler(), FVector(0, 0, -90), 0.0001);
+		EXPECT_FVECTOR_NEAR(Frame.GetTranslation(), FVector(-10, 0, 0), 0.0001);
 
 		Frame = GeometryCollectionAlgo::GlobalMatrix(Collection->Transform, Collection->Parent, 1);
-		R.ExpectTrue(Frame.GetRotation().Euler().Equals(FVector(0, 0, -180), 0.0001));
-		R.ExpectTrue(Frame.GetTranslation().Equals(FVector(-10, 10, 0), 0.0001));
-
-		return !R.HasError();
+		EXPECT_FVECTOR_NEAR(Frame.GetRotation().Euler(), FVector(0, 0, -180), 0.0001);
+		EXPECT_FVECTOR_NEAR(Frame.GetTranslation(), FVector(-10, 10, 0), 0.0001);
 	}
-	template bool BasicGlobalMatrices<float>(ExampleResponse&& R);
+	template void BasicGlobalMatrices<float>();
 
 	template<class T>
-	bool ReparentingMatrices(ExampleResponse&& R)
+	void ReparentingMatrices()
 	{
 		TSharedPtr<FGeometryCollection> Collection = GeometryCollection::MakeCubeElement(FTransform(FQuat::MakeFromEuler(FVector(0, 0, 90.)), FVector(0, 10, 0)), FVector(1.0));
 		Collection->AppendGeometry(*GeometryCollection::MakeCubeElement(FTransform(FQuat::MakeFromEuler(FVector(0, 0, 90.)), FVector(0, 10, 0)), FVector(1.0)));
@@ -82,37 +84,35 @@ namespace GeometryCollectionExample
 		//  ......2
 		TArray<int32> Bones = { 2 };
 		GeometryCollectionAlgo::ParentTransforms(Collection.Get(), 1, Bones);
-		R.ExpectTrue((Collection->Transform)[2].GetTranslation().Equals(FVector(0, 10, 0), 0.0001));
-		R.ExpectTrue((Collection->Transform)[2].GetRotation().Euler().Equals(FVector(0, 0, 90.), 0.0001));
+		EXPECT_FVECTOR_NEAR((Collection->Transform)[2].GetTranslation(), FVector(0, 10, 0), 0.0001);
+		EXPECT_FVECTOR_NEAR((Collection->Transform)[2].GetRotation().Euler(), FVector(0, 0, 90.), 0.0001);
 
 		TArray<FTransform> GlobalTransform;
 		GeometryCollectionAlgo::GlobalMatrices(Collection->Transform, Collection->Parent, GlobalTransform);
 
 		FVector Rot0 = GlobalTransform[0].GetRotation().Euler();
-		R.ExpectTrue(Rot0.Equals(FVector(0, 0, 90), 0.0001));
+		EXPECT_FVECTOR_NEAR(Rot0, FVector(0, 0, 90), 0.0001);
 
 		FVector Rot1 = GlobalTransform[1].GetRotation().Euler();
-		R.ExpectTrue(Rot1.Equals(FVector(0, 0, -180), 0.0001));
+		EXPECT_FVECTOR_NEAR(Rot1, FVector(0, 0, -180), 0.0001);
 
 		FVector Rot2 = GlobalTransform[2].GetRotation().Euler();
-		R.ExpectTrue(Rot2.Equals(FVector(0, 0, -90), 0.0001));
+		EXPECT_FVECTOR_NEAR(Rot2, FVector(0, 0, -90), 0.0001);
 
 		FVector Pos0 = GlobalTransform[0].GetTranslation();
-		R.ExpectTrue(Pos0.Equals(FVector(0, 10, 0), 0.0001));
+		EXPECT_FVECTOR_NEAR(Pos0, FVector(0, 10, 0), 0.0001);
 
 		FVector Pos1 = GlobalTransform[1].GetTranslation();
-		R.ExpectTrue(Pos1.Equals(FVector(-10, 10, 0), 0.0001));
+		EXPECT_FVECTOR_NEAR(Pos1, FVector(-10, 10, 0), 0.0001);
 
 		FVector Pos2 = GlobalTransform[2].GetTranslation();
-		R.ExpectTrue(Pos2.Equals(FVector(-10, 0, 0), 0.0001));
-
-		return !R.HasError();
+		EXPECT_FVECTOR_NEAR(Pos2, FVector(-10, 0, 0), 0.0001);
 	}
-	template bool ReparentingMatrices<float>(ExampleResponse&& R);
+	template void ReparentingMatrices<float>();
 
 
 	template<class T>
-	bool TransformMatrixElement(ExampleResponse&& R)
+	void TransformMatrixElement()
 	{
 		FTransformCollection Collection;
 
@@ -146,14 +146,12 @@ namespace GeometryCollectionExample
 		//	UE_LOG(GCTM_Log, Verbose, TEXT("... GlobalM0[%d] : (%3.5f,%3.5f,%3.5f)"), rdx, GlobalMatrices0[rdx].GetTranslation().X, GlobalMatrices0[rdx].GetTranslation().Y, GlobalMatrices0[rdx].GetTranslation().Z);
 		//	UE_LOG(GCTM_Log, Verbose, TEXT("... GlobalM1[%d] : (%3.5f,%3.5f,%3.5f)"), rdx, GlobalMatrices1[rdx].GetTranslation().X, GlobalMatrices1[rdx].GetTranslation().Y, GlobalMatrices1[rdx].GetTranslation().Z);
 		//}
-		R.ExpectTrue(FMath::Abs((GlobalMatrices1[0].GetTranslation() - GlobalMatrices0[0].GetTranslation()).Size()) < 1e-3);
-		R.ExpectTrue(FMath::Abs((GlobalMatrices1[1].GetTranslation() - GlobalMatrices0[1].GetTranslation()).Size()) > 1.0);
-		R.ExpectTrue(FMath::Abs((GlobalMatrices1[2].GetTranslation() - GlobalMatrices0[2].GetTranslation()).Size()) > 1.0);
-		R.ExpectTrue(FMath::Abs((GlobalMatrices1[3].GetTranslation() - GlobalMatrices0[3].GetTranslation()).Size()) < 1e-3);
-
-		return !R.HasError();
+		EXPECT_LT(FMath::Abs((GlobalMatrices1[0].GetTranslation() - GlobalMatrices0[0].GetTranslation()).Size()), 1e-3);
+		EXPECT_GT(FMath::Abs((GlobalMatrices1[1].GetTranslation() - GlobalMatrices0[1].GetTranslation()).Size()), 1.0);
+		EXPECT_GT(FMath::Abs((GlobalMatrices1[2].GetTranslation() - GlobalMatrices0[2].GetTranslation()).Size()), 1.0);
+		EXPECT_LT(FMath::Abs((GlobalMatrices1[3].GetTranslation() - GlobalMatrices0[3].GetTranslation()).Size()), 1e-3);
 	}
-	template bool TransformMatrixElement<float>(ExampleResponse&& R);
+	template void TransformMatrixElement<float>();
 
 }
 
