Copying //UE4/Dev-Physics to Dev-Main (//UE4/Dev-Main) [at] 9098278

diff --git a/Engine/Plugins/Experimental/GeometryCollectionPlugin/Source/GeometryCollectionExampleCore/Private/GeometryCollection/GeometryCollectionExampleProximity.cpp b/Engine/Plugins/Experimental/GeometryCollectionPlugin/Source/GeometryCollectionExampleCore/Private/GeometryCollection/GeometryCollectionExampleProximity.cpp
--- a/Engine/Plugins/Experimental/GeometryCollectionPlugin/Source/GeometryCollectionExampleCore/Private/GeometryCollection/GeometryCollectionExampleProximity.cpp
+++ b/Engine/Plugins/Experimental/GeometryCollectionPlugin/Source/GeometryCollectionExampleCore/Private/GeometryCollection/GeometryCollectionExampleProximity.cpp
@@ -14,7 +14,7 @@
 namespace GeometryCollectionExample
 {
 	template<class T>
-	bool BuildProximity(ExampleResponse&& R)
+	void BuildProximity()
 	{
 		TSharedPtr<FGeometryCollection> Collection = GeometryCollection::MakeCubeElement(FTransform(FQuat::MakeFromEuler(FVector(0.f, 0.f, 0.f)), FVector(0.f, 0.f, 0.f)), FVector(1.0));
 		Collection->AppendGeometry(*GeometryCollection::MakeCubeElement(FTransform(FQuat::MakeFromEuler(FVector(0.f, 0.f, 0.f)), FVector(1.f, 0.f, 0.f)),FVector(1.0)));
@@ -40,20 +40,18 @@ namespace GeometryCollectionExample
 		// Breaking Data
 		TManagedArray<TSet<int32>>& Proximity = Collection->GetAttribute<TSet<int32>>("Proximity", FGeometryCollection::GeometryGroup);
 
-		R.ExpectTrue((Proximity)[0].Contains(1));
-		R.ExpectTrue((Proximity)[1].Contains(0));
-		R.ExpectTrue((Proximity)[1].Contains(2));
-		R.ExpectTrue((Proximity)[2].Contains(1));
+		EXPECT_TRUE((Proximity)[0].Contains(1));
+		EXPECT_TRUE((Proximity)[1].Contains(0));
+		EXPECT_TRUE((Proximity)[1].Contains(2));
+		EXPECT_TRUE((Proximity)[2].Contains(1));
 
-		R.ExpectTrue(!(Proximity)[0].Contains(2));
-		R.ExpectTrue(!(Proximity)[2].Contains(0));
-
-		return !R.HasError();
+		EXPECT_FALSE((Proximity)[0].Contains(2));
+		EXPECT_FALSE((Proximity)[2].Contains(0));
 	}
-	template bool BuildProximity<float>(ExampleResponse&& R);
+	template void BuildProximity<float>();
 
 	template<class T>
-	bool GeometryDeleteFromStart(ExampleResponse&& R)
+	void GeometryDeleteFromStart()
 	{
 		TSharedPtr<FGeometryCollection> Collection = GeometryCollection::MakeCubeElement(FTransform(FQuat::MakeFromEuler(FVector(0.f, 0.f, 0.f)), FVector(0.f, 0.f, 0.f)), FVector(1.0));
 		Collection->AppendGeometry(*GeometryCollection::MakeCubeElement(FTransform(FQuat::MakeFromEuler(FVector(0.f, 0.f, 0.f)), FVector(1.f, 0.f, 0.f)),FVector(1.0)));
@@ -95,98 +93,96 @@ namespace GeometryCollectionExample
 
 		// Proximity = [(3,4,1), (0,4,5,2), (1,5), (0,4), (0,1,3,5), (1,2,4)]
 
-		R.ExpectTrue((Proximity)[0].Contains(3));
-		R.ExpectTrue((Proximity)[0].Contains(4));
-		R.ExpectTrue((Proximity)[0].Contains(1));
-		R.ExpectTrue(!(Proximity)[0].Contains(0));
-		R.ExpectTrue(!(Proximity)[0].Contains(2));
-		R.ExpectTrue(!(Proximity)[0].Contains(5));
-
-		R.ExpectTrue((Proximity)[1].Contains(0));
-		R.ExpectTrue((Proximity)[1].Contains(4));
-		R.ExpectTrue((Proximity)[1].Contains(5));
-		R.ExpectTrue((Proximity)[1].Contains(2));
-		R.ExpectTrue(!(Proximity)[1].Contains(1));
-		R.ExpectTrue(!(Proximity)[1].Contains(3));
-
-		R.ExpectTrue((Proximity)[2].Contains(1));
-		R.ExpectTrue((Proximity)[2].Contains(5));
-		R.ExpectTrue(!(Proximity)[2].Contains(0));
-		R.ExpectTrue(!(Proximity)[2].Contains(2));
-		R.ExpectTrue(!(Proximity)[2].Contains(3));
-		R.ExpectTrue(!(Proximity)[2].Contains(4));
-
-		R.ExpectTrue((Proximity)[3].Contains(0));
-		R.ExpectTrue((Proximity)[3].Contains(4));
-		R.ExpectTrue(!(Proximity)[3].Contains(1));
-		R.ExpectTrue(!(Proximity)[3].Contains(2));
-		R.ExpectTrue(!(Proximity)[3].Contains(3));
-		R.ExpectTrue(!(Proximity)[3].Contains(5));
-
-		R.ExpectTrue((Proximity)[4].Contains(0));
-		R.ExpectTrue((Proximity)[4].Contains(1));
-		R.ExpectTrue((Proximity)[4].Contains(3));
-		R.ExpectTrue((Proximity)[4].Contains(5));
-		R.ExpectTrue(!(Proximity)[4].Contains(2));
-		R.ExpectTrue(!(Proximity)[4].Contains(4));
-
-		R.ExpectTrue((Proximity)[5].Contains(1));
-		R.ExpectTrue((Proximity)[5].Contains(2));
-		R.ExpectTrue((Proximity)[5].Contains(4));
-		R.ExpectTrue(!(Proximity)[5].Contains(0));
-		R.ExpectTrue(!(Proximity)[5].Contains(3));
-		R.ExpectTrue(!(Proximity)[5].Contains(5));
+		EXPECT_TRUE((Proximity)[0].Contains(3));
+		EXPECT_TRUE((Proximity)[0].Contains(4));
+		EXPECT_TRUE((Proximity)[0].Contains(1));
+		EXPECT_FALSE((Proximity)[0].Contains(0));
+		EXPECT_FALSE((Proximity)[0].Contains(2));
+		EXPECT_FALSE((Proximity)[0].Contains(5));
+
+		EXPECT_TRUE((Proximity)[1].Contains(0));
+		EXPECT_TRUE((Proximity)[1].Contains(4));
+		EXPECT_TRUE((Proximity)[1].Contains(5));
+		EXPECT_TRUE((Proximity)[1].Contains(2));
+		EXPECT_FALSE((Proximity)[1].Contains(1));
+		EXPECT_FALSE((Proximity)[1].Contains(3));
+
+		EXPECT_TRUE((Proximity)[2].Contains(1));
+		EXPECT_TRUE((Proximity)[2].Contains(5));
+		EXPECT_FALSE((Proximity)[2].Contains(0));
+		EXPECT_FALSE((Proximity)[2].Contains(2));
+		EXPECT_FALSE((Proximity)[2].Contains(3));
+		EXPECT_FALSE((Proximity)[2].Contains(4));
+
+		EXPECT_TRUE((Proximity)[3].Contains(0));
+		EXPECT_TRUE((Proximity)[3].Contains(4));
+		EXPECT_FALSE((Proximity)[3].Contains(1));
+		EXPECT_FALSE((Proximity)[3].Contains(2));
+		EXPECT_FALSE((Proximity)[3].Contains(3));
+		EXPECT_FALSE((Proximity)[3].Contains(5));
+
+		EXPECT_TRUE((Proximity)[4].Contains(0));
+		EXPECT_TRUE((Proximity)[4].Contains(1));
+		EXPECT_TRUE((Proximity)[4].Contains(3));
+		EXPECT_TRUE((Proximity)[4].Contains(5));
+		EXPECT_FALSE((Proximity)[4].Contains(2));
+		EXPECT_FALSE((Proximity)[4].Contains(4));
+
+		EXPECT_TRUE((Proximity)[5].Contains(1));
+		EXPECT_TRUE((Proximity)[5].Contains(2));
+		EXPECT_TRUE((Proximity)[5].Contains(4));
+		EXPECT_FALSE((Proximity)[5].Contains(0));
+		EXPECT_FALSE((Proximity)[5].Contains(3));
+		EXPECT_FALSE((Proximity)[5].Contains(5));
 
 		TArray<int32> DelList = { 0 };
 		Coll->RemoveElements(FGeometryCollection::GeometryGroup, DelList);
 
 		// Proximity = [(3,4,1), (0,4), (3), (0,2,4), (0,1,3)]
 
-		R.ExpectTrue((Proximity)[0].Contains(3));
-		R.ExpectTrue((Proximity)[0].Contains(4));
-		R.ExpectTrue((Proximity)[0].Contains(1));
-		R.ExpectTrue(!(Proximity)[0].Contains(0));
-		R.ExpectTrue(!(Proximity)[0].Contains(2));
-		R.ExpectTrue(!(Proximity)[0].Contains(5));
-
-		R.ExpectTrue((Proximity)[1].Contains(0));
-		R.ExpectTrue((Proximity)[1].Contains(4));
-		R.ExpectTrue(!(Proximity)[1].Contains(1));
-		R.ExpectTrue(!(Proximity)[1].Contains(2));
-		R.ExpectTrue(!(Proximity)[1].Contains(3));
-		R.ExpectTrue(!(Proximity)[1].Contains(5));
-
-		R.ExpectTrue((Proximity)[2].Contains(3));
-		R.ExpectTrue(!(Proximity)[2].Contains(5));
-		R.ExpectTrue(!(Proximity)[2].Contains(0));
-		R.ExpectTrue(!(Proximity)[2].Contains(1));
-		R.ExpectTrue(!(Proximity)[2].Contains(2));
-		R.ExpectTrue(!(Proximity)[2].Contains(4));
-
-		R.ExpectTrue((Proximity)[3].Contains(0));
-		R.ExpectTrue((Proximity)[3].Contains(2));
-		R.ExpectTrue((Proximity)[3].Contains(4));
-		R.ExpectTrue(!(Proximity)[3].Contains(1));
-		R.ExpectTrue(!(Proximity)[3].Contains(3));
-		R.ExpectTrue(!(Proximity)[3].Contains(5));
-
-		R.ExpectTrue((Proximity)[4].Contains(0));
-		R.ExpectTrue((Proximity)[4].Contains(1));
-		R.ExpectTrue((Proximity)[4].Contains(3));
-		R.ExpectTrue(!(Proximity)[4].Contains(2));
-		R.ExpectTrue(!(Proximity)[4].Contains(4));
-		R.ExpectTrue(!(Proximity)[4].Contains(5));
-
-		R.ExpectTrue(Collection->NumElements(FGeometryCollection::GeometryGroup) == 5);
-
-		return !R.HasError();
+		EXPECT_TRUE((Proximity)[0].Contains(3));
+		EXPECT_TRUE((Proximity)[0].Contains(4));
+		EXPECT_TRUE((Proximity)[0].Contains(1));
+		EXPECT_FALSE((Proximity)[0].Contains(0));
+		EXPECT_FALSE((Proximity)[0].Contains(2));
+		EXPECT_FALSE((Proximity)[0].Contains(5));
+
+		EXPECT_TRUE((Proximity)[1].Contains(0));
+		EXPECT_TRUE((Proximity)[1].Contains(4));
+		EXPECT_FALSE((Proximity)[1].Contains(1));
+		EXPECT_FALSE((Proximity)[1].Contains(2));
+		EXPECT_FALSE((Proximity)[1].Contains(3));
+		EXPECT_FALSE((Proximity)[1].Contains(5));
+
+		EXPECT_TRUE((Proximity)[2].Contains(3));
+		EXPECT_FALSE((Proximity)[2].Contains(5));
+		EXPECT_FALSE((Proximity)[2].Contains(0));
+		EXPECT_FALSE((Proximity)[2].Contains(1));
+		EXPECT_FALSE((Proximity)[2].Contains(2));
+		EXPECT_FALSE((Proximity)[2].Contains(4));
+
+		EXPECT_TRUE((Proximity)[3].Contains(0));
+		EXPECT_TRUE((Proximity)[3].Contains(2));
+		EXPECT_TRUE((Proximity)[3].Contains(4));
+		EXPECT_FALSE((Proximity)[3].Contains(1));
+		EXPECT_FALSE((Proximity)[3].Contains(3));
+		EXPECT_FALSE((Proximity)[3].Contains(5));
+
+		EXPECT_TRUE((Proximity)[4].Contains(0));
+		EXPECT_TRUE((Proximity)[4].Contains(1));
+		EXPECT_TRUE((Proximity)[4].Contains(3));
+		EXPECT_FALSE((Proximity)[4].Contains(2));
+		EXPECT_FALSE((Proximity)[4].Contains(4));
+		EXPECT_FALSE((Proximity)[4].Contains(5));
+
+		EXPECT_EQ(Collection->NumElements(FGeometryCollection::GeometryGroup), 5);
 	}
-	template bool GeometryDeleteFromStart<float>(ExampleResponse&& R);
+	template void GeometryDeleteFromStart<float>();
 
 
 
 	template<class T>
-	bool GeometryDeleteFromEnd(ExampleResponse&& R)
+	void GeometryDeleteFromEnd()
 	{
 		TSharedPtr<FGeometryCollection> Collection = GeometryCollection::MakeCubeElement(FTransform(FQuat::MakeFromEuler(FVector(0.f, 0.f, 0.f)), FVector(0.f, 0.f, 0.f)), FVector(1.0));
 		Collection->AppendGeometry(*GeometryCollection::MakeCubeElement(FTransform(FQuat::MakeFromEuler(FVector(0.f, 0.f, 0.f)), FVector(1.f, 0.f, 0.f)),FVector(1.0)));
@@ -228,96 +224,94 @@ namespace GeometryCollectionExample
 
 		// Proximity = [(3,4,1), (0,4,5,2), (1,5), (0,4), (0,1,3,5), (1,2,4)]
 
-		R.ExpectTrue((Proximity)[0].Contains(3));
-		R.ExpectTrue((Proximity)[0].Contains(4));
-		R.ExpectTrue((Proximity)[0].Contains(1));
-		R.ExpectTrue(!(Proximity)[0].Contains(0));
-		R.ExpectTrue(!(Proximity)[0].Contains(2));
-		R.ExpectTrue(!(Proximity)[0].Contains(5));
-
-		R.ExpectTrue((Proximity)[1].Contains(0));
-		R.ExpectTrue((Proximity)[1].Contains(4));
-		R.ExpectTrue((Proximity)[1].Contains(5));
-		R.ExpectTrue((Proximity)[1].Contains(2));
-		R.ExpectTrue(!(Proximity)[1].Contains(1));
-		R.ExpectTrue(!(Proximity)[1].Contains(3));
-
-		R.ExpectTrue((Proximity)[2].Contains(1));
-		R.ExpectTrue((Proximity)[2].Contains(5));
-		R.ExpectTrue(!(Proximity)[2].Contains(0));
-		R.ExpectTrue(!(Proximity)[2].Contains(2));
-		R.ExpectTrue(!(Proximity)[2].Contains(3));
-		R.ExpectTrue(!(Proximity)[2].Contains(4));
-
-		R.ExpectTrue((Proximity)[3].Contains(0));
-		R.ExpectTrue((Proximity)[3].Contains(4));
-		R.ExpectTrue(!(Proximity)[3].Contains(1));
-		R.ExpectTrue(!(Proximity)[3].Contains(2));
-		R.ExpectTrue(!(Proximity)[3].Contains(3));
-		R.ExpectTrue(!(Proximity)[3].Contains(5));
-
-		R.ExpectTrue((Proximity)[4].Contains(0));
-		R.ExpectTrue((Proximity)[4].Contains(1));
-		R.ExpectTrue((Proximity)[4].Contains(3));
-		R.ExpectTrue((Proximity)[4].Contains(5));
-		R.ExpectTrue(!(Proximity)[4].Contains(2));
-		R.ExpectTrue(!(Proximity)[4].Contains(4));
-
-		R.ExpectTrue((Proximity)[5].Contains(1));
-		R.ExpectTrue((Proximity)[5].Contains(2));
-		R.ExpectTrue((Proximity)[5].Contains(4));
-		R.ExpectTrue(!(Proximity)[5].Contains(0));
-		R.ExpectTrue(!(Proximity)[5].Contains(3));
-		R.ExpectTrue(!(Proximity)[5].Contains(5));
+		EXPECT_TRUE((Proximity)[0].Contains(3));
+		EXPECT_TRUE((Proximity)[0].Contains(4));
+		EXPECT_TRUE((Proximity)[0].Contains(1));
+		EXPECT_FALSE((Proximity)[0].Contains(0));
+		EXPECT_FALSE((Proximity)[0].Contains(2));
+		EXPECT_FALSE((Proximity)[0].Contains(5));
+
+		EXPECT_TRUE((Proximity)[1].Contains(0));
+		EXPECT_TRUE((Proximity)[1].Contains(4));
+		EXPECT_TRUE((Proximity)[1].Contains(5));
+		EXPECT_TRUE((Proximity)[1].Contains(2));
+		EXPECT_FALSE((Proximity)[1].Contains(1));
+		EXPECT_FALSE((Proximity)[1].Contains(3));
+
+		EXPECT_TRUE((Proximity)[2].Contains(1));
+		EXPECT_TRUE((Proximity)[2].Contains(5));
+		EXPECT_FALSE((Proximity)[2].Contains(0));
+		EXPECT_FALSE((Proximity)[2].Contains(2));
+		EXPECT_FALSE((Proximity)[2].Contains(3));
+		EXPECT_FALSE((Proximity)[2].Contains(4));
+
+		EXPECT_TRUE((Proximity)[3].Contains(0));
+		EXPECT_TRUE((Proximity)[3].Contains(4));
+		EXPECT_FALSE((Proximity)[3].Contains(1));
+		EXPECT_FALSE((Proximity)[3].Contains(2));
+		EXPECT_FALSE((Proximity)[3].Contains(3));
+		EXPECT_FALSE((Proximity)[3].Contains(5));
+
+		EXPECT_TRUE((Proximity)[4].Contains(0));
+		EXPECT_TRUE((Proximity)[4].Contains(1));
+		EXPECT_TRUE((Proximity)[4].Contains(3));
+		EXPECT_TRUE((Proximity)[4].Contains(5));
+		EXPECT_FALSE((Proximity)[4].Contains(2));
+		EXPECT_FALSE((Proximity)[4].Contains(4));
+
+		EXPECT_TRUE((Proximity)[5].Contains(1));
+		EXPECT_TRUE((Proximity)[5].Contains(2));
+		EXPECT_TRUE((Proximity)[5].Contains(4));
+		EXPECT_FALSE((Proximity)[5].Contains(0));
+		EXPECT_FALSE((Proximity)[5].Contains(3));
+		EXPECT_FALSE((Proximity)[5].Contains(5));
 
 		TArray<int32> DelList = { 5 };
 		Coll->RemoveElements(FGeometryCollection::GeometryGroup, DelList);
 
 		// Proximity = [(3,4,1), (0,4,2), (1), (0,4), (0,1,3)]
 
-		R.ExpectTrue((Proximity)[0].Contains(3));
-		R.ExpectTrue((Proximity)[0].Contains(4));
-		R.ExpectTrue((Proximity)[0].Contains(1));
-		R.ExpectTrue(!(Proximity)[0].Contains(0));
-		R.ExpectTrue(!(Proximity)[0].Contains(2));
-		R.ExpectTrue(!(Proximity)[0].Contains(5));
-
-		R.ExpectTrue((Proximity)[1].Contains(0));
-		R.ExpectTrue((Proximity)[1].Contains(4));
-		R.ExpectTrue((Proximity)[1].Contains(2));
-		R.ExpectTrue(!(Proximity)[1].Contains(1));
-		R.ExpectTrue(!(Proximity)[1].Contains(3));
-		R.ExpectTrue(!(Proximity)[1].Contains(5));
-
-		R.ExpectTrue((Proximity)[2].Contains(1));
-		R.ExpectTrue(!(Proximity)[2].Contains(5));
-		R.ExpectTrue(!(Proximity)[2].Contains(0));
-		R.ExpectTrue(!(Proximity)[2].Contains(3));
-		R.ExpectTrue(!(Proximity)[2].Contains(2));
-		R.ExpectTrue(!(Proximity)[2].Contains(4));
-
-		R.ExpectTrue((Proximity)[3].Contains(0));
-		R.ExpectTrue((Proximity)[3].Contains(4));
-		R.ExpectTrue(!(Proximity)[3].Contains(2));
-		R.ExpectTrue(!(Proximity)[3].Contains(1));
-		R.ExpectTrue(!(Proximity)[3].Contains(3));
-		R.ExpectTrue(!(Proximity)[3].Contains(5));
-
-		R.ExpectTrue((Proximity)[4].Contains(0));
-		R.ExpectTrue((Proximity)[4].Contains(1));
-		R.ExpectTrue((Proximity)[4].Contains(3));
-		R.ExpectTrue(!(Proximity)[4].Contains(2));
-		R.ExpectTrue(!(Proximity)[4].Contains(4));
-		R.ExpectTrue(!(Proximity)[4].Contains(5));
-
-		R.ExpectTrue(Collection->NumElements(FGeometryCollection::GeometryGroup) == 5);
-
-		return !R.HasError();
+		EXPECT_TRUE((Proximity)[0].Contains(3));
+		EXPECT_TRUE((Proximity)[0].Contains(4));
+		EXPECT_TRUE((Proximity)[0].Contains(1));
+		EXPECT_FALSE((Proximity)[0].Contains(0));
+		EXPECT_FALSE((Proximity)[0].Contains(2));
+		EXPECT_FALSE((Proximity)[0].Contains(5));
+
+		EXPECT_TRUE((Proximity)[1].Contains(0));
+		EXPECT_TRUE((Proximity)[1].Contains(4));
+		EXPECT_TRUE((Proximity)[1].Contains(2));
+		EXPECT_FALSE((Proximity)[1].Contains(1));
+		EXPECT_FALSE((Proximity)[1].Contains(3));
+		EXPECT_FALSE((Proximity)[1].Contains(5));
+
+		EXPECT_TRUE((Proximity)[2].Contains(1));
+		EXPECT_FALSE((Proximity)[2].Contains(5));
+		EXPECT_FALSE((Proximity)[2].Contains(0));
+		EXPECT_FALSE((Proximity)[2].Contains(3));
+		EXPECT_FALSE((Proximity)[2].Contains(2));
+		EXPECT_FALSE((Proximity)[2].Contains(4));
+
+		EXPECT_TRUE((Proximity)[3].Contains(0));
+		EXPECT_TRUE((Proximity)[3].Contains(4));
+		EXPECT_FALSE((Proximity)[3].Contains(2));
+		EXPECT_FALSE((Proximity)[3].Contains(1));
+		EXPECT_FALSE((Proximity)[3].Contains(3));
+		EXPECT_FALSE((Proximity)[3].Contains(5));
+
+		EXPECT_TRUE((Proximity)[4].Contains(0));
+		EXPECT_TRUE((Proximity)[4].Contains(1));
+		EXPECT_TRUE((Proximity)[4].Contains(3));
+		EXPECT_FALSE((Proximity)[4].Contains(2));
+		EXPECT_FALSE((Proximity)[4].Contains(4));
+		EXPECT_FALSE((Proximity)[4].Contains(5));
+
+		EXPECT_EQ(Collection->NumElements(FGeometryCollection::GeometryGroup), 5);
 	}
-	template bool GeometryDeleteFromEnd<float>(ExampleResponse&& R);
+	template void GeometryDeleteFromEnd<float>();
 
 	template<class T>
-	bool GeometryDeleteFromMiddle(ExampleResponse&& R)
+	void GeometryDeleteFromMiddle()
 	{
 		TSharedPtr<FGeometryCollection> Collection = GeometryCollection::MakeCubeElement(FTransform(FQuat::MakeFromEuler(FVector(0.f, 0.f, 0.f)), FVector(0.f, 0.f, 0.f)), FVector(1.0));
 		Collection->AppendGeometry(*GeometryCollection::MakeCubeElement(FTransform(FQuat::MakeFromEuler(FVector(0.f, 0.f, 0.f)), FVector(1.f, 0.f, 0.f)),FVector(1.0)));
@@ -359,97 +353,95 @@ namespace GeometryCollectionExample
 
 		// Proximity = [(3,4,1), (0,4,5,2), (1,5), (0,4), (0,1,3,5), (1,2,4)]
 
-		R.ExpectTrue((Proximity)[0].Contains(3));
-		R.ExpectTrue((Proximity)[0].Contains(4));
-		R.ExpectTrue((Proximity)[0].Contains(1));
-		R.ExpectTrue(!(Proximity)[0].Contains(0));
-		R.ExpectTrue(!(Proximity)[0].Contains(2));
-		R.ExpectTrue(!(Proximity)[0].Contains(5));
-
-		R.ExpectTrue((Proximity)[1].Contains(0));
-		R.ExpectTrue((Proximity)[1].Contains(4));
-		R.ExpectTrue((Proximity)[1].Contains(5));
-		R.ExpectTrue((Proximity)[1].Contains(2));
-		R.ExpectTrue(!(Proximity)[1].Contains(1));
-		R.ExpectTrue(!(Proximity)[1].Contains(3));
-
-		R.ExpectTrue((Proximity)[2].Contains(1));
-		R.ExpectTrue((Proximity)[2].Contains(5));
-		R.ExpectTrue(!(Proximity)[2].Contains(0));
-		R.ExpectTrue(!(Proximity)[2].Contains(2));
-		R.ExpectTrue(!(Proximity)[2].Contains(3));
-		R.ExpectTrue(!(Proximity)[2].Contains(4));
-
-		R.ExpectTrue((Proximity)[3].Contains(0));
-		R.ExpectTrue((Proximity)[3].Contains(4));
-		R.ExpectTrue(!(Proximity)[3].Contains(1));
-		R.ExpectTrue(!(Proximity)[3].Contains(2));
-		R.ExpectTrue(!(Proximity)[3].Contains(3));
-		R.ExpectTrue(!(Proximity)[3].Contains(5));
-
-		R.ExpectTrue((Proximity)[4].Contains(0));
-		R.ExpectTrue((Proximity)[4].Contains(1));
-		R.ExpectTrue((Proximity)[4].Contains(3));
-		R.ExpectTrue((Proximity)[4].Contains(5));
-		R.ExpectTrue(!(Proximity)[4].Contains(2));
-		R.ExpectTrue(!(Proximity)[4].Contains(4));
-
-		R.ExpectTrue((Proximity)[5].Contains(1));
-		R.ExpectTrue((Proximity)[5].Contains(2));
-		R.ExpectTrue((Proximity)[5].Contains(4));
-		R.ExpectTrue(!(Proximity)[5].Contains(0));
-		R.ExpectTrue(!(Proximity)[5].Contains(3));
-		R.ExpectTrue(!(Proximity)[5].Contains(5));
+		EXPECT_TRUE((Proximity)[0].Contains(3));
+		EXPECT_TRUE((Proximity)[0].Contains(4));
+		EXPECT_TRUE((Proximity)[0].Contains(1));
+		EXPECT_FALSE((Proximity)[0].Contains(0));
+		EXPECT_FALSE((Proximity)[0].Contains(2));
+		EXPECT_FALSE((Proximity)[0].Contains(5));
+
+		EXPECT_TRUE((Proximity)[1].Contains(0));
+		EXPECT_TRUE((Proximity)[1].Contains(4));
+		EXPECT_TRUE((Proximity)[1].Contains(5));
+		EXPECT_TRUE((Proximity)[1].Contains(2));
+		EXPECT_FALSE((Proximity)[1].Contains(1));
+		EXPECT_FALSE((Proximity)[1].Contains(3));
+
+		EXPECT_TRUE((Proximity)[2].Contains(1));
+		EXPECT_TRUE((Proximity)[2].Contains(5));
+		EXPECT_FALSE((Proximity)[2].Contains(0));
+		EXPECT_FALSE((Proximity)[2].Contains(2));
+		EXPECT_FALSE((Proximity)[2].Contains(3));
+		EXPECT_FALSE((Proximity)[2].Contains(4));
+
+		EXPECT_TRUE((Proximity)[3].Contains(0));
+		EXPECT_TRUE((Proximity)[3].Contains(4));
+		EXPECT_FALSE((Proximity)[3].Contains(1));
+		EXPECT_FALSE((Proximity)[3].Contains(2));
+		EXPECT_FALSE((Proximity)[3].Contains(3));
+		EXPECT_FALSE((Proximity)[3].Contains(5));
+
+		EXPECT_TRUE((Proximity)[4].Contains(0));
+		EXPECT_TRUE((Proximity)[4].Contains(1));
+		EXPECT_TRUE((Proximity)[4].Contains(3));
+		EXPECT_TRUE((Proximity)[4].Contains(5));
+		EXPECT_FALSE((Proximity)[4].Contains(2));
+		EXPECT_FALSE((Proximity)[4].Contains(4));
+
+		EXPECT_TRUE((Proximity)[5].Contains(1));
+		EXPECT_TRUE((Proximity)[5].Contains(2));
+		EXPECT_TRUE((Proximity)[5].Contains(4));
+		EXPECT_FALSE((Proximity)[5].Contains(0));
+		EXPECT_FALSE((Proximity)[5].Contains(3));
+		EXPECT_FALSE((Proximity)[5].Contains(5));
 
 		TArray<int32> DelList = { 3 };
 		Coll->RemoveElements(FGeometryCollection::GeometryGroup, DelList);
 
 		// Proximity = [(3,1), (0,3,4,2), (1,4), (0,1,4), (1,2,3)]
 
-		R.ExpectTrue((Proximity)[0].Contains(3));
-		R.ExpectTrue((Proximity)[0].Contains(1));
-		R.ExpectTrue(!(Proximity)[0].Contains(2));
-		R.ExpectTrue(!(Proximity)[0].Contains(0));
-		R.ExpectTrue(!(Proximity)[0].Contains(4));
-		R.ExpectTrue(!(Proximity)[0].Contains(5));
+		EXPECT_TRUE((Proximity)[0].Contains(3));
+		EXPECT_TRUE((Proximity)[0].Contains(1));
+		EXPECT_FALSE((Proximity)[0].Contains(2));
+		EXPECT_FALSE((Proximity)[0].Contains(0));
+		EXPECT_FALSE((Proximity)[0].Contains(4));
+		EXPECT_FALSE((Proximity)[0].Contains(5));
 
 	
-		R.ExpectTrue((Proximity)[1].Contains(0));
-		R.ExpectTrue((Proximity)[1].Contains(3));
-		R.ExpectTrue((Proximity)[1].Contains(4));
-		R.ExpectTrue((Proximity)[1].Contains(2));
-		R.ExpectTrue(!(Proximity)[1].Contains(1));
-		R.ExpectTrue(!(Proximity)[1].Contains(5));
-
-		R.ExpectTrue((Proximity)[2].Contains(1));
-		R.ExpectTrue((Proximity)[2].Contains(4));
-		R.ExpectTrue(!(Proximity)[2].Contains(0));
-		R.ExpectTrue(!(Proximity)[2].Contains(2));
-		R.ExpectTrue(!(Proximity)[2].Contains(3));
-		R.ExpectTrue(!(Proximity)[2].Contains(5));
-
-		R.ExpectTrue((Proximity)[3].Contains(0));
-		R.ExpectTrue((Proximity)[3].Contains(1));
-		R.ExpectTrue((Proximity)[3].Contains(4));
-		R.ExpectTrue(!(Proximity)[3].Contains(2));
-		R.ExpectTrue(!(Proximity)[3].Contains(3));
-		R.ExpectTrue(!(Proximity)[3].Contains(5));
-
-		R.ExpectTrue((Proximity)[4].Contains(1));
-		R.ExpectTrue((Proximity)[4].Contains(2));
-		R.ExpectTrue((Proximity)[4].Contains(3));
-		R.ExpectTrue(!(Proximity)[4].Contains(0));
-		R.ExpectTrue(!(Proximity)[4].Contains(4));
-		R.ExpectTrue(!(Proximity)[4].Contains(5));
-
-		R.ExpectTrue(Collection->NumElements(FGeometryCollection::GeometryGroup) == 5);
-
-		return !R.HasError();
+		EXPECT_TRUE((Proximity)[1].Contains(0));
+		EXPECT_TRUE((Proximity)[1].Contains(3));
+		EXPECT_TRUE((Proximity)[1].Contains(4));
+		EXPECT_TRUE((Proximity)[1].Contains(2));
+		EXPECT_FALSE((Proximity)[1].Contains(1));
+		EXPECT_FALSE((Proximity)[1].Contains(5));
+
+		EXPECT_TRUE((Proximity)[2].Contains(1));
+		EXPECT_TRUE((Proximity)[2].Contains(4));
+		EXPECT_FALSE((Proximity)[2].Contains(0));
+		EXPECT_FALSE((Proximity)[2].Contains(2));
+		EXPECT_FALSE((Proximity)[2].Contains(3));
+		EXPECT_FALSE((Proximity)[2].Contains(5));
+
+		EXPECT_TRUE((Proximity)[3].Contains(0));
+		EXPECT_TRUE((Proximity)[3].Contains(1));
+		EXPECT_TRUE((Proximity)[3].Contains(4));
+		EXPECT_FALSE((Proximity)[3].Contains(2));
+		EXPECT_FALSE((Proximity)[3].Contains(3));
+		EXPECT_FALSE((Proximity)[3].Contains(5));
+
+		EXPECT_TRUE((Proximity)[4].Contains(1));
+		EXPECT_TRUE((Proximity)[4].Contains(2));
+		EXPECT_TRUE((Proximity)[4].Contains(3));
+		EXPECT_FALSE((Proximity)[4].Contains(0));
+		EXPECT_FALSE((Proximity)[4].Contains(4));
+		EXPECT_FALSE((Proximity)[4].Contains(5));
+
+		EXPECT_EQ(Collection->NumElements(FGeometryCollection::GeometryGroup), 5);
 	}
-	template bool GeometryDeleteFromMiddle<float>(ExampleResponse&& R);
+	template void GeometryDeleteFromMiddle<float>();
 
 	template<class T>
-	bool GeometryDeleteMultipleFromMiddle(ExampleResponse&& R)
+	void GeometryDeleteMultipleFromMiddle()
 	{
 		TSharedPtr<FGeometryCollection> Collection = GeometryCollection::MakeCubeElement(FTransform(FQuat::MakeFromEuler(FVector(0.f, 0.f, 0.f)), FVector(0.f, 0.f, 0.f)), FVector(1.0));
 		Collection->AppendGeometry(*GeometryCollection::MakeCubeElement(FTransform(FQuat::MakeFromEuler(FVector(0.f, 0.f, 0.f)), FVector(1.f, 0.f, 0.f)),FVector(1.0)));
@@ -491,82 +483,80 @@ namespace GeometryCollectionExample
 
 		// Proximity = [(3,4,1), (0,4,5,2), (1,5), (0,4), (0,1,3,5), (1,2,4)]
 
-		R.ExpectTrue((Proximity)[0].Contains(3));
-		R.ExpectTrue((Proximity)[0].Contains(4));
-		R.ExpectTrue((Proximity)[0].Contains(1));
-		R.ExpectTrue(!(Proximity)[0].Contains(0));
-		R.ExpectTrue(!(Proximity)[0].Contains(2));
-		R.ExpectTrue(!(Proximity)[0].Contains(5));
-
-		R.ExpectTrue((Proximity)[1].Contains(0));
-		R.ExpectTrue((Proximity)[1].Contains(4));
-		R.ExpectTrue((Proximity)[1].Contains(5));
-		R.ExpectTrue((Proximity)[1].Contains(2));
-		R.ExpectTrue(!(Proximity)[1].Contains(1));
-		R.ExpectTrue(!(Proximity)[1].Contains(3));
-
-		R.ExpectTrue((Proximity)[2].Contains(1));
-		R.ExpectTrue((Proximity)[2].Contains(5));
-		R.ExpectTrue(!(Proximity)[2].Contains(0));
-		R.ExpectTrue(!(Proximity)[2].Contains(2));
-		R.ExpectTrue(!(Proximity)[2].Contains(3));
-		R.ExpectTrue(!(Proximity)[2].Contains(4));
-
-		R.ExpectTrue((Proximity)[3].Contains(0));
-		R.ExpectTrue((Proximity)[3].Contains(4));
-		R.ExpectTrue(!(Proximity)[3].Contains(1));
-		R.ExpectTrue(!(Proximity)[3].Contains(2));
-		R.ExpectTrue(!(Proximity)[3].Contains(3));
-		R.ExpectTrue(!(Proximity)[3].Contains(5));
-
-		R.ExpectTrue((Proximity)[4].Contains(0));
-		R.ExpectTrue((Proximity)[4].Contains(1));
-		R.ExpectTrue((Proximity)[4].Contains(3));
-		R.ExpectTrue((Proximity)[4].Contains(5));
-		R.ExpectTrue(!(Proximity)[4].Contains(2));
-		R.ExpectTrue(!(Proximity)[4].Contains(4));
-
-		R.ExpectTrue((Proximity)[5].Contains(1));
-		R.ExpectTrue((Proximity)[5].Contains(2));
-		R.ExpectTrue((Proximity)[5].Contains(4));
-		R.ExpectTrue(!(Proximity)[5].Contains(0));
-		R.ExpectTrue(!(Proximity)[5].Contains(3));
-		R.ExpectTrue(!(Proximity)[5].Contains(5));
+		EXPECT_TRUE((Proximity)[0].Contains(3));
+		EXPECT_TRUE((Proximity)[0].Contains(4));
+		EXPECT_TRUE((Proximity)[0].Contains(1));
+		EXPECT_FALSE((Proximity)[0].Contains(0));
+		EXPECT_FALSE((Proximity)[0].Contains(2));
+		EXPECT_FALSE((Proximity)[0].Contains(5));
+
+		EXPECT_TRUE((Proximity)[1].Contains(0));
+		EXPECT_TRUE((Proximity)[1].Contains(4));
+		EXPECT_TRUE((Proximity)[1].Contains(5));
+		EXPECT_TRUE((Proximity)[1].Contains(2));
+		EXPECT_FALSE((Proximity)[1].Contains(1));
+		EXPECT_FALSE((Proximity)[1].Contains(3));
+
+		EXPECT_TRUE((Proximity)[2].Contains(1));
+		EXPECT_TRUE((Proximity)[2].Contains(5));
+		EXPECT_FALSE((Proximity)[2].Contains(0));
+		EXPECT_FALSE((Proximity)[2].Contains(2));
+		EXPECT_FALSE((Proximity)[2].Contains(3));
+		EXPECT_FALSE((Proximity)[2].Contains(4));
+
+		EXPECT_TRUE((Proximity)[3].Contains(0));
+		EXPECT_TRUE((Proximity)[3].Contains(4));
+		EXPECT_FALSE((Proximity)[3].Contains(1));
+		EXPECT_FALSE((Proximity)[3].Contains(2));
+		EXPECT_FALSE((Proximity)[3].Contains(3));
+		EXPECT_FALSE((Proximity)[3].Contains(5));
+
+		EXPECT_TRUE((Proximity)[4].Contains(0));
+		EXPECT_TRUE((Proximity)[4].Contains(1));
+		EXPECT_TRUE((Proximity)[4].Contains(3));
+		EXPECT_TRUE((Proximity)[4].Contains(5));
+		EXPECT_FALSE((Proximity)[4].Contains(2));
+		EXPECT_FALSE((Proximity)[4].Contains(4));
+
+		EXPECT_TRUE((Proximity)[5].Contains(1));
+		EXPECT_TRUE((Proximity)[5].Contains(2));
+		EXPECT_TRUE((Proximity)[5].Contains(4));
+		EXPECT_FALSE((Proximity)[5].Contains(0));
+		EXPECT_FALSE((Proximity)[5].Contains(3));
+		EXPECT_FALSE((Proximity)[5].Contains(5));
 
 		TArray<int32> DelList = { 2,3,4 };
 		Coll->RemoveElements(FGeometryCollection::GeometryGroup, DelList);
 
 		// Proximity = [(1), (0,2), (1)]
 
-		R.ExpectTrue((Proximity)[0].Contains(1));
-		R.ExpectTrue(!(Proximity)[0].Contains(0));
-		R.ExpectTrue(!(Proximity)[0].Contains(2));
-		R.ExpectTrue(!(Proximity)[0].Contains(3));
-		R.ExpectTrue(!(Proximity)[0].Contains(4));
-		R.ExpectTrue(!(Proximity)[0].Contains(5));
-
-		R.ExpectTrue((Proximity)[1].Contains(0));
-		R.ExpectTrue((Proximity)[1].Contains(2));
-		R.ExpectTrue(!(Proximity)[1].Contains(1));
-		R.ExpectTrue(!(Proximity)[1].Contains(3));
-		R.ExpectTrue(!(Proximity)[1].Contains(4));
-		R.ExpectTrue(!(Proximity)[1].Contains(5));
-
-		R.ExpectTrue((Proximity)[2].Contains(1));
-		R.ExpectTrue(!(Proximity)[2].Contains(4));
-		R.ExpectTrue(!(Proximity)[2].Contains(0));
-		R.ExpectTrue(!(Proximity)[2].Contains(2));
-		R.ExpectTrue(!(Proximity)[2].Contains(3));
-		R.ExpectTrue(!(Proximity)[2].Contains(5));
-
-		R.ExpectTrue(Collection->NumElements(FGeometryCollection::GeometryGroup) == 3);
-
-		return !R.HasError();
+		EXPECT_TRUE((Proximity)[0].Contains(1));
+		EXPECT_FALSE((Proximity)[0].Contains(0));
+		EXPECT_FALSE((Proximity)[0].Contains(2));
+		EXPECT_FALSE((Proximity)[0].Contains(3));
+		EXPECT_FALSE((Proximity)[0].Contains(4));
+		EXPECT_FALSE((Proximity)[0].Contains(5));
+
+		EXPECT_TRUE((Proximity)[1].Contains(0));
+		EXPECT_TRUE((Proximity)[1].Contains(2));
+		EXPECT_FALSE((Proximity)[1].Contains(1));
+		EXPECT_FALSE((Proximity)[1].Contains(3));
+		EXPECT_FALSE((Proximity)[1].Contains(4));
+		EXPECT_FALSE((Proximity)[1].Contains(5));
+
+		EXPECT_TRUE((Proximity)[2].Contains(1));
+		EXPECT_FALSE((Proximity)[2].Contains(4));
+		EXPECT_FALSE((Proximity)[2].Contains(0));
+		EXPECT_FALSE((Proximity)[2].Contains(2));
+		EXPECT_FALSE((Proximity)[2].Contains(3));
+		EXPECT_FALSE((Proximity)[2].Contains(5));
+
+		EXPECT_EQ(Collection->NumElements(FGeometryCollection::GeometryGroup), 3);
 	}
-	template bool GeometryDeleteMultipleFromMiddle<float>(ExampleResponse&& R);
+	template void GeometryDeleteMultipleFromMiddle<float>();
 
 	template<class T>
-	bool GeometryDeleteRandom(ExampleResponse&& R)
+	void GeometryDeleteRandom()
 	{
 		TSharedPtr<FGeometryCollection> Collection = GeometryCollection::MakeCubeElement(FTransform(FQuat::MakeFromEuler(FVector(0.f, 0.f, 0.f)), FVector(0.f, 0.f, 0.f)), FVector(1.0));
 		Collection->AppendGeometry(*GeometryCollection::MakeCubeElement(FTransform(FQuat::MakeFromEuler(FVector(0.f, 0.f, 0.f)), FVector(1.f, 0.f, 0.f)),FVector(1.0)));
@@ -608,82 +598,80 @@ namespace GeometryCollectionExample
 
 		// Proximity = [(3,4,1), (0,4,5,2), (1,5), (0,4), (0,1,3,5), (1,2,4)]
 
-		R.ExpectTrue((Proximity)[0].Contains(3));
-		R.ExpectTrue((Proximity)[0].Contains(4));
-		R.ExpectTrue((Proximity)[0].Contains(1));
-		R.ExpectTrue(!(Proximity)[0].Contains(0));
-		R.ExpectTrue(!(Proximity)[0].Contains(2));
-		R.ExpectTrue(!(Proximity)[0].Contains(5));
-
-		R.ExpectTrue((Proximity)[1].Contains(0));
-		R.ExpectTrue((Proximity)[1].Contains(4));
-		R.ExpectTrue((Proximity)[1].Contains(5));
-		R.ExpectTrue((Proximity)[1].Contains(2));
-		R.ExpectTrue(!(Proximity)[1].Contains(1));
-		R.ExpectTrue(!(Proximity)[1].Contains(3));
-
-		R.ExpectTrue((Proximity)[2].Contains(1));
-		R.ExpectTrue((Proximity)[2].Contains(5));
-		R.ExpectTrue(!(Proximity)[2].Contains(0));
-		R.ExpectTrue(!(Proximity)[2].Contains(2));
-		R.ExpectTrue(!(Proximity)[2].Contains(3));
-		R.ExpectTrue(!(Proximity)[2].Contains(4));
-
-		R.ExpectTrue((Proximity)[3].Contains(0));
-		R.ExpectTrue((Proximity)[3].Contains(4));
-		R.ExpectTrue(!(Proximity)[3].Contains(1));
-		R.ExpectTrue(!(Proximity)[3].Contains(2));
-		R.ExpectTrue(!(Proximity)[3].Contains(3));
-		R.ExpectTrue(!(Proximity)[3].Contains(5));
-
-		R.ExpectTrue((Proximity)[4].Contains(0));
-		R.ExpectTrue((Proximity)[4].Contains(1));
-		R.ExpectTrue((Proximity)[4].Contains(3));
-		R.ExpectTrue((Proximity)[4].Contains(5));
-		R.ExpectTrue(!(Proximity)[4].Contains(2));
-		R.ExpectTrue(!(Proximity)[4].Contains(4));
-
-		R.ExpectTrue((Proximity)[5].Contains(1));
-		R.ExpectTrue((Proximity)[5].Contains(2));
-		R.ExpectTrue((Proximity)[5].Contains(4));
-		R.ExpectTrue(!(Proximity)[5].Contains(0));
-		R.ExpectTrue(!(Proximity)[5].Contains(3));
-		R.ExpectTrue(!(Proximity)[5].Contains(5));
+		EXPECT_TRUE((Proximity)[0].Contains(3));
+		EXPECT_TRUE((Proximity)[0].Contains(4));
+		EXPECT_TRUE((Proximity)[0].Contains(1));
+		EXPECT_FALSE((Proximity)[0].Contains(0));
+		EXPECT_FALSE((Proximity)[0].Contains(2));
+		EXPECT_FALSE((Proximity)[0].Contains(5));
+
+		EXPECT_TRUE((Proximity)[1].Contains(0));
+		EXPECT_TRUE((Proximity)[1].Contains(4));
+		EXPECT_TRUE((Proximity)[1].Contains(5));
+		EXPECT_TRUE((Proximity)[1].Contains(2));
+		EXPECT_FALSE((Proximity)[1].Contains(1));
+		EXPECT_FALSE((Proximity)[1].Contains(3));
+
+		EXPECT_TRUE((Proximity)[2].Contains(1));
+		EXPECT_TRUE((Proximity)[2].Contains(5));
+		EXPECT_FALSE((Proximity)[2].Contains(0));
+		EXPECT_FALSE((Proximity)[2].Contains(2));
+		EXPECT_FALSE((Proximity)[2].Contains(3));
+		EXPECT_FALSE((Proximity)[2].Contains(4));
+
+		EXPECT_TRUE((Proximity)[3].Contains(0));
+		EXPECT_TRUE((Proximity)[3].Contains(4));
+		EXPECT_FALSE((Proximity)[3].Contains(1));
+		EXPECT_FALSE((Proximity)[3].Contains(2));
+		EXPECT_FALSE((Proximity)[3].Contains(3));
+		EXPECT_FALSE((Proximity)[3].Contains(5));
+
+		EXPECT_TRUE((Proximity)[4].Contains(0));
+		EXPECT_TRUE((Proximity)[4].Contains(1));
+		EXPECT_TRUE((Proximity)[4].Contains(3));
+		EXPECT_TRUE((Proximity)[4].Contains(5));
+		EXPECT_FALSE((Proximity)[4].Contains(2));
+		EXPECT_FALSE((Proximity)[4].Contains(4));
+
+		EXPECT_TRUE((Proximity)[5].Contains(1));
+		EXPECT_TRUE((Proximity)[5].Contains(2));
+		EXPECT_TRUE((Proximity)[5].Contains(4));
+		EXPECT_FALSE((Proximity)[5].Contains(0));
+		EXPECT_FALSE((Proximity)[5].Contains(3));
+		EXPECT_FALSE((Proximity)[5].Contains(5));
 
 		TArray<int32> DelList = { 1,3,5 };
 		Coll->RemoveElements(FGeometryCollection::GeometryGroup, DelList);
 
 		// Proximity = [(2), (), (0)]
 
-		R.ExpectTrue((Proximity)[0].Contains(2));
-		R.ExpectTrue(!(Proximity)[0].Contains(0));
-		R.ExpectTrue(!(Proximity)[0].Contains(1));
-		R.ExpectTrue(!(Proximity)[0].Contains(3));
-		R.ExpectTrue(!(Proximity)[0].Contains(4));
-		R.ExpectTrue(!(Proximity)[0].Contains(5));
-
-		R.ExpectTrue(!(Proximity)[1].Contains(0));
-		R.ExpectTrue(!(Proximity)[1].Contains(2));
-		R.ExpectTrue(!(Proximity)[1].Contains(1));
-		R.ExpectTrue(!(Proximity)[1].Contains(3));
-		R.ExpectTrue(!(Proximity)[1].Contains(4));
-		R.ExpectTrue(!(Proximity)[1].Contains(5));
-
-		R.ExpectTrue((Proximity)[2].Contains(0));
-		R.ExpectTrue(!(Proximity)[2].Contains(2));
-		R.ExpectTrue(!(Proximity)[2].Contains(1));
-		R.ExpectTrue(!(Proximity)[2].Contains(3));
-		R.ExpectTrue(!(Proximity)[2].Contains(4));
-		R.ExpectTrue(!(Proximity)[2].Contains(5));
-
-		R.ExpectTrue(Collection->NumElements(FGeometryCollection::GeometryGroup) == 3);
-
-		return !R.HasError();
+		EXPECT_TRUE((Proximity)[0].Contains(2));
+		EXPECT_FALSE((Proximity)[0].Contains(0));
+		EXPECT_FALSE((Proximity)[0].Contains(1));
+		EXPECT_FALSE((Proximity)[0].Contains(3));
+		EXPECT_FALSE((Proximity)[0].Contains(4));
+		EXPECT_FALSE((Proximity)[0].Contains(5));
+
+		EXPECT_FALSE((Proximity)[1].Contains(0));
+		EXPECT_FALSE((Proximity)[1].Contains(2));
+		EXPECT_FALSE((Proximity)[1].Contains(1));
+		EXPECT_FALSE((Proximity)[1].Contains(3));
+		EXPECT_FALSE((Proximity)[1].Contains(4));
+		EXPECT_FALSE((Proximity)[1].Contains(5));
+
+		EXPECT_TRUE((Proximity)[2].Contains(0));
+		EXPECT_FALSE((Proximity)[2].Contains(2));
+		EXPECT_FALSE((Proximity)[2].Contains(1));
+		EXPECT_FALSE((Proximity)[2].Contains(3));
+		EXPECT_FALSE((Proximity)[2].Contains(4));
+		EXPECT_FALSE((Proximity)[2].Contains(5));
+
+		EXPECT_EQ(Collection->NumElements(FGeometryCollection::GeometryGroup), 3);
 	}
-	template bool GeometryDeleteRandom<float>(ExampleResponse&& R); 
+	template void GeometryDeleteRandom<float>(); 
 
 	template<class T>
-	bool GeometryDeleteRandom2(ExampleResponse&& R)
+	void GeometryDeleteRandom2()
 	{
 		TSharedPtr<FGeometryCollection> Collection = GeometryCollection::MakeCubeElement(FTransform(FQuat::MakeFromEuler(FVector(0.f, 0.f, 0.f)), FVector(0.f, 0.f, 0.f)), FVector(1.0));
 		Collection->AppendGeometry(*GeometryCollection::MakeCubeElement(FTransform(FQuat::MakeFromEuler(FVector(0.f, 0.f, 0.f)), FVector(1.f, 0.f, 0.f)),FVector(1.0)));
@@ -725,75 +713,73 @@ namespace GeometryCollectionExample
 
 		// Proximity = [(3,4,1), (0,4,5,2), (1,5), (0,4), (0,1,3,5), (1,2,4)]
 
-		R.ExpectTrue((Proximity)[0].Contains(3));
-		R.ExpectTrue((Proximity)[0].Contains(4));
-		R.ExpectTrue((Proximity)[0].Contains(1));
-		R.ExpectTrue(!(Proximity)[0].Contains(0));
-		R.ExpectTrue(!(Proximity)[0].Contains(2));
-		R.ExpectTrue(!(Proximity)[0].Contains(5));
-
-		R.ExpectTrue((Proximity)[1].Contains(0));
-		R.ExpectTrue((Proximity)[1].Contains(4));
-		R.ExpectTrue((Proximity)[1].Contains(5));
-		R.ExpectTrue((Proximity)[1].Contains(2));
-		R.ExpectTrue(!(Proximity)[1].Contains(1));
-		R.ExpectTrue(!(Proximity)[1].Contains(3));
-
-		R.ExpectTrue((Proximity)[2].Contains(1));
-		R.ExpectTrue((Proximity)[2].Contains(5));
-		R.ExpectTrue(!(Proximity)[2].Contains(0));
-		R.ExpectTrue(!(Proximity)[2].Contains(2));
-		R.ExpectTrue(!(Proximity)[2].Contains(3));
-		R.ExpectTrue(!(Proximity)[2].Contains(4));
-
-		R.ExpectTrue((Proximity)[3].Contains(0));
-		R.ExpectTrue((Proximity)[3].Contains(4));
-		R.ExpectTrue(!(Proximity)[3].Contains(1));
-		R.ExpectTrue(!(Proximity)[3].Contains(2));
-		R.ExpectTrue(!(Proximity)[3].Contains(3));
-		R.ExpectTrue(!(Proximity)[3].Contains(5));
-
-		R.ExpectTrue((Proximity)[4].Contains(0));
-		R.ExpectTrue((Proximity)[4].Contains(1));
-		R.ExpectTrue((Proximity)[4].Contains(3));
-		R.ExpectTrue((Proximity)[4].Contains(5));
-		R.ExpectTrue(!(Proximity)[4].Contains(2));
-		R.ExpectTrue(!(Proximity)[4].Contains(4));
-
-		R.ExpectTrue((Proximity)[5].Contains(1));
-		R.ExpectTrue((Proximity)[5].Contains(2));
-		R.ExpectTrue((Proximity)[5].Contains(4));
-		R.ExpectTrue(!(Proximity)[5].Contains(0));
-		R.ExpectTrue(!(Proximity)[5].Contains(3));
-		R.ExpectTrue(!(Proximity)[5].Contains(5));
+		EXPECT_TRUE((Proximity)[0].Contains(3));
+		EXPECT_TRUE((Proximity)[0].Contains(4));
+		EXPECT_TRUE((Proximity)[0].Contains(1));
+		EXPECT_FALSE((Proximity)[0].Contains(0));
+		EXPECT_FALSE((Proximity)[0].Contains(2));
+		EXPECT_FALSE((Proximity)[0].Contains(5));
+
+		EXPECT_TRUE((Proximity)[1].Contains(0));
+		EXPECT_TRUE((Proximity)[1].Contains(4));
+		EXPECT_TRUE((Proximity)[1].Contains(5));
+		EXPECT_TRUE((Proximity)[1].Contains(2));
+		EXPECT_FALSE((Proximity)[1].Contains(1));
+		EXPECT_FALSE((Proximity)[1].Contains(3));
+
+		EXPECT_TRUE((Proximity)[2].Contains(1));
+		EXPECT_TRUE((Proximity)[2].Contains(5));
+		EXPECT_FALSE((Proximity)[2].Contains(0));
+		EXPECT_FALSE((Proximity)[2].Contains(2));
+		EXPECT_FALSE((Proximity)[2].Contains(3));
+		EXPECT_FALSE((Proximity)[2].Contains(4));
+
+		EXPECT_TRUE((Proximity)[3].Contains(0));
+		EXPECT_TRUE((Proximity)[3].Contains(4));
+		EXPECT_FALSE((Proximity)[3].Contains(1));
+		EXPECT_FALSE((Proximity)[3].Contains(2));
+		EXPECT_FALSE((Proximity)[3].Contains(3));
+		EXPECT_FALSE((Proximity)[3].Contains(5));
+
+		EXPECT_TRUE((Proximity)[4].Contains(0));
+		EXPECT_TRUE((Proximity)[4].Contains(1));
+		EXPECT_TRUE((Proximity)[4].Contains(3));
+		EXPECT_TRUE((Proximity)[4].Contains(5));
+		EXPECT_FALSE((Proximity)[4].Contains(2));
+		EXPECT_FALSE((Proximity)[4].Contains(4));
+
+		EXPECT_TRUE((Proximity)[5].Contains(1));
+		EXPECT_TRUE((Proximity)[5].Contains(2));
+		EXPECT_TRUE((Proximity)[5].Contains(4));
+		EXPECT_FALSE((Proximity)[5].Contains(0));
+		EXPECT_FALSE((Proximity)[5].Contains(3));
+		EXPECT_FALSE((Proximity)[5].Contains(5));
 
 		TArray<int32> DelList = { 0,1,4,5 };
 		Coll->RemoveElements(FGeometryCollection::GeometryGroup, DelList);
 
 		// Proximity = [(), ()]
 
-		R.ExpectTrue(!(Proximity)[0].Contains(1));
-		R.ExpectTrue(!(Proximity)[0].Contains(0));
-		R.ExpectTrue(!(Proximity)[0].Contains(2));
-		R.ExpectTrue(!(Proximity)[0].Contains(3));
-		R.ExpectTrue(!(Proximity)[0].Contains(4));
-		R.ExpectTrue(!(Proximity)[0].Contains(5));
-
-		R.ExpectTrue(!(Proximity)[1].Contains(0));
-		R.ExpectTrue(!(Proximity)[1].Contains(2));
-		R.ExpectTrue(!(Proximity)[1].Contains(1));
-		R.ExpectTrue(!(Proximity)[1].Contains(3));
-		R.ExpectTrue(!(Proximity)[1].Contains(4));
-		R.ExpectTrue(!(Proximity)[1].Contains(5));
-
-		R.ExpectTrue(Collection->NumElements(FGeometryCollection::GeometryGroup) == 2);
-
-		return !R.HasError();
+		EXPECT_FALSE((Proximity)[0].Contains(1));
+		EXPECT_FALSE((Proximity)[0].Contains(0));
+		EXPECT_FALSE((Proximity)[0].Contains(2));
+		EXPECT_FALSE((Proximity)[0].Contains(3));
+		EXPECT_FALSE((Proximity)[0].Contains(4));
+		EXPECT_FALSE((Proximity)[0].Contains(5));
+
+		EXPECT_FALSE((Proximity)[1].Contains(0));
+		EXPECT_FALSE((Proximity)[1].Contains(2));
+		EXPECT_FALSE((Proximity)[1].Contains(1));
+		EXPECT_FALSE((Proximity)[1].Contains(3));
+		EXPECT_FALSE((Proximity)[1].Contains(4));
+		EXPECT_FALSE((Proximity)[1].Contains(5));
+
+		EXPECT_EQ(Collection->NumElements(FGeometryCollection::GeometryGroup), 2);
 	}
-	template bool GeometryDeleteRandom2<float>(ExampleResponse&& R);
+	template void GeometryDeleteRandom2<float>();
 
 	template<class T>
-	bool GeometryDeleteAll(ExampleResponse&& R)
+	void GeometryDeleteAll()
 	{
 		TSharedPtr<FGeometryCollection> Collection = GeometryCollection::MakeCubeElement(FTransform(FQuat::MakeFromEuler(FVector(0.f, 0.f, 0.f)), FVector(0.f, 0.f, 0.f)), FVector(1.0));
 		Collection->AppendGeometry(*GeometryCollection::MakeCubeElement(FTransform(FQuat::MakeFromEuler(FVector(0.f, 0.f, 0.f)), FVector(1.f, 0.f, 0.f)),FVector(1.0)));
@@ -835,61 +821,59 @@ namespace GeometryCollectionExample
 
 		// Proximity = [(3,4,1), (0,4,5,2), (1,5), (0,4), (0,1,3,5), (1,2,4)]
 
-		R.ExpectTrue((Proximity)[0].Contains(3));
-		R.ExpectTrue((Proximity)[0].Contains(4));
-		R.ExpectTrue((Proximity)[0].Contains(1));
-		R.ExpectTrue(!(Proximity)[0].Contains(0));
-		R.ExpectTrue(!(Proximity)[0].Contains(2));
-		R.ExpectTrue(!(Proximity)[0].Contains(5));
-
-		R.ExpectTrue((Proximity)[1].Contains(0));
-		R.ExpectTrue((Proximity)[1].Contains(4));
-		R.ExpectTrue((Proximity)[1].Contains(5));
-		R.ExpectTrue((Proximity)[1].Contains(2));
-		R.ExpectTrue(!(Proximity)[1].Contains(1));
-		R.ExpectTrue(!(Proximity)[1].Contains(3));
-
-		R.ExpectTrue((Proximity)[2].Contains(1));
-		R.ExpectTrue((Proximity)[2].Contains(5));
-		R.ExpectTrue(!(Proximity)[2].Contains(0));
-		R.ExpectTrue(!(Proximity)[2].Contains(2));
-		R.ExpectTrue(!(Proximity)[2].Contains(3));
-		R.ExpectTrue(!(Proximity)[2].Contains(4));
-
-		R.ExpectTrue((Proximity)[3].Contains(0));
-		R.ExpectTrue((Proximity)[3].Contains(4));
-		R.ExpectTrue(!(Proximity)[3].Contains(1));
-		R.ExpectTrue(!(Proximity)[3].Contains(2));
-		R.ExpectTrue(!(Proximity)[3].Contains(3));
-		R.ExpectTrue(!(Proximity)[3].Contains(5));
-
-		R.ExpectTrue((Proximity)[4].Contains(0));
-		R.ExpectTrue((Proximity)[4].Contains(1));
-		R.ExpectTrue((Proximity)[4].Contains(3));
-		R.ExpectTrue((Proximity)[4].Contains(5));
-		R.ExpectTrue(!(Proximity)[4].Contains(2));
-		R.ExpectTrue(!(Proximity)[4].Contains(4));
-
-		R.ExpectTrue((Proximity)[5].Contains(1));
-		R.ExpectTrue((Proximity)[5].Contains(2));
-		R.ExpectTrue((Proximity)[5].Contains(4));
-		R.ExpectTrue(!(Proximity)[5].Contains(0));
-		R.ExpectTrue(!(Proximity)[5].Contains(3));
-		R.ExpectTrue(!(Proximity)[5].Contains(5));
+		EXPECT_TRUE((Proximity)[0].Contains(3));
+		EXPECT_TRUE((Proximity)[0].Contains(4));
+		EXPECT_TRUE((Proximity)[0].Contains(1));
+		EXPECT_FALSE((Proximity)[0].Contains(0));
+		EXPECT_FALSE((Proximity)[0].Contains(2));
+		EXPECT_FALSE((Proximity)[0].Contains(5));
+
+		EXPECT_TRUE((Proximity)[1].Contains(0));
+		EXPECT_TRUE((Proximity)[1].Contains(4));
+		EXPECT_TRUE((Proximity)[1].Contains(5));
+		EXPECT_TRUE((Proximity)[1].Contains(2));
+		EXPECT_FALSE((Proximity)[1].Contains(1));
+		EXPECT_FALSE((Proximity)[1].Contains(3));
+
+		EXPECT_TRUE((Proximity)[2].Contains(1));
+		EXPECT_TRUE((Proximity)[2].Contains(5));
+		EXPECT_FALSE((Proximity)[2].Contains(0));
+		EXPECT_FALSE((Proximity)[2].Contains(2));
+		EXPECT_FALSE((Proximity)[2].Contains(3));
+		EXPECT_FALSE((Proximity)[2].Contains(4));
+
+		EXPECT_TRUE((Proximity)[3].Contains(0));
+		EXPECT_TRUE((Proximity)[3].Contains(4));
+		EXPECT_FALSE((Proximity)[3].Contains(1));
+		EXPECT_FALSE((Proximity)[3].Contains(2));
+		EXPECT_FALSE((Proximity)[3].Contains(3));
+		EXPECT_FALSE((Proximity)[3].Contains(5));
+
+		EXPECT_TRUE((Proximity)[4].Contains(0));
+		EXPECT_TRUE((Proximity)[4].Contains(1));
+		EXPECT_TRUE((Proximity)[4].Contains(3));
+		EXPECT_TRUE((Proximity)[4].Contains(5));
+		EXPECT_FALSE((Proximity)[4].Contains(2));
+		EXPECT_FALSE((Proximity)[4].Contains(4));
+
+		EXPECT_TRUE((Proximity)[5].Contains(1));
+		EXPECT_TRUE((Proximity)[5].Contains(2));
+		EXPECT_TRUE((Proximity)[5].Contains(4));
+		EXPECT_FALSE((Proximity)[5].Contains(0));
+		EXPECT_FALSE((Proximity)[5].Contains(3));
+		EXPECT_FALSE((Proximity)[5].Contains(5));
 
 		TArray<int32> DelList = { 0,1,2,3,4,5 };
 		Coll->RemoveElements(FGeometryCollection::GeometryGroup, DelList);
 
 		// Proximity = []
 
-		R.ExpectTrue(Collection->NumElements(FGeometryCollection::GeometryGroup) == 0);
-
-		return !R.HasError();
+		EXPECT_EQ(Collection->NumElements(FGeometryCollection::GeometryGroup), 0);
 	}
-	template bool GeometryDeleteAll<float>(ExampleResponse&& R);
+	template void GeometryDeleteAll<float>();
 
 	template<class T>
-	bool GeometrySwapFlat(ExampleResponse&& R)
+	void GeometrySwapFlat()
 	{
 		FGeometryCollection Coll;
 		Coll.AppendGeometry(*GeometryCollection::MakeCubeElement(FTransform(FVector(0,0,0))));
@@ -902,7 +886,7 @@ namespace GeometryCollectionExample
 		{
 			for (int i = 0; i < 8; ++i)
 			{
-				R.ExpectTrue(Coll.BoneMap[i + box * 8] == box);
+				EXPECT_EQ(Coll.BoneMap[i + box * 8], box);
 			}
 		}
 
@@ -928,7 +912,7 @@ namespace GeometryCollectionExample
 			{
 				for (int idx = 0; idx < 3; ++idx)
 				{
-					R.ExpectTrue(Coll.Indices[i + box * 12][idx] == (ExpectedIndices[i * 3 + idx] + box * 8));
+					EXPECT_EQ(Coll.Indices[i + box * 12][idx], (ExpectedIndices[i * 3 + idx] + box * 8));
 				}
 			}
 		}
@@ -936,23 +920,23 @@ namespace GeometryCollectionExample
 		TArray<int32> NewOrder = { 0,3,2,1 };
 		Coll.ReorderElements(FGeometryCollection::TransformGroup, NewOrder);
 		//transforms change
-		R.ExpectTrue(Coll.Transform[0].GetLocation().X == 0.f);
-		R.ExpectTrue(Coll.Transform[1].GetLocation().X == 3.f);
-		R.ExpectTrue(Coll.Transform[2].GetLocation().X == 2.f);
-		R.ExpectTrue(Coll.Transform[3].GetLocation().X == 1.f);
+		EXPECT_EQ(Coll.Transform[0].GetLocation().X, 0.f);
+		EXPECT_EQ(Coll.Transform[1].GetLocation().X, 3.f);
+		EXPECT_EQ(Coll.Transform[2].GetLocation().X, 2.f);
+		EXPECT_EQ(Coll.Transform[3].GetLocation().X, 1.f);
 
 		//groups swap to be contiguous with transform array
-		R.ExpectTrue(Coll.TransformIndex[0] == 0);
-		R.ExpectTrue(Coll.TransformIndex[1] == 1);
-		R.ExpectTrue(Coll.TransformIndex[2] == 2);
-		R.ExpectTrue(Coll.TransformIndex[3] == 3);
+		EXPECT_EQ(Coll.TransformIndex[0], 0);
+		EXPECT_EQ(Coll.TransformIndex[1], 1);
+		EXPECT_EQ(Coll.TransformIndex[2], 2);
+		EXPECT_EQ(Coll.TransformIndex[3], 3);
 
 		//verts are still contiguous
 		for (int box = 0; box < 4; ++box)
 		{
 			for (int i = 0; i < 8; ++i)
 			{
-				R.ExpectTrue(Coll.BoneMap[i + box * 8] == box);
+				EXPECT_EQ(Coll.BoneMap[i + box * 8], box);
 			}
 		}
 
@@ -964,18 +948,16 @@ namespace GeometryCollectionExample
 				for (int idx = 0; idx < 3; ++idx)
 				{
 					//expect verts to reorder with faces so there's no indirection needed. The whole point is that we are contiguous
-					R.ExpectTrue(Coll.Indices[i + box * 12][idx] == (ExpectedIndices[i * 3 + idx] + box * 8));
+					EXPECT_EQ(Coll.Indices[i + box * 12][idx], (ExpectedIndices[i * 3 + idx] + box * 8));
 				}
 			}
 		}
-
-		return !R.HasError();
 	}
-	template bool GeometrySwapFlat<float>(ExampleResponse&& R);
+	template void GeometrySwapFlat<float>();
 
 	
 	template<class T>
-	bool TestFracturedGeometry(ExampleResponse&& R)
+	void TestFracturedGeometry()
 	{
 		FGeometryCollection* TestCollection = FGeometryCollection::NewGeometryCollection(FracturedGeometry::RawVertexArray,
 																						 FracturedGeometry::RawIndicesArray,
@@ -987,11 +969,9 @@ namespace GeometryCollectionExample
 																						 FracturedGeometry::RawSimulationTypeArray,
 																						 FracturedGeometry::RawStatusFlagsArray);
 
-		R.ExpectTrue(TestCollection->NumElements(FGeometryCollection::GeometryGroup) == 11);
-
-		return !R.HasError();
+		EXPECT_EQ(TestCollection->NumElements(FGeometryCollection::GeometryGroup), 11);
 	}
-	template bool TestFracturedGeometry<float>(ExampleResponse&& R);
+	template void TestFracturedGeometry<float>();
 
 
 }
diff --git a/Engine/Plugins/Experimental/GeometryCollectionPlugin/Source/GeometryCollectionExampleCore/Private/GeometryCollection/GeometryCollectionExampleSimulation.cpp b/Engine/Plugins/Experimental/GeometryCollectionPlugin/Source/GeometryCollectionExampleCore/Private/GeometryCollection/GeometryCollectionExampleSimulation.cpp
--- a/Engine/Plugins/Experimental/GeometryCollectionPlugin/Source/GeometryCollectionExampleCore/Private/GeometryCollection/GeometryCollectionExampleSimulation.cpp
+++ b/Engine/Plugins/Experimental/GeometryCollectionPlugin/Source/GeometryCollectionExampleCore/Private/GeometryCollection/GeometryCollectionExampleSimulation.cpp
@@ -25,35 +25,22 @@
 
 namespace GeometryCollectionExample
 {
-
+	
 	template<class T>
-	bool RigidBodiesFallingUnderGravity(ExampleResponse&& R)
+	void RigidBodiesFallingUnderGravity()
 	{
 #if INCLUDE_CHAOS
-		TUniquePtr<Chaos::TChaosPhysicsMaterial<T>> PhysicalMaterial = MakeUnique<Chaos::TChaosPhysicsMaterial<T>>();
-		PhysicalMaterial->Friction = 0;
-		PhysicalMaterial->Restitution = 0;
-		PhysicalMaterial->SleepingLinearThreshold = 0;
-		PhysicalMaterial->SleepingAngularThreshold = 0;
-		PhysicalMaterial->DisabledLinearThreshold = 0;
-		PhysicalMaterial->DisabledAngularThreshold = 0;
-
-		TSharedPtr<FGeometryCollection> RestCollection = GeometryCollection::MakeCubeElement(FTransform::Identity, FVector(1.0));
-		TSharedPtr<FGeometryDynamicCollection> DynamicCollection = GeometryCollectionToGeometryDynamicCollection(RestCollection.Get());
-		
-		auto InitFunc = [&RestCollection, &DynamicCollection, &PhysicalMaterial](FSimulationParameters& InParams)
-		{
-			InParams.RestCollection = RestCollection.Get();
-			InParams.DynamicCollection = DynamicCollection.Get();
-			InParams.PhysicalMaterial  = MakeSerializable(PhysicalMaterial);
-			InParams.Shared.SizeSpecificData[0].CollisionType = ECollisionTypeEnum::Chaos_Volumetric;
-			InParams.Simulating = true;
-			Chaos::FErrorReporter ErrorReporter;
-			BuildSimulationData(ErrorReporter, *RestCollection, InParams.Shared);
-		};
+		TUniquePtr<Chaos::TChaosPhysicsMaterial<T>> PhysicalMaterial = nullptr;
+		TSharedPtr<FGeometryCollection> RestCollection = nullptr;
+		TSharedPtr<FGeometryDynamicCollection> DynamicCollection = nullptr;
+
+		//
+		//  Rigid Body Setup
+		//
+		InitCollectionsParameters InitParams = { FTransform::Identity, FVector(1.0), nullptr, (int32)EObjectStateTypeEnum::Chaos_Object_Kinematic };
+		InitCollections(PhysicalMaterial, RestCollection, DynamicCollection, InitParams);
 
-		FGeometryCollectionPhysicsProxy* PhysObject = new FGeometryCollectionPhysicsProxy(nullptr, DynamicCollection.Get(), InitFunc, nullptr, nullptr);;
-		PhysObject->Initialize();
+		FGeometryCollectionPhysicsProxy* PhysObject = RigidBodySetup(PhysicalMaterial, RestCollection, DynamicCollection);
 
 		Chaos::FPBDRigidsSolver* Solver = FChaosSolversModule::GetModule()->CreateSolver(true);
 #if CHAOS_PARTICLEHANDLE_TODO
@@ -69,55 +56,42 @@ namespace GeometryCollectionExample
 
 		// never touched
 		TManagedArray<FTransform>& RestTransform = RestCollection->Transform;
-		R.ExpectTrue(FMath::Abs(RestTransform[0].GetTranslation().Z) < SMALL_THRESHOLD);
+		EXPECT_LT(FMath::Abs(RestTransform[0].GetTranslation().Z), SMALL_THRESHOLD);
 
 		// simulated
 		TManagedArray<FTransform>& Transform = DynamicCollection->Transform;
-		R.ExpectTrue(Transform.Num() == 1);
-		R.ExpectTrue(Transform[0].GetTranslation().Z < 0);
+		EXPECT_EQ(Transform.Num(), 1);
+		EXPECT_LT(Transform[0].GetTranslation().Z, 0);
 		
 		FChaosSolversModule::GetModule()->DestroySolver(Solver);
 
 		delete PhysObject;
-
 #endif
-
-
-		return !R.HasError();
 	}
-	template bool RigidBodiesFallingUnderGravity<float>(ExampleResponse&& R);
+	template void RigidBodiesFallingUnderGravity<float>();
 
 
 
 	template<class T>
-	bool RigidBodiesCollidingWithSolverFloor(ExampleResponse&& R)
+	void RigidBodiesCollidingWithSolverFloor()
 	{
 #if INCLUDE_CHAOS
-		TUniquePtr<Chaos::TChaosPhysicsMaterial<T>> PhysicalMaterial = MakeUnique<Chaos::TChaosPhysicsMaterial<T>>();
-		PhysicalMaterial->Friction = 0;
-		PhysicalMaterial->Restitution = 0;
-		PhysicalMaterial->SleepingLinearThreshold = 0;
-		PhysicalMaterial->SleepingAngularThreshold = 0;
-		PhysicalMaterial->DisabledLinearThreshold = 0;
-		PhysicalMaterial->DisabledAngularThreshold = 0;
-
-		TSharedPtr<FGeometryCollection> RestCollection = GeometryCollection::MakeCubeElement(FTransform::Identity, FVector(1.0));
-		TSharedPtr<FGeometryDynamicCollection> DynamicCollection = GeometryCollectionToGeometryDynamicCollection(RestCollection.Get());
+		TUniquePtr<Chaos::TChaosPhysicsMaterial<T>> PhysicalMaterial = nullptr;
+		TSharedPtr<FGeometryCollection> RestCollection = nullptr;
+		TSharedPtr<FGeometryDynamicCollection> DynamicCollection = nullptr;
+
+		//
+		//  Rigid Body Setup
+		//
+		InitCollectionsParameters InitParams = { FTransform::Identity, FVector(1.0), nullptr, (int32)EObjectStateTypeEnum::Chaos_Object_Kinematic };
+		InitCollections(PhysicalMaterial, RestCollection, DynamicCollection, InitParams);
 				
-		auto InitFunc = [&RestCollection, &DynamicCollection, &PhysicalMaterial](FSimulationParameters& InParams)
+		auto CustomFunc = [&RestCollection, &DynamicCollection, &PhysicalMaterial](FSimulationParameters& InParams)
 		{
-			InParams.RestCollection = RestCollection.Get();
-			InParams.DynamicCollection = DynamicCollection.Get();
-			InParams.PhysicalMaterial  = MakeSerializable(PhysicalMaterial);
-			InParams.Shared.SizeSpecificData[0].CollisionType = ECollisionTypeEnum::Chaos_Volumetric;
 			InParams.Shared.SizeSpecificData[0].ImplicitType = EImplicitTypeEnum::Chaos_Implicit_Box;
-			InParams.Simulating = true;
-			Chaos::FErrorReporter ErrorReporter;
-			BuildSimulationData(ErrorReporter, *RestCollection, InParams.Shared);
 		};
 
-		FGeometryCollectionPhysicsProxy* PhysObject = new FGeometryCollectionPhysicsProxy(nullptr, DynamicCollection.Get(), InitFunc, nullptr, nullptr);;
-		PhysObject->Initialize();
+		FGeometryCollectionPhysicsProxy* PhysObject = RigidBodySetup(PhysicalMaterial, RestCollection, DynamicCollection, CustomFunc);
 
 		Chaos::FPBDRigidsSolver* Solver = FChaosSolversModule::GetModule()->CreateSolver(true);
 #if CHAOS_PARTICLEHANDLE_TODO
@@ -134,56 +108,47 @@ namespace GeometryCollectionExample
 
 		// never touched
 		TManagedArray<FTransform>& RestTransform = RestCollection->Transform;
-		R.ExpectTrue(FMath::Abs(RestTransform[0].GetTranslation().Z) < SMALL_THRESHOLD);
+		EXPECT_LT(FMath::Abs(RestTransform[0].GetTranslation().Z), SMALL_THRESHOLD);
 
 		// simulated
 		TManagedArray<FTransform>& Transform = DynamicCollection->Transform;
-		R.ExpectTrue(Transform.Num() == 1);
-		R.ExpectTrue(FMath::Abs(Transform[0].GetTranslation().Z - 0.5) < SMALL_THRESHOLD);
+		EXPECT_EQ(Transform.Num(), 1);
+		EXPECT_LT(FMath::Abs(Transform[0].GetTranslation().Z - 0.5), SMALL_THRESHOLD);
 		
 
 		FChaosSolversModule::GetModule()->DestroySolver(Solver);
 
 		delete PhysObject;
-
 #endif
-
-
-		return !R.HasError();
 	}
-	template bool RigidBodiesCollidingWithSolverFloor<float>(ExampleResponse&& R);
+	template void RigidBodiesCollidingWithSolverFloor<float>();
 
 
 	template<class T>
-	bool RigidBodiesSingleSphereCollidingWithSolverFloor(ExampleResponse&& R)
+	void RigidBodiesSingleSphereCollidingWithSolverFloor()
 	{
 #if INCLUDE_CHAOS
-		TUniquePtr<Chaos::TChaosPhysicsMaterial<T>> PhysicalMaterial = MakeUnique<Chaos::TChaosPhysicsMaterial<T>>();
-		PhysicalMaterial->Friction = 0;
-		PhysicalMaterial->Restitution = 0;
-		PhysicalMaterial->SleepingLinearThreshold = 0;
-		PhysicalMaterial->SleepingAngularThreshold = 0;
-		PhysicalMaterial->DisabledLinearThreshold = 0;
-		PhysicalMaterial->DisabledAngularThreshold = 0;
-
-		TSharedPtr<FGeometryCollection> RestCollection = GeometryCollection::MakeCubeElement(FTransform::Identity, FVector(1.0));
-		(RestCollection->Transform)[0].SetTranslation(FVector(0, 0, 10));
-		TSharedPtr<FGeometryDynamicCollection> DynamicCollection = GeometryCollectionToGeometryDynamicCollection(RestCollection.Get());
+		TUniquePtr<Chaos::TChaosPhysicsMaterial<T>> PhysicalMaterial = nullptr;
+		TSharedPtr<FGeometryCollection> RestCollection = nullptr;
+		TSharedPtr<FGeometryDynamicCollection> DynamicCollection = nullptr;
+
+		//
+		//  Rigid Body Setup
+		//
+		auto RestInitFunc = [](TSharedPtr<FGeometryCollection>& RestCollection)
+		{
+			RestCollection->Transform[0].SetTranslation(FVector(0, 0, 10));
+		};
+
+		InitCollectionsParameters InitParams = { FTransform::Identity, FVector(1.0), RestInitFunc, (int32)EObjectStateTypeEnum::Chaos_Object_Kinematic };
+		InitCollections(PhysicalMaterial, RestCollection, DynamicCollection, InitParams);
 				
-		auto InitFunc = [&RestCollection, &DynamicCollection, &PhysicalMaterial](FSimulationParameters& InParams)
+		auto CustomFunc = [&RestCollection, &DynamicCollection, &PhysicalMaterial](FSimulationParameters& InParams)
 		{
-			InParams.RestCollection = RestCollection.Get();
-			InParams.DynamicCollection = DynamicCollection.Get();
-			InParams.PhysicalMaterial  = MakeSerializable(PhysicalMaterial);
-			InParams.Shared.SizeSpecificData[0].CollisionType = ECollisionTypeEnum::Chaos_Volumetric;
 			InParams.Shared.SizeSpecificData[0].ImplicitType = EImplicitTypeEnum::Chaos_Implicit_Sphere;
-			InParams.Simulating = true;
-			Chaos::FErrorReporter ErrorReporter;
-			BuildSimulationData(ErrorReporter, *RestCollection, InParams.Shared);
 		};
 
-		FGeometryCollectionPhysicsProxy* PhysObject = new FGeometryCollectionPhysicsProxy(nullptr, DynamicCollection.Get(), InitFunc, nullptr, nullptr);;
-		PhysObject->Initialize();
+		FGeometryCollectionPhysicsProxy* PhysObject = RigidBodySetup(PhysicalMaterial, RestCollection, DynamicCollection, CustomFunc);
 
 		Chaos::FPBDRigidsSolver* Solver = FChaosSolversModule::GetModule()->CreateSolver(true);
 #if CHAOS_PARTICLEHANDLE_TODO
@@ -203,57 +168,44 @@ namespace GeometryCollectionExample
 
 		// never touched
 		TManagedArray<FTransform>& RestTransform = RestCollection->Transform;
-		R.ExpectTrue(FMath::Abs(RestTransform[0].GetTranslation().Z-10.0) < KINDA_SMALL_NUMBER);
+		EXPECT_LT(FMath::Abs(RestTransform[0].GetTranslation().Z-10.0), KINDA_SMALL_NUMBER);
 
 		// simulated
 		TManagedArray<FTransform>& Transform = DynamicCollection->Transform;
 		TManagedArray<float>& InnerRadius = RestCollection->InnerRadius;
 		//UE_LOG(LogTest, Verbose, TEXT("Height : (%3.5f), Inner Radius:%3.5f"), Transform[0].GetTranslation().Z, InnerRadius[0]);
 
-		R.ExpectTrue(Transform.Num() == 1);
-		R.ExpectTrue(FMath::Abs(Transform[0].GetTranslation().Z - 0.5) < 0.1); // @todo - Why is this not 0.5f
+		EXPECT_EQ(Transform.Num(), 1);
+		EXPECT_LT(FMath::Abs(Transform[0].GetTranslation().Z - 0.5), 0.1); // @todo - Why is this not 0.5f
 		
 		FChaosSolversModule::GetModule()->DestroySolver(Solver);
 
 		delete PhysObject;
-
 #endif
-
-
-		return !R.HasError();
 	}
-	template bool RigidBodiesSingleSphereCollidingWithSolverFloor<float>(ExampleResponse&& R);
+	template void RigidBodiesSingleSphereCollidingWithSolverFloor<float>();
 
 
 	template<class T>
-	bool RigidBodiesSingleSphereIntersectingWithSolverFloor(ExampleResponse&& R)
+	void RigidBodiesSingleSphereIntersectingWithSolverFloor()
 	{
 #if INCLUDE_CHAOS
-		TUniquePtr<Chaos::TChaosPhysicsMaterial<T>> PhysicalMaterial = MakeUnique<Chaos::TChaosPhysicsMaterial<T>>();
-		PhysicalMaterial->Friction = 0;
-		PhysicalMaterial->Restitution = 0;
-		PhysicalMaterial->SleepingLinearThreshold = 0;
-		PhysicalMaterial->SleepingAngularThreshold = 0;
-		PhysicalMaterial->DisabledLinearThreshold = 0;
-		PhysicalMaterial->DisabledAngularThreshold = 0;
-
-		TSharedPtr<FGeometryCollection> RestCollection = GeometryCollection::MakeCubeElement(FTransform::Identity, FVector(1.0));
-		TSharedPtr<FGeometryDynamicCollection> DynamicCollection = GeometryCollectionToGeometryDynamicCollection(RestCollection.Get());
-
-		auto InitFunc = [&RestCollection, &DynamicCollection, &PhysicalMaterial](FSimulationParameters& InParams)
+		TUniquePtr<Chaos::TChaosPhysicsMaterial<T>> PhysicalMaterial = nullptr;
+		TSharedPtr<FGeometryCollection> RestCollection = nullptr;
+		TSharedPtr<FGeometryDynamicCollection> DynamicCollection = nullptr;
+
+		//
+		//  Rigid Body Setup
+		//
+		InitCollectionsParameters InitParams = { FTransform::Identity, FVector(1.0), nullptr, (int32)EObjectStateTypeEnum::Chaos_Object_Kinematic };
+		InitCollections(PhysicalMaterial, RestCollection, DynamicCollection, InitParams);
+
+		auto CustomFunc = [&RestCollection, &DynamicCollection, &PhysicalMaterial](FSimulationParameters& InParams)
 		{
-			InParams.RestCollection = RestCollection.Get();
-			InParams.DynamicCollection = DynamicCollection.Get();
-			InParams.PhysicalMaterial  = MakeSerializable(PhysicalMaterial);
-			InParams.Shared.SizeSpecificData[0].CollisionType = ECollisionTypeEnum::Chaos_Volumetric;
 			InParams.Shared.SizeSpecificData[0].ImplicitType = EImplicitTypeEnum::Chaos_Implicit_Sphere;
-			InParams.Simulating = true;
-			Chaos::FErrorReporter ErrorReporter;
-			BuildSimulationData(ErrorReporter, *RestCollection, InParams.Shared);
 		};
 
-		FGeometryCollectionPhysicsProxy* PhysObject = new FGeometryCollectionPhysicsProxy(nullptr, DynamicCollection.Get(), InitFunc, nullptr, nullptr);;
-		PhysObject->Initialize();
+		FGeometryCollectionPhysicsProxy* PhysObject = RigidBodySetup(PhysicalMaterial, RestCollection, DynamicCollection, CustomFunc);
 
 		Chaos::FPBDRigidsSolver* Solver = FChaosSolversModule::GetModule()->CreateSolver(true);
 #if CHAOS_PARTICLEHANDLE_TODO
@@ -270,14 +222,14 @@ namespace GeometryCollectionExample
 
 		// never touched
 		TManagedArray<FTransform>& RestTransform = RestCollection->Transform;
-		R.ExpectTrue(FMath::Abs(RestTransform[0].GetTranslation().Z) < KINDA_SMALL_NUMBER);
+		EXPECT_LT(FMath::Abs(RestTransform[0].GetTranslation().Z), KINDA_SMALL_NUMBER);
 
 		// simulated
 		TManagedArray<FTransform>& Transform = DynamicCollection->Transform;
-		R.ExpectTrue(Transform.Num() == 1);
+		EXPECT_EQ(Transform.Num(), 1);
 		//UE_LOG(LogTest, Verbose, TEXT("Position : (%3.5f,%3.5f,%3.5f)"), Transform[0].GetTranslation().X, Transform[0].GetTranslation().Y, Transform[0].GetTranslation().Z);
 
-		R.ExpectTrue(FMath::Abs(Transform[0].GetTranslation().Z - 0.5) < KINDA_SMALL_NUMBER);
+		EXPECT_LT(FMath::Abs(Transform[0].GetTranslation().Z - 0.5), KINDA_SMALL_NUMBER);
 		
 
 		FChaosSolversModule::GetModule()->DestroySolver(Solver);
@@ -285,40 +237,25 @@ namespace GeometryCollectionExample
 		delete PhysObject;
 
 #endif
-
-		return !R.HasError();
 	}
-	template bool RigidBodiesSingleSphereIntersectingWithSolverFloor<float>(ExampleResponse&& R);
+	template void RigidBodiesSingleSphereIntersectingWithSolverFloor<float>();
 
 
 	template<class T>
-	bool RigidBodiesKinematic(ExampleResponse&& R)
+	void RigidBodiesKinematic()
 	{
 #if INCLUDE_CHAOS
-		TUniquePtr<Chaos::TChaosPhysicsMaterial<T>> PhysicalMaterial = MakeUnique<Chaos::TChaosPhysicsMaterial<T>>();
-		PhysicalMaterial->Friction = 0;
-		PhysicalMaterial->Restitution = 0;
-		PhysicalMaterial->SleepingLinearThreshold = 0;
-		PhysicalMaterial->SleepingAngularThreshold = 0;
-		PhysicalMaterial->DisabledLinearThreshold = 0;
-		PhysicalMaterial->DisabledAngularThreshold = 0;
-
-		TSharedPtr<FGeometryCollection> RestCollection = GeometryCollection::MakeCubeElement(FTransform::Identity, FVector(1.0));
-		TSharedPtr<FGeometryDynamicCollection> DynamicCollection = GeometryCollectionToGeometryDynamicCollection(RestCollection.Get(), (uint8)EObjectStateTypeEnum::Chaos_Object_Kinematic);
-
-		auto InitFunc = [&RestCollection, &DynamicCollection, &PhysicalMaterial](FSimulationParameters& InParams)
-		{
-			InParams.RestCollection = RestCollection.Get();
-			InParams.DynamicCollection = DynamicCollection.Get();
-			InParams.PhysicalMaterial  = MakeSerializable(PhysicalMaterial);
-			InParams.Shared.SizeSpecificData[0].CollisionType = ECollisionTypeEnum::Chaos_Volumetric;
-			InParams.Simulating = true;
-			Chaos::FErrorReporter ErrorReporter;
-			BuildSimulationData(ErrorReporter, *RestCollection, InParams.Shared);
-		};
+		TUniquePtr<Chaos::TChaosPhysicsMaterial<T>> PhysicalMaterial = nullptr;
+		TSharedPtr<FGeometryCollection> RestCollection = nullptr;
+		TSharedPtr<FGeometryDynamicCollection> DynamicCollection = nullptr;
 
-		FGeometryCollectionPhysicsProxy* PhysObject = new FGeometryCollectionPhysicsProxy(nullptr, DynamicCollection.Get(), InitFunc, nullptr, nullptr);;
-		PhysObject->Initialize();
+		//
+		//  Rigid Body Setup
+		//
+		InitCollectionsParameters InitParams = { FTransform::Identity, FVector(1.0), nullptr, (uint8)EObjectStateTypeEnum::Chaos_Object_Kinematic };
+		InitCollections(PhysicalMaterial, RestCollection, DynamicCollection, InitParams);
+
+		FGeometryCollectionPhysicsProxy* PhysObject = RigidBodySetup(PhysicalMaterial, RestCollection, DynamicCollection);
 
 		Chaos::FPBDRigidsSolver* Solver = FChaosSolversModule::GetModule()->CreateSolver(true);
 #if CHAOS_PARTICLEHANDLE_TODO
@@ -338,59 +275,46 @@ namespace GeometryCollectionExample
 
 		// simulated
 		TManagedArray<FTransform>& Transform = DynamicCollection->Transform;
-		R.ExpectTrue(Transform.Num() == 1);
+		EXPECT_EQ(Transform.Num(), 1);
 		//UE_LOG(LogTest, Verbose, TEXT("Position : (%3.5f,%3.5f,%3.5f)"), Transform[0].GetTranslation().X, Transform[0].GetTranslation().Y, Transform[0].GetTranslation().Z);
-		R.ExpectTrue(Transform[0].GetTranslation().Z == 0.f);
+		EXPECT_EQ(Transform[0].GetTranslation().Z, 0.f);
 		
 		FChaosSolversModule::GetModule()->DestroySolver(Solver);
 
 		delete PhysObject;
 
 #endif
-
-		return !R.HasError();
 	}
-	template bool RigidBodiesKinematic<float>(ExampleResponse&& R);
+	template void RigidBodiesKinematic<float>();
 
 
 	template<class T>
-	bool RigidBodiesSleepingActivation(ExampleResponse&& R)
+	void RigidBodiesSleepingActivation()
 	{
 #if INCLUDE_CHAOS
-		TUniquePtr<Chaos::TChaosPhysicsMaterial<T>> PhysicalMaterial = MakeUnique<Chaos::TChaosPhysicsMaterial<T>>();
-		PhysicalMaterial->Friction = 0;
-		PhysicalMaterial->Restitution = 0;
-		PhysicalMaterial->SleepingLinearThreshold = 0;
-		PhysicalMaterial->SleepingAngularThreshold = 0;
-		PhysicalMaterial->DisabledLinearThreshold = 0;
-		PhysicalMaterial->DisabledAngularThreshold = 0;
-
-		TSharedPtr<FGeometryCollection> RestCollection = GeometryCollection::MakeCubeElement(FTransform::Identity, FVector(1.0));
+		TUniquePtr<Chaos::TChaosPhysicsMaterial<T>> PhysicalMaterial = nullptr;
+		TSharedPtr<FGeometryCollection> RestCollection = nullptr;
+		TSharedPtr<FGeometryDynamicCollection> DynamicCollection = nullptr;
 
-		RestCollection->AppendGeometry(*RestCollection);
-		TManagedArray<FTransform>& RestTransform = RestCollection->Transform;
-		RestTransform[1].SetTranslation(FVector(0.f, 0.f, 5.f));
+		//
+		//  Rigid Body Setup
+		//
+		auto RestInitFunc = [](TSharedPtr<FGeometryCollection>& RestCollection)
+		{
+			RestCollection->AppendGeometry(*RestCollection);
+			RestCollection->Transform[1].SetTranslation(FVector(0.f, 0.f, 5.f));
+		};
 
-		TSharedPtr<FGeometryDynamicCollection> DynamicCollection = GeometryCollectionToGeometryDynamicCollection(RestCollection.Get(), (int32)EObjectStateTypeEnum::Chaos_Object_Kinematic);
+		InitCollectionsParameters InitParams = { FTransform::Identity, FVector(1.0), RestInitFunc, (int32)EObjectStateTypeEnum::Chaos_Object_Kinematic };
+		InitCollections(PhysicalMaterial, RestCollection, DynamicCollection, InitParams);
 
-		auto InitFunc = [&RestCollection, &DynamicCollection, &PhysicalMaterial](FSimulationParameters& InParams)
+		auto CustomFunc = [&RestCollection, &DynamicCollection, &PhysicalMaterial](FSimulationParameters& InParams)
 		{
-			InParams.RestCollection = RestCollection.Get();
-			InParams.DynamicCollection = DynamicCollection.Get();
-			InParams.PhysicalMaterial  = MakeSerializable(PhysicalMaterial);
-			InParams.Shared.SizeSpecificData[0].CollisionType = ECollisionTypeEnum::Chaos_Volumetric;
 			//InParams.Shared.SizeSpecificData[0].ImplicitType = EImplicitTypeEnum::Chaos_Implicit_Cube;
-			InParams.Simulating = true;
-			Chaos::FErrorReporter ErrorReporter;
-			BuildSimulationData(ErrorReporter, *RestCollection, InParams.Shared);
 		};
 
-		FGeometryCollectionPhysicsProxy* PhysObject = new FGeometryCollectionPhysicsProxy(nullptr, DynamicCollection.Get(), InitFunc, nullptr, nullptr);;
-		PhysObject->Initialize();
+		FGeometryCollectionPhysicsProxy* PhysObject = RigidBodySetup(PhysicalMaterial, RestCollection, DynamicCollection, CustomFunc);
 
-		//
-		//
-		//
 		TManagedArray<int32>& ObjectType = DynamicCollection->GetAttribute<int32>("DynamicState", FTransformCollection::TransformGroup);
 		ObjectType[0] = (int32)EObjectStateTypeEnum::Chaos_Object_Sleeping;
 		ObjectType[1] = (int32)EObjectStateTypeEnum::Chaos_Object_Dynamic;
@@ -419,50 +343,40 @@ namespace GeometryCollectionExample
 		FChaosSolversModule::GetModule()->DestroySolver(Solver);
 
 		delete PhysObject;
-
 #endif
-		return !R.HasError();
 	}
-	template bool RigidBodiesSleepingActivation<float>(ExampleResponse&& R);
+	template void RigidBodiesSleepingActivation<float>();
 
 	template<class T>
-	bool RigidBodies_CollisionGroup(ExampleResponse&& R)
+	void RigidBodies_CollisionGroup()
 	{
 #if INCLUDE_CHAOS
-		TUniquePtr<Chaos::TChaosPhysicsMaterial<T>> PhysicalMaterial = MakeUnique<Chaos::TChaosPhysicsMaterial<T>>();
-		PhysicalMaterial->Friction = 0;
-		PhysicalMaterial->Restitution = 0;
-		PhysicalMaterial->SleepingLinearThreshold = 0;
-		PhysicalMaterial->SleepingAngularThreshold = 0;
-		PhysicalMaterial->DisabledLinearThreshold = 0;
-		PhysicalMaterial->DisabledAngularThreshold = 0;
+		TUniquePtr<Chaos::TChaosPhysicsMaterial<T>> PhysicalMaterial = nullptr;
+		TSharedPtr<FGeometryCollection> RestCollection = nullptr;
+		TSharedPtr<FGeometryDynamicCollection> DynamicCollection = nullptr;
 
 		//
-		// Generate Geometry
+		//  Rigid Body Setup
 		//
-		TSharedPtr<FGeometryCollection> RestCollection = GeometryCollection::MakeCubeElement(FTransform(FVector(0, 0, 100.0)), FVector(100.0));
-		RestCollection->AppendGeometry(*GeometryCollection::MakeCubeElement(FTransform(FVector(0, 0, 210.0)), FVector(100.0)));
-		RestCollection->AppendGeometry(*GeometryCollection::MakeCubeElement(FTransform(FVector(0, 0, 320.0)), FVector(100.0)));
-		RestCollection->AppendGeometry(*GeometryCollection::MakeCubeElement(FTransform(FVector(0, 0, 430.0)), FVector(100.0)));
-		TSharedPtr<FGeometryDynamicCollection> DynamicCollection = GeometryCollectionToGeometryDynamicCollection(RestCollection.Get());
+		auto RestInitFunc = [](TSharedPtr<FGeometryCollection>& RestCollection)
+		{
+			RestCollection->AppendGeometry(*GeometryCollection::MakeCubeElement(FTransform(FVector(0, 0, 210.0)), FVector(100.0)));
+			RestCollection->AppendGeometry(*GeometryCollection::MakeCubeElement(FTransform(FVector(0, 0, 320.0)), FVector(100.0)));
+			RestCollection->AppendGeometry(*GeometryCollection::MakeCubeElement(FTransform(FVector(0, 0, 430.0)), FVector(100.0)));
+		};
+
+		InitCollectionsParameters InitParams = { FTransform(FVector(0, 0, 100.0)), FVector(100.0), RestInitFunc, (int32)EObjectStateTypeEnum::Chaos_Object_Kinematic };
+		InitCollections(PhysicalMaterial, RestCollection, DynamicCollection, InitParams);
 
 		//
 		// Solver setup
 		//
-		auto InitFunc = [&RestCollection, &DynamicCollection, &PhysicalMaterial](FSimulationParameters& InParams)
+		auto CustomFunc = [&RestCollection, &DynamicCollection, &PhysicalMaterial](FSimulationParameters& InParams)
 		{
-			InParams.RestCollection = RestCollection.Get();
-			InParams.DynamicCollection = DynamicCollection.Get();
-			InParams.PhysicalMaterial  = MakeSerializable(PhysicalMaterial);
-			InParams.Shared.SizeSpecificData[0].CollisionType = ECollisionTypeEnum::Chaos_Volumetric;
 			InParams.Shared.SizeSpecificData[0].ImplicitType = EImplicitTypeEnum::Chaos_Implicit_Box;
-			InParams.Simulating = true;
-			Chaos::FErrorReporter ErrorReporter;
-			BuildSimulationData(ErrorReporter, *RestCollection, InParams.Shared);
 		};
 
-		FGeometryCollectionPhysicsProxy* PhysObject = new FGeometryCollectionPhysicsProxy(nullptr, DynamicCollection.Get(), InitFunc, nullptr, nullptr);;
-		PhysObject->Initialize();
+		FGeometryCollectionPhysicsProxy* PhysObject = RigidBodySetup(PhysicalMaterial, RestCollection, DynamicCollection, CustomFunc);
 		PhysObject->SetCollisionParticlesPerObjectFraction( 1.0 );
 
 		Chaos::FPBDRigidsSolver* Solver = FChaosSolversModule::GetModule()->CreateSolver(true);
@@ -492,17 +406,17 @@ namespace GeometryCollectionExample
 			}
 			if (Frame == 13)
 			{
-				R.ExpectTrue(FMath::Abs(Particles.X(0).Z) < SMALL_NUMBER);
-				R.ExpectTrue(FMath::Abs(Particles.X(1).Z - 50.f) < 10.f);
-				R.ExpectTrue(FMath::Abs(Particles.X(2).Z - 150.f) < 10.f);
+				EXPECT_LT(FMath::Abs(Particles.X(0).Z), SMALL_NUMBER);
+				EXPECT_LT(FMath::Abs(Particles.X(1).Z - 50.f), 10.f);
+				EXPECT_LT(FMath::Abs(Particles.X(2).Z - 150.f), 10.f);
 			}
 			if( Frame == 30 )
 			{
-				R.ExpectTrue(FMath::Abs(Particles.X(0).Z) < SMALL_NUMBER);
-				R.ExpectTrue(FMath::Abs(Particles.X(1).Z - 50.f) < 10.f);
-				R.ExpectTrue(FMath::Abs(Particles.X(2).Z - 150.f) < 10.f);
-				R.ExpectTrue(Particles.X(3).Z > 50.f);
-				R.ExpectTrue(Particles.X(4).Z < -100);
+				EXPECT_LT(FMath::Abs(Particles.X(0).Z), SMALL_NUMBER);
+				EXPECT_LT(FMath::Abs(Particles.X(1).Z - 50.f), 10.f);
+				EXPECT_LT(FMath::Abs(Particles.X(2).Z - 150.f), 10.f);
+				EXPECT_GT(Particles.X(3).Z, 50.f);
+				EXPECT_LT(Particles.X(4).Z, -100);
 			}
 			if (Frame == 31)
 			{
@@ -514,25 +428,23 @@ namespace GeometryCollectionExample
 			}
 		}
 
-		R.ExpectTrue(FMath::Abs(Particles.X(0).Z) < SMALL_NUMBER);
-		R.ExpectTrue(Particles.X(1).Z < -10000);
-		R.ExpectTrue(Particles.X(2).Z > 50.0);
-		R.ExpectTrue(Particles.X(3).Z < -10000);
-		R.ExpectTrue(Particles.X(4).Z < -10000);
+		EXPECT_LT(FMath::Abs(Particles.X(0).Z), SMALL_NUMBER);
+		EXPECT_LT(Particles.X(1).Z, -10000);
+		EXPECT_GT(Particles.X(2).Z, 50.0);
+		EXPECT_LT(Particles.X(3).Z, -10000);
+		EXPECT_LT(Particles.X(4).Z, -10000);
 #endif
 
 		FChaosSolversModule::GetModule()->DestroySolver(Solver);
 		delete PhysObject;
 #endif
-
-		return !R.HasError();
 	}
-	template bool RigidBodies_CollisionGroup<float>(ExampleResponse&& R);
+	template void RigidBodies_CollisionGroup<float>();
 
 
 
 	template<class T>
-	bool RigidBodies_Initialize_ParticleImplicitCollisionGeometry(ExampleResponse&& R)
+	void RigidBodies_Initialize_ParticleImplicitCollisionGeometry()
 	{
 #if INCLUDE_CHAOS
 		typedef Chaos::TVector<T, 3> Vec;
@@ -548,7 +460,7 @@ namespace GeometryCollectionExample
 		// check implicit domain
 		typedef TUniquePtr<Chaos::TImplicitObject<float, 3>> FImplicitPointer;
 		const TManagedArray<FImplicitPointer> & Implicits = Object->RestCollection->template GetAttribute<FImplicitPointer>("Implicits", FTransformCollection::TransformGroup);
-		R.ExpectTrue(Implicits.Num() == 1);
+		EXPECT_EQ(Implicits.Num(), 1);
 
 		const Chaos::TImplicitObject<float, 3> & Implicit = *Implicits[0];
 		for (float x = -1.05; x < 1.0; x += 0.1)
@@ -557,20 +469,20 @@ namespace GeometryCollectionExample
 			T phi = Implicit.PhiWithNormal(Vec(x, 0, 0), Normal);
 			if (x < -0.5 || 0.5 < x)
 			{
-				R.ExpectTrue(phi > 0);
+				EXPECT_GT(phi, 0);
 			}
 			else
 			{
-				R.ExpectTrue(phi < 0);
+				EXPECT_LT(phi, 0);
 			}
 
 			if (x < -0.25)
 			{
-				R.ExpectTrue((Normal-Vec(-1,0,0)).Size()<KINDA_SMALL_NUMBER);
+				EXPECT_LT((Normal-Vec(-1,0,0)).Size(), KINDA_SMALL_NUMBER);
 			}
 			else if (x > 0.25)
 			{
-				R.ExpectTrue((Normal - Vec(1, 0, 0)).Size() < KINDA_SMALL_NUMBER);
+				EXPECT_LT((Normal - Vec(1, 0, 0)).Size(), KINDA_SMALL_NUMBER);
 			}
 		}
 
@@ -578,20 +490,19 @@ namespace GeometryCollectionExample
 		// check simplicial elements
 		typedef TUniquePtr< FCollisionStructureManager::FSimplicial > FSimplicialPointer;
 		const TManagedArray<FSimplicialPointer> & Simplicials = Object->RestCollection->template GetAttribute<FSimplicialPointer>(FGeometryCollectionPhysicsProxy::SimplicialsAttribute, FTransformCollection::TransformGroup);
-		R.ExpectTrue(Simplicials.Num() == 1);
-		R.ExpectTrue(Simplicials[0].IsValid());
-		R.ExpectTrue(Simplicials[0]->Size() == 8);
+		EXPECT_EQ(Simplicials.Num(), 1);
+		EXPECT_TRUE(Simplicials[0].IsValid());
+		EXPECT_EQ(Simplicials[0]->Size(), 8);
 		for (int32 Index = 0; Index < (int32)Simplicials[0]->Size(); Index++)
 		{
 			//const FCollisionStructureManager::FSimplicial & Simplical = Simplicials[0];
 			Chaos::TVector<float, 3> Vert = Simplicials[0]->X(Index);
-			R.ExpectTrue((FMath::Abs(FMath::Abs(Vert.X) + FMath::Abs(Vert.Y) + FMath::Abs(Vert.Z))-1.5)<KINDA_SMALL_NUMBER );
+			EXPECT_LT((FMath::Abs(FMath::Abs(Vert.X) + FMath::Abs(Vert.Y) + FMath::Abs(Vert.Z))-1.5), KINDA_SMALL_NUMBER );
 		}
 		delete Object;
 #endif
-		return !R.HasError();
 	}
-	template bool RigidBodies_Initialize_ParticleImplicitCollisionGeometry<float>(ExampleResponse&& R);
+	template void RigidBodies_Initialize_ParticleImplicitCollisionGeometry<float>();
 
 }
 
