Copying //UE4/Dev-Physics to Dev-Main (//UE4/Dev-Main) [at] 9098278

diff --git a/Engine/Plugins/Experimental/GeometryCollectionPlugin/Source/GeometryCollectionExampleCore/Private/GeometryCollection/GeometryCollectionExampleSimulationField.cpp b/Engine/Plugins/Experimental/GeometryCollectionPlugin/Source/GeometryCollectionExampleCore/Private/GeometryCollection/GeometryCollectionExampleSimulationField.cpp
--- a/Engine/Plugins/Experimental/GeometryCollectionPlugin/Source/GeometryCollectionExampleCore/Private/GeometryCollection/GeometryCollectionExampleSimulationField.cpp
+++ b/Engine/Plugins/Experimental/GeometryCollectionPlugin/Source/GeometryCollectionExampleCore/Private/GeometryCollection/GeometryCollectionExampleSimulationField.cpp
@@ -28,22 +28,24 @@
 namespace GeometryCollectionExample
 {
 
-
 	template<class T>
-	bool RigidBodies_Field_KinematicActivation(ExampleResponse&& R)
+	void RigidBodies_Field_KinematicActivation()
 	{
 #if INCLUDE_CHAOS
-		TUniquePtr<Chaos::TChaosPhysicsMaterial<T>> PhysicalMaterial = MakeUnique<Chaos::TChaosPhysicsMaterial<T>>();
-		PhysicalMaterial->Friction = 0;
-		PhysicalMaterial->Restitution = 0;
-		PhysicalMaterial->SleepingLinearThreshold = 0;
-		PhysicalMaterial->SleepingAngularThreshold = 0;
-		PhysicalMaterial->DisabledLinearThreshold = 0;
-		PhysicalMaterial->DisabledAngularThreshold = 0;
-
-		TSharedPtr<FGeometryCollection> RestCollection = GeometryCollection::MakeCubeElement(FTransform::Identity, FVector(1.0));
-		RestCollection->Transform[0].SetTranslation(FVector(0, 0, 1));
-		TSharedPtr<FGeometryDynamicCollection> DynamicCollection = GeometryCollectionToGeometryDynamicCollection(RestCollection.Get(), (int32)EObjectStateTypeEnum::Chaos_Object_Kinematic);
+		TUniquePtr<Chaos::TChaosPhysicsMaterial<T>> PhysicalMaterial = nullptr;
+		TSharedPtr<FGeometryCollection> RestCollection = nullptr;
+		TSharedPtr<FGeometryDynamicCollection> DynamicCollection = nullptr;
+
+		//
+		//  Rigid Body Setup
+		//
+		auto RestInitFunc = [](TSharedPtr<FGeometryCollection>& RestCollection)
+		{
+			RestCollection->Transform[0].SetTranslation(FVector(0, 0, 1));
+		};
+
+		InitCollectionsParameters InitParams = { FTransform::Identity, FVector(1.0), RestInitFunc, (int32)EObjectStateTypeEnum::Chaos_Object_Kinematic };
+		InitCollections(PhysicalMaterial, RestCollection, DynamicCollection, InitParams);
 
 		FRadialIntMask * RadialMask = new FRadialIntMask();
 		RadialMask->Position = FVector(0.0, 0.0, 0.0);
@@ -52,19 +54,7 @@ namespace GeometryCollectionExample
 		RadialMask->ExteriorValue = (int32)EObjectStateTypeEnum::Chaos_Object_Kinematic;
 		RadialMask->SetMaskCondition = ESetMaskConditionType::Field_Set_IFF_NOT_Interior;
 
-		auto InitFunc = [&RestCollection, &DynamicCollection, &PhysicalMaterial](FSimulationParameters& InParams)
-		{
-			InParams.RestCollection = RestCollection.Get();
-			InParams.DynamicCollection = DynamicCollection.Get();
-			InParams.PhysicalMaterial  = MakeSerializable(PhysicalMaterial);
-			InParams.Shared.SizeSpecificData[0].CollisionType = ECollisionTypeEnum::Chaos_Volumetric;
-			InParams.Simulating = true;
-			Chaos::FErrorReporter ErrorReporter;
-			BuildSimulationData(ErrorReporter, *RestCollection, InParams.Shared);
-		};
-
-		FGeometryCollectionPhysicsProxy* PhysObject = new FGeometryCollectionPhysicsProxy(nullptr, DynamicCollection.Get(), InitFunc, nullptr, nullptr);;
-		PhysObject->Initialize();
+		FGeometryCollectionPhysicsProxy* PhysObject = RigidBodySetup(PhysicalMaterial, RestCollection, DynamicCollection);
 
 		Chaos::FPBDRigidsSolver* Solver = FChaosSolversModule::GetModule()->CreateSolver(true);
 #if CHAOS_PARTICLEHANDLE_TODO
@@ -84,8 +74,8 @@ namespace GeometryCollectionExample
 
 		// simulated
 		TManagedArray<FTransform>& Transform = DynamicCollection->Transform;
-		R.ExpectTrue(Transform.Num() == 1);
-		R.ExpectTrue(Transform[0].GetTranslation().Z == 1.f);
+		EXPECT_EQ(Transform.Num(), 1);
+		EXPECT_EQ(Transform[0].GetTranslation().Z, 1.f);
 		FName TargetName = GetFieldPhysicsName(EFieldPhysicsType::Field_DynamicState);
 		PhysObject->BufferCommand(Solver, { TargetName, RadialMask });
 
@@ -96,47 +86,35 @@ namespace GeometryCollectionExample
 
 		FinalizeSolver(*Solver);
 
-		R.ExpectTrue(Transform[0].GetTranslation().Z <= 0.f);	
-		
+		EXPECT_LE(Transform[0].GetTranslation().Z, 0.f);
+
 		FChaosSolversModule::GetModule()->DestroySolver(Solver);
 		delete PhysObject;
 
 #endif
-
-
-		return !R.HasError();
 	}
-	template bool RigidBodies_Field_KinematicActivation<float>(ExampleResponse&& R);
+	template void RigidBodies_Field_KinematicActivation<float>();
 
 	template<class T>
-	bool RigidBodies_Field_InitialLinearVelocity(ExampleResponse&& R)
+	void RigidBodies_Field_InitialLinearVelocity()
 	{
 #if INCLUDE_CHAOS
-		TUniquePtr<Chaos::TChaosPhysicsMaterial<T>> PhysicalMaterial = MakeUnique<Chaos::TChaosPhysicsMaterial<T>>();
-		PhysicalMaterial->Friction = 0;
-		PhysicalMaterial->Restitution = 0;
-		PhysicalMaterial->SleepingLinearThreshold = 0;
-		PhysicalMaterial->SleepingAngularThreshold = 0;
-		PhysicalMaterial->DisabledLinearThreshold = 0;
-		PhysicalMaterial->DisabledAngularThreshold = 0;
-
-		TSharedPtr<FGeometryCollection> RestCollection = GeometryCollection::MakeCubeElement(FTransform::Identity, FVector(1.0));
-		TSharedPtr<FGeometryDynamicCollection> DynamicCollection = GeometryCollectionToGeometryDynamicCollection(RestCollection.Get(),(int)EObjectStateTypeEnum::Chaos_Object_Kinematic);
-
-		auto InitFunc = [&RestCollection, &DynamicCollection, &PhysicalMaterial](FSimulationParameters& InParams)
+		TUniquePtr<Chaos::TChaosPhysicsMaterial<T>> PhysicalMaterial = nullptr;
+		TSharedPtr<FGeometryCollection> RestCollection = nullptr;
+		TSharedPtr<FGeometryDynamicCollection> DynamicCollection = nullptr;
+
+		//
+		//  Rigid Body Setup
+		//
+		InitCollectionsParameters InitParams = { FTransform::Identity, FVector(1.0), nullptr, (int)EObjectStateTypeEnum::Chaos_Object_Kinematic };
+		InitCollections(PhysicalMaterial, RestCollection, DynamicCollection, InitParams);
+
+		auto CustomFunc = [&RestCollection, &DynamicCollection, &PhysicalMaterial](FSimulationParameters& InParams)
 		{
-			InParams.RestCollection = RestCollection.Get();
-			InParams.DynamicCollection = DynamicCollection.Get();
-			InParams.PhysicalMaterial  = MakeSerializable(PhysicalMaterial);
-			InParams.Shared.SizeSpecificData[0].CollisionType = ECollisionTypeEnum::Chaos_Volumetric;
 			InParams.InitialVelocityType = EInitialVelocityTypeEnum::Chaos_Initial_Velocity_User_Defined;
 			InParams.InitialLinearVelocity = FVector(0.f, 100.f, 0.f);
-			InParams.Simulating = true;
-			Chaos::FErrorReporter ErrorReporter;
-			BuildSimulationData(ErrorReporter, *RestCollection, InParams.Shared);
 		};
 
-
 		//
 		// Field setup
 		//
@@ -151,8 +129,7 @@ namespace GeometryCollectionExample
 		// Solver setup
 		//
 		FFieldSystemPhysicsProxy* FieldObject = new FFieldSystemPhysicsProxy(nullptr);
-		FGeometryCollectionPhysicsProxy* PhysObject = new FGeometryCollectionPhysicsProxy(nullptr, DynamicCollection.Get(), InitFunc, nullptr, nullptr);;
-		PhysObject->Initialize();
+		FGeometryCollectionPhysicsProxy* PhysObject = RigidBodySetup(PhysicalMaterial, RestCollection, DynamicCollection, CustomFunc);
 
 		Chaos::FPBDRigidsSolver* Solver = FChaosSolversModule::GetModule()->CreateSolver(true);
 #if CHAOS_PARTICLEHANDLE_TODO
@@ -167,8 +144,8 @@ namespace GeometryCollectionExample
 		TManagedArray<FTransform>& Transform = DynamicCollection->Transform;
 
 		float PreviousY = 0.f;
-		R.ExpectTrue(Transform[0].GetTranslation().X == 0);
-		R.ExpectTrue(Transform[0].GetTranslation().Y == 0);
+		EXPECT_EQ(Transform[0].GetTranslation().X, 0);
+		EXPECT_EQ(Transform[0].GetTranslation().Y, 0);
 
 		for (int Frame = 0; Frame < 10; Frame++)
 		{
@@ -183,59 +160,44 @@ namespace GeometryCollectionExample
 			FinalizeSolver(*Solver);
 			if (Frame >= 2)
 			{
-				R.ExpectTrue(Transform[0].GetTranslation().X == 0);
-				R.ExpectTrue(Transform[0].GetTranslation().Y > PreviousY);
+				EXPECT_EQ(Transform[0].GetTranslation().X, 0);
+				EXPECT_GT(Transform[0].GetTranslation().Y, PreviousY);
 			}
 			else
 			{
-				R.ExpectTrue(Transform[0].GetTranslation().X == 0);
-				R.ExpectTrue(Transform[0].GetTranslation().Y == 0);
-				R.ExpectTrue(Transform[0].GetTranslation().Z == 0);
+				EXPECT_EQ(Transform[0].GetTranslation().X, 0);
+				EXPECT_EQ(Transform[0].GetTranslation().Y, 0);
+				EXPECT_EQ(Transform[0].GetTranslation().Z, 0);
 			}
 			PreviousY = Transform[0].GetTranslation().Y;
 		}
-		
+
 		FChaosSolversModule::GetModule()->DestroySolver(Solver);
+		delete FieldObject;
 		delete PhysObject;
 
 #endif
-
-
-		return !R.HasError();
 	}
-	template bool RigidBodies_Field_InitialLinearVelocity<float>(ExampleResponse&& R);
+	template void RigidBodies_Field_InitialLinearVelocity<float>();
 
 	template<class T>
-	bool RigidBodies_Field_StayDynamic(ExampleResponse&& R)
+	void RigidBodies_Field_StayDynamic()
 	{
 #if INCLUDE_CHAOS
-		TUniquePtr<Chaos::TChaosPhysicsMaterial<T>> PhysicalMaterial = MakeUnique<Chaos::TChaosPhysicsMaterial<T>>();
-		PhysicalMaterial->Friction = 0;
-		PhysicalMaterial->Restitution = 0;
-		PhysicalMaterial->SleepingLinearThreshold = 0;
-		PhysicalMaterial->SleepingAngularThreshold = 0;
-		PhysicalMaterial->DisabledLinearThreshold = 0;
-		PhysicalMaterial->DisabledAngularThreshold = 0;
+		TUniquePtr<Chaos::TChaosPhysicsMaterial<T>> PhysicalMaterial = nullptr;
+		TSharedPtr<FGeometryCollection> RestCollection = nullptr;
+		TSharedPtr<FGeometryDynamicCollection> DynamicCollection = nullptr;
+
+		auto RestInitFunc = [](TSharedPtr<FGeometryCollection>& RestCollection)
+		{
+			RestCollection->Transform[0].SetTranslation(FVector(0.f, 0.f, 5.f));
+		};
 
 		//
 		//  Rigid Body Setup
 		//
-		TSharedPtr<FGeometryCollection> RestCollection = GeometryCollection::MakeCubeElement(FTransform::Identity, FVector(1.0));
-		TManagedArray<FTransform>& RestTransform = RestCollection->Transform;
-
-		RestTransform[0].SetTranslation(FVector(0.f, 0.f, 5.f));
-		TSharedPtr<FGeometryDynamicCollection> DynamicCollection = GeometryCollectionToGeometryDynamicCollection(RestCollection.Get(), (int32)EObjectStateTypeEnum::Chaos_Object_Static);
-
-		auto InitFunc = [&RestCollection, &DynamicCollection, &PhysicalMaterial](FSimulationParameters& InParams)
-		{
-			InParams.RestCollection = RestCollection.Get();
-			InParams.DynamicCollection = DynamicCollection.Get();
-			InParams.PhysicalMaterial  = MakeSerializable(PhysicalMaterial);
-			InParams.Shared.SizeSpecificData[0].CollisionType = ECollisionTypeEnum::Chaos_Volumetric;
-			InParams.Simulating = true;
-			Chaos::FErrorReporter ErrorReporter;
-			BuildSimulationData(ErrorReporter, *RestCollection, InParams.Shared);
-		};
+		InitCollectionsParameters InitParams = { FTransform::Identity, FVector(1.0), RestInitFunc, (int32)EObjectStateTypeEnum::Chaos_Object_Static };
+		InitCollections(PhysicalMaterial, RestCollection, DynamicCollection, InitParams);
 
 		//
 		// Field setup
@@ -247,12 +209,10 @@ namespace GeometryCollectionExample
 		RadialMask->ExteriorValue = (int32)EObjectStateTypeEnum::Chaos_Object_Kinematic;
 		RadialMask->SetMaskCondition = ESetMaskConditionType::Field_Set_IFF_NOT_Interior;
 
-
 		//
 		// Solver setup
 		//
-		FGeometryCollectionPhysicsProxy* PhysObject = new FGeometryCollectionPhysicsProxy(nullptr, DynamicCollection.Get(), InitFunc, nullptr, nullptr);;
-		PhysObject->Initialize();
+		FGeometryCollectionPhysicsProxy* PhysObject = RigidBodySetup(PhysicalMaterial, RestCollection, DynamicCollection);
 
 		FFieldSystemPhysicsProxy* FieldObject = new FFieldSystemPhysicsProxy(nullptr);
 
@@ -283,11 +243,11 @@ namespace GeometryCollectionExample
 
 			if (Frame < 5)
 			{
-				R.ExpectTrue(FMath::Abs(Transform[0].GetTranslation().Z - 5.f) < SMALL_THRESHOLD);
+				EXPECT_LT(FMath::Abs(Transform[0].GetTranslation().Z - 5.f), SMALL_THRESHOLD);
 			}
 			else
 			{
-				R.ExpectTrue(Transform[0].GetTranslation().Z < PreviousHeight);
+				EXPECT_LT(Transform[0].GetTranslation().Z, PreviousHeight);
 			}
 			PreviousHeight = Transform[0].GetTranslation().Z;
 
@@ -297,48 +257,35 @@ namespace GeometryCollectionExample
 			//	UE_LOG(GCTCL_Log, Verbose, TEXT("... ... ...Disabled[%d] : %d"), rdx, Particles.Disabled(rdx));
 			//}
 		}
-		
+
 		FChaosSolversModule::GetModule()->DestroySolver(Solver);
 
 		delete PhysObject;
 		delete FieldObject;
 
 #endif
-		return !R.HasError();
 	}
-	template bool RigidBodies_Field_StayDynamic<float>(ExampleResponse&& R);
+	template void RigidBodies_Field_StayDynamic<float>();
 
 
 	template<class T>
-	bool RigidBodies_Field_LinearForce(ExampleResponse&& R)
+	void RigidBodies_Field_LinearForce()
 	{
 #if INCLUDE_CHAOS
-		TUniquePtr<Chaos::TChaosPhysicsMaterial<T>> PhysicalMaterial = MakeUnique<Chaos::TChaosPhysicsMaterial<T>>();
-		PhysicalMaterial->Friction = 0;
-		PhysicalMaterial->Restitution = 0;
-		PhysicalMaterial->SleepingLinearThreshold = 0;
-		PhysicalMaterial->SleepingAngularThreshold = 0;
-		PhysicalMaterial->DisabledLinearThreshold = 0;
-		PhysicalMaterial->DisabledAngularThreshold = 0;
+		TUniquePtr<Chaos::TChaosPhysicsMaterial<T>> PhysicalMaterial = nullptr;
+		TSharedPtr<FGeometryCollection> RestCollection = nullptr;
+		TSharedPtr<FGeometryDynamicCollection> DynamicCollection = nullptr;
+
+		auto RestInitFunc = [](TSharedPtr<FGeometryCollection>& RestCollection)
+		{
+			RestCollection->Transform[0].SetTranslation(FVector(0.f, 0.f, 5.f));
+		};
 
 		//
 		//  Rigid Body Setup
 		//
-		TSharedPtr<FGeometryCollection> RestCollection = GeometryCollection::MakeCubeElement(FTransform::Identity, FVector(1.0));
-		TManagedArray<FTransform>& RestTransform = RestCollection->Transform;
-		RestTransform[0].SetTranslation(FVector(0.f, 0.f, 5.f));
-		TSharedPtr<FGeometryDynamicCollection> DynamicCollection = GeometryCollectionToGeometryDynamicCollection(RestCollection.Get(), (int32)EObjectStateTypeEnum::Chaos_Object_Dynamic);
-
-		auto InitFunc = [&RestCollection, &DynamicCollection, &PhysicalMaterial](FSimulationParameters& InParams)
-		{
-			InParams.RestCollection = RestCollection.Get();
-			InParams.DynamicCollection = DynamicCollection.Get();
-			InParams.PhysicalMaterial  = MakeSerializable(PhysicalMaterial);
-			InParams.Shared.SizeSpecificData[0].CollisionType = ECollisionTypeEnum::Chaos_Volumetric;
-			InParams.Simulating = true;
-			Chaos::FErrorReporter ErrorReporter;
-			BuildSimulationData(ErrorReporter, *RestCollection, InParams.Shared);
-		};
+		InitCollectionsParameters InitParams = { FTransform::Identity, FVector(1.0), RestInitFunc, (int32)EObjectStateTypeEnum::Chaos_Object_Dynamic };
+		InitCollections(PhysicalMaterial, RestCollection, DynamicCollection, InitParams);
 
 		//
 		// Field setup
@@ -350,8 +297,7 @@ namespace GeometryCollectionExample
 		//
 		// Solver setup
 		//
-		FGeometryCollectionPhysicsProxy* PhysObject = new FGeometryCollectionPhysicsProxy(nullptr, DynamicCollection.Get(), InitFunc, nullptr, nullptr);;
-		PhysObject->Initialize();
+		FGeometryCollectionPhysicsProxy* PhysObject = RigidBodySetup(PhysicalMaterial, RestCollection, DynamicCollection);
 
 		FFieldSystemPhysicsProxy* FieldObject = new FFieldSystemPhysicsProxy(nullptr);
 
@@ -381,11 +327,11 @@ namespace GeometryCollectionExample
 
 			if (Frame < 5)
 			{
-				R.ExpectTrue(FMath::Abs(Transform[0].GetTranslation().Y) < SMALL_THRESHOLD);
+				EXPECT_LT(FMath::Abs(Transform[0].GetTranslation().Y), SMALL_THRESHOLD);
 			}
 			else
 			{
-				R.ExpectTrue(Transform[0].GetTranslation().Y > PreviousY);
+				EXPECT_GT(Transform[0].GetTranslation().Y, PreviousY);
 			}
 
 			PreviousY = Transform[0].GetTranslation().Y;
@@ -398,42 +344,28 @@ namespace GeometryCollectionExample
 		delete PhysObject;
 		delete FieldObject;
 #endif
-
-		return !R.HasError();		
 	}
-	template bool RigidBodies_Field_LinearForce<float>(ExampleResponse&& R);
+	template void RigidBodies_Field_LinearForce<float>();
 
 	template<class T>
-	bool RigidBodies_Field_Torque(ExampleResponse&& R)
+	void RigidBodies_Field_Torque()
 	{
 #if INCLUDE_CHAOS
-		TUniquePtr<Chaos::TChaosPhysicsMaterial<T>> PhysicalMaterial = MakeUnique<Chaos::TChaosPhysicsMaterial<T>>();
-		PhysicalMaterial->Friction = 0;
-		PhysicalMaterial->Restitution = 0;
-		PhysicalMaterial->SleepingLinearThreshold = 0;
-		PhysicalMaterial->SleepingAngularThreshold = 0;
-		PhysicalMaterial->DisabledLinearThreshold = 0;
-		PhysicalMaterial->DisabledAngularThreshold = 0;
+		TUniquePtr<Chaos::TChaosPhysicsMaterial<T>> PhysicalMaterial = nullptr;
+		TSharedPtr<FGeometryCollection> RestCollection = nullptr;
+		TSharedPtr<FGeometryDynamicCollection> DynamicCollection = nullptr;
 
 		//
 		//  Rigid Body Setup
 		//
-		TSharedPtr<FGeometryCollection> RestCollection = GeometryCollection::MakeCubeElement(FTransform::Identity, FVector(10.0));
-		TManagedArray<FTransform>& RestTransform = RestCollection->Transform;
-		RestTransform[0].SetTranslation(FVector(0.f, 0.f, 5.f));
-		TSharedPtr<FGeometryDynamicCollection> DynamicCollection = GeometryCollectionToGeometryDynamicCollection(RestCollection.Get(), (int32)EObjectStateTypeEnum::Chaos_Object_Dynamic);
-
-		auto InitFunc = [&RestCollection, &DynamicCollection, &PhysicalMaterial](FSimulationParameters& InParams)
+		auto RestInitFunc = [](TSharedPtr<FGeometryCollection>& RestCollection)
 		{
-			InParams.RestCollection = RestCollection.Get();
-			InParams.DynamicCollection = DynamicCollection.Get();
-			InParams.PhysicalMaterial  = MakeSerializable(PhysicalMaterial);
-			InParams.Shared.SizeSpecificData[0].CollisionType = ECollisionTypeEnum::Chaos_Volumetric;
-			InParams.Simulating = true;
-			Chaos::FErrorReporter ErrorReporter;
-			BuildSimulationData(ErrorReporter, *RestCollection, InParams.Shared);
+			RestCollection->Transform[0].SetTranslation(FVector(0.f, 0.f, 5.f));
 		};
 
+		InitCollectionsParameters InitParams = {FTransform::Identity, FVector(10.0), RestInitFunc, (int32)EObjectStateTypeEnum::Chaos_Object_Dynamic};
+		InitCollections(PhysicalMaterial, RestCollection, DynamicCollection, InitParams);
+
 		//
 		// Field setup
 		//
@@ -444,8 +376,7 @@ namespace GeometryCollectionExample
 		//
 		// Solver setup
 		//
-		FGeometryCollectionPhysicsProxy* PhysObject = new FGeometryCollectionPhysicsProxy(nullptr, DynamicCollection.Get(), InitFunc, nullptr, nullptr);;
-		PhysObject->Initialize();
+		FGeometryCollectionPhysicsProxy* PhysObject = RigidBodySetup(PhysicalMaterial, RestCollection, DynamicCollection);
 
 		FFieldSystemPhysicsProxy* FieldObject = new FFieldSystemPhysicsProxy(nullptr);
 
@@ -479,12 +410,12 @@ namespace GeometryCollectionExample
 
 			if (Frame < 5)
 			{
-				R.ExpectTrue(FMath::Abs(Transform[0].GetRotation().Euler().Y) < SMALL_THRESHOLD);
+				EXPECT_LT(FMath::Abs(Transform[0].GetRotation().Euler().Y), SMALL_THRESHOLD);
 			}
 			else
 			{
-				R.ExpectTrue(FMath::Abs(Transform[0].GetRotation().Euler().Y) != SMALL_THRESHOLD);
-				R.ExpectTrue(Particles.W(0).Y > PreviousY);
+				EXPECT_NE(FMath::Abs(Transform[0].GetRotation().Euler().Y), SMALL_THRESHOLD);
+				EXPECT_GT(Particles.W(0).Y, PreviousY);
 			}
 
 			PreviousY = Particles.W(0).Y;
@@ -498,44 +429,29 @@ namespace GeometryCollectionExample
 		delete PhysObject;
 		delete FieldObject;
 #endif
-
-		return !R.HasError();
 	}
-	template bool RigidBodies_Field_Torque<float>(ExampleResponse&& R);
+	template void RigidBodies_Field_Torque<float>();
 
 
 
 	template<class T>
-	bool RigidBodies_Field_Kill(ExampleResponse&& R)
+	void RigidBodies_Field_Kill()
 	{
 #if INCLUDE_CHAOS
-		TUniquePtr<Chaos::TChaosPhysicsMaterial<T>> PhysicalMaterial = MakeUnique<Chaos::TChaosPhysicsMaterial<T>>();
-		PhysicalMaterial->Friction = 0;
-		PhysicalMaterial->Restitution = 0;
-		PhysicalMaterial->SleepingLinearThreshold = 0;
-		PhysicalMaterial->SleepingAngularThreshold = 0;
-		PhysicalMaterial->DisabledLinearThreshold = 0;
-		PhysicalMaterial->DisabledAngularThreshold = 0;
+		TUniquePtr<Chaos::TChaosPhysicsMaterial<T>> PhysicalMaterial = nullptr;
+		TSharedPtr<FGeometryCollection> RestCollection = nullptr;
+		TSharedPtr<FGeometryDynamicCollection> DynamicCollection = nullptr;
+
+		auto RestInitFunc = [](TSharedPtr<FGeometryCollection>& RestCollection)
+		{
+			RestCollection->Transform[0].SetTranslation(FVector(0.f, 0.f, 20.f));
+		};
 
 		//
 		//  Rigid Body Setup
 		//
-		TSharedPtr<FGeometryCollection> RestCollection = GeometryCollection::MakeCubeElement(FTransform::Identity, FVector(1.0));
-		TManagedArray<FTransform>& RestTransform = RestCollection->Transform;
-
-		RestTransform[0].SetTranslation(FVector(0.f, 0.f, 20.f));
-		TSharedPtr<FGeometryDynamicCollection> DynamicCollection = GeometryCollectionToGeometryDynamicCollection(RestCollection.Get());
-
-		auto InitFunc = [&RestCollection, &DynamicCollection, &PhysicalMaterial](FSimulationParameters& InParams)
-		{
-			InParams.RestCollection = RestCollection.Get();
-			InParams.DynamicCollection = DynamicCollection.Get();
-			InParams.PhysicalMaterial  = MakeSerializable(PhysicalMaterial);
-			InParams.Shared.SizeSpecificData[0].CollisionType = ECollisionTypeEnum::Chaos_Volumetric;
-			InParams.Simulating = true;
-			Chaos::FErrorReporter ErrorReporter;
-			BuildSimulationData(ErrorReporter, *RestCollection, InParams.Shared);
-		};
+		InitCollectionsParameters InitParams = { FTransform::Identity, FVector(1.0), RestInitFunc, (int32)EObjectStateTypeEnum::Chaos_Object_Dynamic };
+		InitCollections(PhysicalMaterial, RestCollection, DynamicCollection, InitParams);
 
 		//
 		// Field setup
@@ -549,8 +465,7 @@ namespace GeometryCollectionExample
 		//
 		// Solver setup
 		//
-		FGeometryCollectionPhysicsProxy* PhysObject = new FGeometryCollectionPhysicsProxy(nullptr, DynamicCollection.Get(), InitFunc, nullptr, nullptr);;
-		PhysObject->Initialize();
+		FGeometryCollectionPhysicsProxy* PhysObject = RigidBodySetup(PhysicalMaterial, RestCollection, DynamicCollection);
 
 		FFieldSystemPhysicsProxy* FieldObject = new FFieldSystemPhysicsProxy(nullptr);
 
@@ -575,8 +490,8 @@ namespace GeometryCollectionExample
 
 			FinalizeSolver(*Solver);
 
-			R.ExpectTrue(Transform[0].GetTranslation().Z < 20.f);
-			R.ExpectTrue(Transform[0].GetTranslation().Z > -10.);
+			EXPECT_LT(Transform[0].GetTranslation().Z, 20.f);
+			EXPECT_GT(Transform[0].GetTranslation().Z, -10.);
 		}
 
 		FChaosSolversModule::GetModule()->DestroySolver(Solver);
@@ -586,41 +501,27 @@ namespace GeometryCollectionExample
 		delete FalloffField;
 
 #endif
-		return !R.HasError();
 	}
-	template bool RigidBodies_Field_Kill<float>(ExampleResponse&& R);
+	template void RigidBodies_Field_Kill<float>();
 
 	template<class T>
-	bool RigidBodies_Field_LinearVelocity(ExampleResponse&& R)
+	void RigidBodies_Field_LinearVelocity()
 	{
 #if INCLUDE_CHAOS
-		TUniquePtr<Chaos::TChaosPhysicsMaterial<T>> PhysicalMaterial = MakeUnique<Chaos::TChaosPhysicsMaterial<T>>();
-		PhysicalMaterial->Friction = 0;
-		PhysicalMaterial->Restitution = 0;
-		PhysicalMaterial->SleepingLinearThreshold = 0;
-		PhysicalMaterial->SleepingAngularThreshold = 0;
-		PhysicalMaterial->DisabledLinearThreshold = 0;
-		PhysicalMaterial->DisabledAngularThreshold = 0;
+		TUniquePtr<Chaos::TChaosPhysicsMaterial<T>> PhysicalMaterial = nullptr;
+		TSharedPtr<FGeometryCollection> RestCollection = nullptr;
+		TSharedPtr<FGeometryDynamicCollection> DynamicCollection = nullptr;
+
+		auto RestInitFunc = [](TSharedPtr<FGeometryCollection>& RestCollection)
+		{
+			RestCollection->Transform[0].SetTranslation(FVector(0.f, 0.f, 20.f));
+		};
 
 		//
 		//  Rigid Body Setup
 		//
-		TSharedPtr<FGeometryCollection> RestCollection = GeometryCollection::MakeCubeElement(FTransform::Identity, FVector(1.0));
-		TManagedArray<FTransform>& RestTransform = RestCollection->Transform;
-
-		RestTransform[0].SetTranslation(FVector(0.f, 0.f, 20.f));
-		TSharedPtr<FGeometryDynamicCollection> DynamicCollection = GeometryCollectionToGeometryDynamicCollection(RestCollection.Get());
-
-		auto InitFunc = [&RestCollection, &DynamicCollection, &PhysicalMaterial](FSimulationParameters& InParams)
-		{
-			InParams.RestCollection = RestCollection.Get();
-			InParams.DynamicCollection = DynamicCollection.Get();
-			InParams.PhysicalMaterial  = MakeSerializable(PhysicalMaterial);
-			InParams.Shared.SizeSpecificData[0].CollisionType = ECollisionTypeEnum::Chaos_Volumetric;
-			InParams.Simulating = true;
-			Chaos::FErrorReporter ErrorReporter;
-			BuildSimulationData(ErrorReporter, *RestCollection, InParams.Shared);
-		};
+		InitCollectionsParameters InitParams = { FTransform::Identity, FVector(1.0), RestInitFunc, (int32)EObjectStateTypeEnum::Chaos_Object_Dynamic };
+		InitCollections(PhysicalMaterial, RestCollection, DynamicCollection, InitParams);
 
 		//
 		// Field setup
@@ -632,8 +533,7 @@ namespace GeometryCollectionExample
 		//
 		// Solver setup
 		//
-		FGeometryCollectionPhysicsProxy* PhysObject = new FGeometryCollectionPhysicsProxy(nullptr, DynamicCollection.Get(), InitFunc, nullptr, nullptr);;
-		PhysObject->Initialize();
+		FGeometryCollectionPhysicsProxy* PhysObject = RigidBodySetup(PhysicalMaterial, RestCollection, DynamicCollection);
 
 		FFieldSystemPhysicsProxy* FieldObject = new FFieldSystemPhysicsProxy(nullptr);
 
@@ -660,7 +560,7 @@ namespace GeometryCollectionExample
 			Solver->AdvanceSolverBy(1 / 24.);
 			FinalizeSolver(*Solver);
 
-			R.ExpectTrue(Transform[0].GetTranslation().X > PreviousX);
+			EXPECT_GT(Transform[0].GetTranslation().X, PreviousX);
 			PreviousX = Transform[0].GetTranslation().X;
 		}
 
@@ -671,36 +571,36 @@ namespace GeometryCollectionExample
 		delete VectorField;
 
 #endif
-		return !R.HasError();
 	}
-	template bool RigidBodies_Field_LinearVelocity<float>(ExampleResponse&& R);
+	template void RigidBodies_Field_LinearVelocity<float>();
 
 
 	/**
-	 * Create a stack of boxes on the ground and verify that we we change their collision 
+	 * Create a stack of boxes on the ground and verify that we we change their collision
 	 * group, they drop through the ground.
 	 */
 	template<class T>
-	bool RigidBodies_Field_CollisionGroup(ExampleResponse&& R)
+	void RigidBodies_Field_CollisionGroup()
 	{
 #if INCLUDE_CHAOS
-		TUniquePtr<Chaos::TChaosPhysicsMaterial<T>> PhysicalMaterial = MakeUnique<Chaos::TChaosPhysicsMaterial<T>>();
-		PhysicalMaterial->Friction = 0;
-		PhysicalMaterial->Restitution = 0;
-		PhysicalMaterial->SleepingLinearThreshold = 0;
-		PhysicalMaterial->SleepingAngularThreshold = 0;
-		PhysicalMaterial->DisabledLinearThreshold = 0;
-		PhysicalMaterial->DisabledAngularThreshold = 0;
+		TUniquePtr<Chaos::TChaosPhysicsMaterial<T>> PhysicalMaterial = nullptr;
+		TSharedPtr<FGeometryCollection> RestCollection = nullptr;
+		TSharedPtr<FGeometryDynamicCollection> DynamicCollection = nullptr;
 
 		//
 		// Generate Geometry - a stack of boxes.
 		// The bottom box is on the ground, and the others are dropped into it.
 		//
-		TSharedPtr<FGeometryCollection> RestCollection = GeometryCollection::MakeCubeElement(FTransform(FVector(0, 0, 100)), FVector(200));
-		RestCollection->AppendGeometry(*GeometryCollection::MakeCubeElement(FTransform(FVector(0, 0, 400)), FVector(100)));
-		RestCollection->AppendGeometry(*GeometryCollection::MakeCubeElement(FTransform(FVector(0, 0, 600)), FVector(100)));
-		RestCollection->AppendGeometry(*GeometryCollection::MakeCubeElement(FTransform(FVector(0, 0, 800)), FVector(100)));
-		TSharedPtr<FGeometryDynamicCollection> DynamicCollection = GeometryCollectionToGeometryDynamicCollection(RestCollection.Get());
+		auto RestInitFunc = [](TSharedPtr<FGeometryCollection>& RestCollection)
+		{
+			RestCollection->AppendGeometry(*GeometryCollection::MakeCubeElement(FTransform(FVector(0, 0, 400)), FVector(100)));
+			RestCollection->AppendGeometry(*GeometryCollection::MakeCubeElement(FTransform(FVector(0, 0, 600)), FVector(100)));
+			RestCollection->AppendGeometry(*GeometryCollection::MakeCubeElement(FTransform(FVector(0, 0, 800)), FVector(100)));
+		};
+
+		InitCollectionsParameters InitParams = { FTransform(FVector(0, 0, 100)), FVector(200), RestInitFunc, (int32)EObjectStateTypeEnum::Chaos_Object_Dynamic };
+
+		InitCollections(PhysicalMaterial, RestCollection, DynamicCollection, InitParams);
 
 		//
 		// Field setup
@@ -715,23 +615,7 @@ namespace GeometryCollectionExample
 		//
 		// Solver setup
 		//
-
-		//
-		// Solver setup
-		//
-		auto InitFunc = [&RestCollection, &DynamicCollection, &PhysicalMaterial](FSimulationParameters& InParams)
-		{
-			InParams.RestCollection = RestCollection.Get();
-			InParams.DynamicCollection = DynamicCollection.Get();
-			InParams.PhysicalMaterial  = MakeSerializable(PhysicalMaterial);
-			InParams.Shared.SizeSpecificData[0].CollisionType = ECollisionTypeEnum::Chaos_Volumetric;
-			InParams.Simulating = true;
-			Chaos::FErrorReporter ErrorReporter;
-			BuildSimulationData(ErrorReporter, *RestCollection, InParams.Shared);
-		};
-
-		FGeometryCollectionPhysicsProxy* PhysObject = new FGeometryCollectionPhysicsProxy(nullptr, DynamicCollection.Get(), InitFunc, nullptr, nullptr);;
-		PhysObject->Initialize();
+		FGeometryCollectionPhysicsProxy* PhysObject = RigidBodySetup(PhysicalMaterial, RestCollection, DynamicCollection);
 
 		Chaos::FPBDRigidsSolver* Solver = FChaosSolversModule::GetModule()->CreateSolver(true);
 #if CHAOS_PARTICLEHANDLE_TODO
@@ -752,11 +636,11 @@ namespace GeometryCollectionExample
 			if (Frame == 30)
 			{
 				// The boxes should have landed on each other and settled by now
-				R.ExpectTrue(FMath::Abs(Particles.X(0).Z) < SMALL_NUMBER);
-				R.ExpectTrue(FMath::IsNearlyEqual(Particles.X(1).Z, (T)100, (T)2));
-				R.ExpectTrue(FMath::IsNearlyEqual(Particles.X(2).Z, (T)250, (T)2));
-				R.ExpectTrue(FMath::IsNearlyEqual(Particles.X(3).Z, (T)350, (T)2));
-				R.ExpectTrue(FMath::IsNearlyEqual(Particles.X(4).Z, (T)450, (T)2));
+				EXPECT_LT(FMath::Abs(Particles.X(0).Z), SMALL_NUMBER);
+				EXPECT_TRUE(FMath::IsNearlyEqual(Particles.X(1).Z, (T)100, (T)2));
+				EXPECT_TRUE(FMath::IsNearlyEqual(Particles.X(2).Z, (T)250, (T)2));
+				EXPECT_TRUE(FMath::IsNearlyEqual(Particles.X(3).Z, (T)350, (T)2));
+				EXPECT_TRUE(FMath::IsNearlyEqual(Particles.X(4).Z, (T)450, (T)2));
 			}
 			if (Frame == 31)
 			{
@@ -766,55 +650,43 @@ namespace GeometryCollectionExample
 		}
 		// The boxes should have fallen below the ground level
 		Chaos::TPBDRigidParticles<float, 3>& Particles = Solver->GetRigidParticles();
-		R.ExpectTrue(FMath::Abs(Particles.X(0).Z) < SMALL_NUMBER);
-		R.ExpectTrue(Particles.X(1).Z < 0);
-		R.ExpectTrue(Particles.X(2).Z < 0);
-		R.ExpectTrue(Particles.X(3).Z < 0);
-		R.ExpectTrue(Particles.X(4).Z < 0);
+		EXPECT_LT(FMath::Abs(Particles.X(0).Z), SMALL_NUMBER);
+		EXPECT_LT(Particles.X(1).Z, 0);
+		EXPECT_LT(Particles.X(2).Z, 0);
+		EXPECT_LT(Particles.X(3).Z, 0);
+		EXPECT_LT(Particles.X(4).Z, 0);
 #endif
 
 		FChaosSolversModule::GetModule()->DestroySolver(Solver);
 
 		delete PhysObject;
+		delete RadialMask;
 #endif
-
-		return !R.HasError();
 	}
-	template bool RigidBodies_Field_CollisionGroup<float>(ExampleResponse&& R);
-
+	template void RigidBodies_Field_CollisionGroup<float>();
 
 	template<class T>
-	bool RigidBodies_Field_ClusterBreak_StrainModel_Test1(ExampleResponse&& R)
+	void RigidBodies_Field_ClusterBreak_StrainModel_Test1()
 	{
 #if INCLUDE_CHAOS
-		TUniquePtr<Chaos::TChaosPhysicsMaterial<T>> PhysicalMaterial = MakeUnique<Chaos::TChaosPhysicsMaterial<T>>();
-		PhysicalMaterial->Friction = 0;
-		PhysicalMaterial->Restitution = 0;
-		PhysicalMaterial->SleepingLinearThreshold = 0;
-		PhysicalMaterial->SleepingAngularThreshold = 0;
-		PhysicalMaterial->DisabledLinearThreshold = 0;
-		PhysicalMaterial->DisabledAngularThreshold = 0;
-
+		TUniquePtr<Chaos::TChaosPhysicsMaterial<T>> PhysicalMaterial = nullptr;
 		TSharedPtr<FGeometryCollection> RestCollection = CreateClusteredBody_TwoByTwo_ThreeTransform(FVector(0));
-		TSharedPtr<FGeometryDynamicCollection> DynamicCollection = GeometryCollectionToGeometryDynamicCollection(RestCollection.Get(), (uint8)EObjectStateTypeEnum::Chaos_Object_Dynamic);
+		TSharedPtr<FGeometryDynamicCollection> DynamicCollection = nullptr;
 
-		auto InitFunc = [&RestCollection, &DynamicCollection, &PhysicalMaterial](FSimulationParameters& InParams)
+		// FTransform::Identity and FVector(0.0) are defaults - they won't be used because RestCollection is already initialized. 
+		InitCollectionsParameters InitParams = { FTransform::Identity, FVector(0.0), nullptr, (uint8)EObjectStateTypeEnum::Chaos_Object_Dynamic };
+		InitCollections(PhysicalMaterial, RestCollection, DynamicCollection, InitParams);
+		
+		auto CustomFunc = [&RestCollection, &DynamicCollection, &PhysicalMaterial](FSimulationParameters& InParams)
 		{
-			InParams.RestCollection = RestCollection.Get();
-			InParams.DynamicCollection = DynamicCollection.Get();
-			InParams.PhysicalMaterial  = MakeSerializable(PhysicalMaterial);
 			InParams.Shared.SizeSpecificData[0].CollisionType = ECollisionTypeEnum::Chaos_Surface_Volumetric;
 			InParams.Shared.SizeSpecificData[0].ImplicitType = EImplicitTypeEnum::Chaos_Implicit_Box;
 			InParams.ClusterConnectionMethod = Chaos::FClusterCreationParameters<T>::EConnectionMethod::DelaunayTriangulation;
 			InParams.MaxClusterLevel = 1000;
 			InParams.ClusterGroupIndex = 0;
 			InParams.DamageThreshold = { 1 };
-			InParams.Simulating = true;
-			Chaos::FErrorReporter ErrorReporter;
-			BuildSimulationData(ErrorReporter, *RestCollection, InParams.Shared);
 		};
 
-
 		FRadialFalloff * FalloffField = new FRadialFalloff();
 		FalloffField->Magnitude = 1.5;
 		FalloffField->Radius = 100.0;
@@ -823,7 +695,9 @@ namespace GeometryCollectionExample
 
 		Chaos::FPBDRigidsSolver* Solver = FChaosSolversModule::GetModule()->CreateSolver(true);
 		FFieldSystemPhysicsProxy* FieldObject = new FFieldSystemPhysicsProxy(nullptr);
-		FGeometryCollectionPhysicsProxy* PhysObject = new FGeometryCollectionPhysicsProxy(nullptr, DynamicCollection.Get(), InitFunc, nullptr, nullptr);;
+
+		FGeometryCollectionPhysicsProxy* PhysObject = RigidBodySetup(PhysicalMaterial, RestCollection, DynamicCollection, CustomFunc);
+
 #if TODO_REIMPLEMENT_RIGID_CLUSTERING
 		Chaos::FPBDRigidsSolver::FClusteringType & Clustering = Solver->GetRigidClustering();
 		const Chaos::FPBDRigidsSolver::FClusteringType::FClusterMap & ClusterMap = Clustering.GetChildrenMap();
@@ -831,7 +705,6 @@ namespace GeometryCollectionExample
 #if TODO_REIMPLEMENT_GET_RIGID_PARTICLES
 		const Chaos::TPBDRigidParticles<float, 3>& Particles = Solver->GetRigidParticles();
 
-		PhysObject->Initialize();
 #if CHAOS_PARTICLEHANDLE_TODO
 		Solver->RegisterObject(PhysObject);
 		Solver->RegisterObject(FieldObject);
@@ -847,32 +720,32 @@ namespace GeometryCollectionExample
 			FName TargetName = GetFieldPhysicsName(EFieldPhysicsType::Field_ExternalClusterStrain);
 			FieldObject->BufferCommand(Solver, { TargetName, FalloffField->NewCopy() });
 
-			R.ExpectTrue(Particles.Disabled(0));
-			R.ExpectTrue(Particles.Disabled(1));
-			R.ExpectTrue(Particles.Disabled(2));
-			R.ExpectTrue(Particles.Disabled(3));
-			R.ExpectTrue(Particles.Disabled(4));
-			R.ExpectTrue(Particles.Disabled(5));
-			R.ExpectTrue(!Particles.Disabled(6));
+			EXPECT_TRUE(Particles.Disabled(0));
+			EXPECT_TRUE(Particles.Disabled(1));
+			EXPECT_TRUE(Particles.Disabled(2));
+			EXPECT_TRUE(Particles.Disabled(3));
+			EXPECT_TRUE(Particles.Disabled(4));
+			EXPECT_TRUE(Particles.Disabled(5));
+			EXPECT_FALSE(Particles.Disabled(6));
 
 			Solver->AdvanceSolverBy(1 / 24.);
 			FinalizeSolver(*Solver);
 
-			R.ExpectTrue(ClusterMap.Num() == 2);
-			R.ExpectTrue(ClusterMap[4]->Num() == 2);
-			R.ExpectTrue(ClusterMap[4]->Contains(2));
-			R.ExpectTrue(ClusterMap[4]->Contains(3));
-			R.ExpectTrue(ClusterMap[5]->Num() == 2);
-			R.ExpectTrue(ClusterMap[5]->Contains(0));
-			R.ExpectTrue(ClusterMap[5]->Contains(1));
-
-			R.ExpectTrue(Particles.Disabled(0));
-			R.ExpectTrue(Particles.Disabled(1));
-			R.ExpectTrue(Particles.Disabled(2));
-			R.ExpectTrue(Particles.Disabled(3));
-			R.ExpectTrue(!Particles.Disabled(4));
-			R.ExpectTrue(!Particles.Disabled(5));
-			R.ExpectTrue(Particles.Disabled(6));
+			EXPECT_EQ(ClusterMap.Num(), 2);
+			EXPECT_EQ(ClusterMap[4]->Num(), 2);
+			EXPECT_TRUE(ClusterMap[4]->Contains(2));
+			EXPECT_TRUE(ClusterMap[4]->Contains(3));
+			EXPECT_EQ(ClusterMap[5]->Num(), 2);
+			EXPECT_TRUE(ClusterMap[5]->Contains(0));
+			EXPECT_TRUE(ClusterMap[5]->Contains(1));
+
+			EXPECT_TRUE(Particles.Disabled(0));
+			EXPECT_TRUE(Particles.Disabled(1));
+			EXPECT_TRUE(Particles.Disabled(2));
+			EXPECT_TRUE(Particles.Disabled(3));
+			EXPECT_FALSE(Particles.Disabled(4));
+			EXPECT_FALSE(Particles.Disabled(5));
+			EXPECT_TRUE(Particles.Disabled(6));
 		}
 #endif
 
@@ -880,46 +753,34 @@ namespace GeometryCollectionExample
 
 		delete PhysObject;
 		delete FalloffField;
-
+		delete FieldObject;
 #endif
-
-		return !R.HasError();
 	}
-	template bool RigidBodies_Field_ClusterBreak_StrainModel_Test1<float>(ExampleResponse&& R);
+	template void RigidBodies_Field_ClusterBreak_StrainModel_Test1<float>();
 
 
 
 	template<class T>
-	bool RigidBodies_Field_ClusterBreak_StrainModel_Test2(ExampleResponse&& R)
+	void RigidBodies_Field_ClusterBreak_StrainModel_Test2()
 	{
 #if INCLUDE_CHAOS
-		TUniquePtr<Chaos::TChaosPhysicsMaterial<T>> PhysicalMaterial = MakeUnique<Chaos::TChaosPhysicsMaterial<T>>();
-		PhysicalMaterial->Friction = 0;
-		PhysicalMaterial->Restitution = 0;
-		PhysicalMaterial->SleepingLinearThreshold = 0;
-		PhysicalMaterial->SleepingAngularThreshold = 0;
-		PhysicalMaterial->DisabledLinearThreshold = 0;
-		PhysicalMaterial->DisabledAngularThreshold = 0;
-
+		TUniquePtr<Chaos::TChaosPhysicsMaterial<T>> PhysicalMaterial = nullptr;
 		TSharedPtr<FGeometryCollection> RestCollection = CreateClusteredBody_ThreeByTwo_ThreeTransform(FVector(0));
-		TSharedPtr<FGeometryDynamicCollection> DynamicCollection = GeometryCollectionToGeometryDynamicCollection(RestCollection.Get(), (uint8)EObjectStateTypeEnum::Chaos_Object_Dynamic);
+		TSharedPtr<FGeometryDynamicCollection> DynamicCollection = nullptr;
+
+		// FTransform::Identity and FVector(0.0) are defaults - they won't be used because RestCollection is already initialized. 
+		InitCollectionsParameters InitParams = { FTransform::Identity, FVector(0.0), nullptr, (uint8)EObjectStateTypeEnum::Chaos_Object_Dynamic };
+		InitCollections(PhysicalMaterial, RestCollection, DynamicCollection, InitParams);
 
-		auto InitFunc = [&RestCollection, &DynamicCollection, &PhysicalMaterial](FSimulationParameters& InParams)
+		auto CustomFunc = [&RestCollection, &DynamicCollection, &PhysicalMaterial](FSimulationParameters& InParams)
 		{
-			InParams.RestCollection = RestCollection.Get();
-			InParams.DynamicCollection = DynamicCollection.Get();
-			InParams.PhysicalMaterial  = MakeSerializable(PhysicalMaterial);
 			InParams.Shared.SizeSpecificData[0].CollisionType = ECollisionTypeEnum::Chaos_Surface_Volumetric;
 			InParams.Shared.SizeSpecificData[0].ImplicitType = EImplicitTypeEnum::Chaos_Implicit_Box;
 			InParams.MaxClusterLevel = 1000;
 			InParams.ClusterGroupIndex = 0;
 			InParams.DamageThreshold = { 1 };
-			InParams.Simulating = true;
-			Chaos::FErrorReporter ErrorReporter;
-			BuildSimulationData(ErrorReporter, *RestCollection, InParams.Shared);
 		};
 
-
 		FRadialFalloff * FalloffField = new FRadialFalloff();
 		FalloffField->Magnitude = 1.5;
 		FalloffField->Radius = 200.0;
@@ -928,7 +789,9 @@ namespace GeometryCollectionExample
 
 		Chaos::FPBDRigidsSolver* Solver = FChaosSolversModule::GetModule()->CreateSolver(true);
 		FFieldSystemPhysicsProxy* FieldObject = new FFieldSystemPhysicsProxy(nullptr);
-		FGeometryCollectionPhysicsProxy* PhysObject = new FGeometryCollectionPhysicsProxy(nullptr, DynamicCollection.Get(), InitFunc, nullptr, nullptr);;
+		
+		FGeometryCollectionPhysicsProxy* PhysObject = RigidBodySetup(PhysicalMaterial, RestCollection, DynamicCollection, CustomFunc);
+
 #if TODO_REIMPLEMENT_RIGID_CLUSTERING
 		Chaos::FPBDRigidsSolver::FClusteringType & Clustering = Solver->GetRigidClustering();
 		const Chaos::FPBDRigidsSolver::FClusteringType::FClusterMap & ClusterMap = Clustering.GetChildrenMap();
@@ -937,7 +800,6 @@ namespace GeometryCollectionExample
 		const Chaos::TPBDRigidParticles<float, 3>& Particles = Solver->GetRigidParticles();
 #endif
 
-		PhysObject->Initialize();
 #if CHAOS_PARTICLEHANDLE_TODO
 		Solver->RegisterObject(PhysObject);
 		Solver->RegisterObject(FieldObject);
@@ -957,15 +819,15 @@ namespace GeometryCollectionExample
 			FieldObject->BufferCommand(Solver, Command);
 
 #if TODO_REIMPLEMENT_GET_RIGID_PARTICLES
-			R.ExpectTrue(Particles.Disabled(0));
-			R.ExpectTrue(Particles.Disabled(1));
-			R.ExpectTrue(Particles.Disabled(2));
-			R.ExpectTrue(Particles.Disabled(3));
-			R.ExpectTrue(Particles.Disabled(4));
-			R.ExpectTrue(Particles.Disabled(5));
-			R.ExpectTrue(Particles.Disabled(6));
-			R.ExpectTrue(Particles.Disabled(7));
-			R.ExpectTrue(!Particles.Disabled(8));
+			EXPECT_TRUE(Particles.Disabled(0));
+			EXPECT_TRUE(Particles.Disabled(1));
+			EXPECT_TRUE(Particles.Disabled(2));
+			EXPECT_TRUE(Particles.Disabled(3));
+			EXPECT_TRUE(Particles.Disabled(4));
+			EXPECT_TRUE(Particles.Disabled(5));
+			EXPECT_TRUE(Particles.Disabled(6));
+			EXPECT_TRUE(Particles.Disabled(7));
+			EXPECT_FALSE(Particles.Disabled(8));
 #endif
 
 			Solver->AdvanceSolverBy(1 / 24.);
@@ -975,66 +837,55 @@ namespace GeometryCollectionExample
 			FinalizeSolver(*Solver);
 
 #if TODO_REIMPLEMENT_RIGID_CLUSTERING
-			R.ExpectTrue(ClusterMap.Num() == 1);
-			R.ExpectTrue(ClusterMap[6]->Num() == 3);
-			R.ExpectTrue(ClusterMap[6]->Contains(3));
-			R.ExpectTrue(ClusterMap[6]->Contains(4));
-			R.ExpectTrue(ClusterMap[6]->Contains(5));
+			EXPECT_EQ(ClusterMap.Num(), 1);
+			EXPECT_EQ(ClusterMap[6]->Num(), 3);
+			EXPECT_TRUE(ClusterMap[6]->Contains(3));
+			EXPECT_TRUE(ClusterMap[6]->Contains(4));
+			EXPECT_TRUE(ClusterMap[6]->Contains(5));
 #endif
 
 #if TODO_REIMPLEMENT_GET_RIGID_PARTICLES
-			R.ExpectTrue(!Particles.Disabled(0));
-			R.ExpectTrue(!Particles.Disabled(1));
-			R.ExpectTrue(!Particles.Disabled(2));
-			R.ExpectTrue(Particles.Disabled(3));
-			R.ExpectTrue(Particles.Disabled(4));
-			R.ExpectTrue(Particles.Disabled(5));
-			R.ExpectTrue(!Particles.Disabled(6));
-			R.ExpectTrue(Particles.Disabled(7));
-			R.ExpectTrue(Particles.Disabled(8));
+			EXPECT_FALSE(Particles.Disabled(0));
+			EXPECT_FALSE(Particles.Disabled(1));
+			EXPECT_FALSE(Particles.Disabled(2));
+			EXPECT_TRUE(Particles.Disabled(3));
+			EXPECT_TRUE(Particles.Disabled(4));
+			EXPECT_TRUE(Particles.Disabled(5));
+			EXPECT_FALSE(Particles.Disabled(6));
+			EXPECT_TRUE(Particles.Disabled(7));
+			EXPECT_TRUE(Particles.Disabled(8));
 #endif
 		}
 
 		FChaosSolversModule::GetModule()->DestroySolver(Solver);
 		delete PhysObject;
 		delete FalloffField;
+		delete FieldObject;
 #endif
-
-		return !R.HasError();
 	}
-	template bool RigidBodies_Field_ClusterBreak_StrainModel_Test2<float>(ExampleResponse&& R);
+	template void RigidBodies_Field_ClusterBreak_StrainModel_Test2<float>();
 
 	template<class T>
-	bool RigidBodies_Field_ClusterBreak_StrainModel_Test3(ExampleResponse&& R)
+	void RigidBodies_Field_ClusterBreak_StrainModel_Test3()
 	{
 #if INCLUDE_CHAOS
-		TUniquePtr<Chaos::TChaosPhysicsMaterial<T>> PhysicalMaterial = MakeUnique<Chaos::TChaosPhysicsMaterial<T>>();
-		PhysicalMaterial->Friction = 0;
-		PhysicalMaterial->Restitution = 0;
-		PhysicalMaterial->SleepingLinearThreshold = 0;
-		PhysicalMaterial->SleepingAngularThreshold = 0;
-		PhysicalMaterial->DisabledLinearThreshold = 0;
-		PhysicalMaterial->DisabledAngularThreshold = 0;
-
+		TUniquePtr<Chaos::TChaosPhysicsMaterial<T>> PhysicalMaterial = nullptr;
 		TSharedPtr<FGeometryCollection> RestCollection = CreateClusteredBody_ThreeByTwo_ThreeTransform(FVector(0));
-		TSharedPtr<FGeometryDynamicCollection> DynamicCollection = GeometryCollectionToGeometryDynamicCollection(RestCollection.Get(), (uint8)EObjectStateTypeEnum::Chaos_Object_Dynamic);
+		TSharedPtr<FGeometryDynamicCollection> DynamicCollection = nullptr;
 
-		auto InitFunc = [&RestCollection, &DynamicCollection, &PhysicalMaterial](FSimulationParameters& InParams)
+		// FTransform::Identity and FVector(0.0) are defaults - they won't be used because RestCollection is already initialized. 
+		InitCollectionsParameters InitParams = { FTransform::Identity, FVector(0.0), nullptr, (uint8)EObjectStateTypeEnum::Chaos_Object_Dynamic };
+		InitCollections(PhysicalMaterial, RestCollection, DynamicCollection, InitParams);
+
+		auto CustomFunc = [&RestCollection, &DynamicCollection, &PhysicalMaterial](FSimulationParameters& InParams)
 		{
-			InParams.RestCollection = RestCollection.Get();
-			InParams.DynamicCollection = DynamicCollection.Get();
-			InParams.PhysicalMaterial  = MakeSerializable(PhysicalMaterial);
 			InParams.Shared.SizeSpecificData[0].CollisionType = ECollisionTypeEnum::Chaos_Surface_Volumetric;
 			InParams.Shared.SizeSpecificData[0].ImplicitType = EImplicitTypeEnum::Chaos_Implicit_Box;
 			InParams.MaxClusterLevel = 1000;
 			InParams.ClusterGroupIndex = 0;
 			InParams.DamageThreshold = { 1 };
-			InParams.Simulating = true;
-			Chaos::FErrorReporter ErrorReporter;
-			BuildSimulationData(ErrorReporter, *RestCollection, InParams.Shared);
 		};
 
-
 		FRadialFalloff * FalloffField = new FRadialFalloff();
 		FalloffField->Magnitude = 1.1;
 		FalloffField->Radius = 200.0;
@@ -1043,7 +894,9 @@ namespace GeometryCollectionExample
 
 		Chaos::FPBDRigidsSolver* Solver = FChaosSolversModule::GetModule()->CreateSolver(true);
 		FFieldSystemPhysicsProxy* FieldObject = new FFieldSystemPhysicsProxy(nullptr);
-		FGeometryCollectionPhysicsProxy* PhysObject = new FGeometryCollectionPhysicsProxy(nullptr, DynamicCollection.Get(), InitFunc, nullptr, nullptr);;
+		
+		FGeometryCollectionPhysicsProxy* PhysObject = RigidBodySetup(PhysicalMaterial, RestCollection, DynamicCollection, CustomFunc);
+
 #if TODO_REIMPLEMENT_RIGID_CLUSTERING
 		Chaos::FPBDRigidsSolver::FClusteringType & Clustering = Solver->GetRigidClustering();
 		const Chaos::FPBDRigidsSolver::FClusteringType::FClusterMap & ClusterMap = Clustering.GetChildrenMap();
@@ -1051,8 +904,6 @@ namespace GeometryCollectionExample
 #if TODO_REIMPLEMENT_GET_RIGID_PARTICLES
 		const Chaos::TPBDRigidParticles<float, 3>& Particles = Solver->GetRigidParticles();
 #endif
-
-		PhysObject->Initialize();
 #if CHAOS_PARTICLEHANDLE_TODO
 		Solver->RegisterObject(PhysObject);
 		Solver->RegisterObject(FieldObject);
@@ -1073,86 +924,75 @@ namespace GeometryCollectionExample
 			FieldObject->BufferCommand(Solver, Command);
 
 #if TODO_REIMPLEMENT_GET_RIGID_PARTICLES
-			R.ExpectTrue(Particles.Disabled(0));
-			R.ExpectTrue(Particles.Disabled(1));
-			R.ExpectTrue(Particles.Disabled(2));
-			R.ExpectTrue(Particles.Disabled(3));
-			R.ExpectTrue(Particles.Disabled(4));
-			R.ExpectTrue(Particles.Disabled(5));
-			R.ExpectTrue(Particles.Disabled(6));
-			R.ExpectTrue(Particles.Disabled(7));
-			R.ExpectTrue(!Particles.Disabled(8));
+			EXPECT_TRUE(Particles.Disabled(0));
+			EXPECT_TRUE(Particles.Disabled(1));
+			EXPECT_TRUE(Particles.Disabled(2));
+			EXPECT_TRUE(Particles.Disabled(3));
+			EXPECT_TRUE(Particles.Disabled(4));
+			EXPECT_TRUE(Particles.Disabled(5));
+			EXPECT_TRUE(Particles.Disabled(6));
+			EXPECT_TRUE(Particles.Disabled(7));
+			EXPECT_FALSE(Particles.Disabled(8));
 #endif
 
 			Solver->AdvanceSolverBy(1 / 24.);
 			FinalizeSolver(*Solver);
 
 #if TODO_REIMPLEMENT_RIGID_CLUSTERING
-			R.ExpectTrue(ClusterMap.Num() == 2);
-			R.ExpectTrue(ClusterMap[6]->Num() == 3);
-			R.ExpectTrue(ClusterMap[6]->Contains(3));
-			R.ExpectTrue(ClusterMap[6]->Contains(4));
-			R.ExpectTrue(ClusterMap[6]->Contains(5));
-			R.ExpectTrue(ClusterMap[7]->Num() == 3);
-			R.ExpectTrue(ClusterMap[7]->Contains(0));
-			R.ExpectTrue(ClusterMap[7]->Contains(1));
-			R.ExpectTrue(ClusterMap[7]->Contains(2));
+			EXPECT_EQ(ClusterMap.Num(), 2);
+			EXPECT_EQ(ClusterMap[6]->Num(), 3);
+			EXPECT_TRUE(ClusterMap[6]->Contains(3));
+			EXPECT_TRUE(ClusterMap[6]->Contains(4));
+			EXPECT_TRUE(ClusterMap[6]->Contains(5));
+			EXPECT_EQ(ClusterMap[7]->Num(), 3);
+			EXPECT_TRUE(ClusterMap[7]->Contains(0));
+			EXPECT_TRUE(ClusterMap[7]->Contains(1));
+			EXPECT_TRUE(ClusterMap[7]->Contains(2));
 #endif
 
 #if TODO_REIMPLEMENT_GET_RIGID_PARTICLES
-			R.ExpectTrue(Particles.Disabled(0));
-			R.ExpectTrue(Particles.Disabled(1));
-			R.ExpectTrue(Particles.Disabled(2));
-			R.ExpectTrue(Particles.Disabled(3));
-			R.ExpectTrue(Particles.Disabled(4));
-			R.ExpectTrue(Particles.Disabled(5));
-			R.ExpectTrue(!Particles.Disabled(6));
-			R.ExpectTrue(!Particles.Disabled(7));
-			R.ExpectTrue(Particles.Disabled(8));
+			EXPECT_TRUE(Particles.Disabled(0));
+			EXPECT_TRUE(Particles.Disabled(1));
+			EXPECT_TRUE(Particles.Disabled(2));
+			EXPECT_TRUE(Particles.Disabled(3));
+			EXPECT_TRUE(Particles.Disabled(4));
+			EXPECT_TRUE(Particles.Disabled(5));
+			EXPECT_FALSE(Particles.Disabled(6));
+			EXPECT_FALSE(Particles.Disabled(7));
+			EXPECT_TRUE(Particles.Disabled(8));
 #endif
 		}
 
 		FChaosSolversModule::GetModule()->DestroySolver(Solver);
 		delete PhysObject;
 		delete FalloffField;
+		delete FieldObject;
 #endif
-
-		return !R.HasError();
 	}
-	template bool RigidBodies_Field_ClusterBreak_StrainModel_Test3<float>(ExampleResponse&& R);
+	template void RigidBodies_Field_ClusterBreak_StrainModel_Test3<float>();
 
 
 	template<class T>
-	bool RigidBodies_Field_ClusterBreak_StrainModel_Test4(ExampleResponse&& R)
+	void RigidBodies_Field_ClusterBreak_StrainModel_Test4()
 	{
 #if INCLUDE_CHAOS
-		TUniquePtr<Chaos::TChaosPhysicsMaterial<T>> PhysicalMaterial = MakeUnique<Chaos::TChaosPhysicsMaterial<T>>();
-		PhysicalMaterial->Friction = 0;
-		PhysicalMaterial->Restitution = 0;
-		PhysicalMaterial->SleepingLinearThreshold = 0;
-		PhysicalMaterial->SleepingAngularThreshold = 0;
-		PhysicalMaterial->DisabledLinearThreshold = 0;
-		PhysicalMaterial->DisabledAngularThreshold = 0;
-
+		TUniquePtr<Chaos::TChaosPhysicsMaterial<T>> PhysicalMaterial = nullptr;
 		TSharedPtr<FGeometryCollection> RestCollection = CreateClusteredBody_TwoByTwo_ThreeTransform(FVector(0));
-		TSharedPtr<FGeometryDynamicCollection> DynamicCollection = GeometryCollectionToGeometryDynamicCollection(RestCollection.Get(), (uint8)EObjectStateTypeEnum::Chaos_Object_Dynamic);
+		TSharedPtr<FGeometryDynamicCollection> DynamicCollection = nullptr;
+
+		// FTransform::Identity and FVector(0.0) are defaults - they won't be used because RestCollection is already initialized. 
+		InitCollectionsParameters InitParams = { FTransform::Identity, FVector(0.0), nullptr, (uint8)EObjectStateTypeEnum::Chaos_Object_Dynamic };
+		InitCollections(PhysicalMaterial, RestCollection, DynamicCollection, InitParams);
 
-		auto InitFunc = [&RestCollection, &DynamicCollection, &PhysicalMaterial](FSimulationParameters& InParams)
+		auto CustomFunc = [&RestCollection, &DynamicCollection, &PhysicalMaterial](FSimulationParameters& InParams)
 		{
-			InParams.RestCollection = RestCollection.Get();
-			InParams.DynamicCollection = DynamicCollection.Get();
-			InParams.PhysicalMaterial = MakeSerializable(PhysicalMaterial);
 			InParams.Shared.SizeSpecificData[0].CollisionType = ECollisionTypeEnum::Chaos_Surface_Volumetric;
 			InParams.Shared.SizeSpecificData[0].ImplicitType = EImplicitTypeEnum::Chaos_Implicit_Box;
 			InParams.MaxClusterLevel = 1000;
 			InParams.ClusterGroupIndex = 0;
 			InParams.DamageThreshold = { 1 };
-			InParams.Simulating = true;
-			Chaos::FErrorReporter ErrorReporter;
-			BuildSimulationData(ErrorReporter, *RestCollection, InParams.Shared);
 		};
 
-
 		FRadialFalloff * FalloffField = new FRadialFalloff();
 		FalloffField->Magnitude = 1.5;
 		FalloffField->Radius = 100.0;
@@ -1161,7 +1001,9 @@ namespace GeometryCollectionExample
 
 		Chaos::FPBDRigidsSolver* Solver = FChaosSolversModule::GetModule()->CreateSolver(true);
 		FFieldSystemPhysicsProxy* FieldObject = new FFieldSystemPhysicsProxy(nullptr);
-		FGeometryCollectionPhysicsProxy* PhysObject = new FGeometryCollectionPhysicsProxy(nullptr, DynamicCollection.Get(), InitFunc, nullptr, nullptr);;
+
+		FGeometryCollectionPhysicsProxy* PhysObject = RigidBodySetup(PhysicalMaterial, RestCollection, DynamicCollection, CustomFunc);
+
 #if TODO_REIMPLEMENT_RIGID_CLUSTERING
 		Chaos::FPBDRigidsSolver::FClusteringType & Clustering = Solver->GetRigidClustering();
 		const Chaos::FPBDRigidsSolver::FClusteringType::FClusterMap & ClusterMap = Clustering.GetChildrenMap();
@@ -1169,8 +1011,6 @@ namespace GeometryCollectionExample
 #if TODO_REIMPLEMENT_GET_RIGID_PARTICLES
 		const Chaos::TPBDRigidParticles<float, 3>& Particles = Solver->GetRigidParticles();
 #endif
-
-		PhysObject->Initialize();
 #if CHAOS_PARTICLEHANDLE_TODO
 		Solver->RegisterObject(PhysObject);
 		Solver->RegisterObject(FieldObject);
@@ -1187,36 +1027,36 @@ namespace GeometryCollectionExample
 			FieldObject->BufferCommand(Solver, { TargetName, FalloffField->NewCopy() });
 
 #if TODO_REIMPLEMENT_GET_RIGID_PARTICLES
-			R.ExpectTrue(Particles.Disabled(0));
-			R.ExpectTrue(Particles.Disabled(1));
-			R.ExpectTrue(Particles.Disabled(2));
-			R.ExpectTrue(Particles.Disabled(3));
-			R.ExpectTrue(Particles.Disabled(4));
-			R.ExpectTrue(Particles.Disabled(5));
-			R.ExpectTrue(!Particles.Disabled(6));
+			EXPECT_TRUE(Particles.Disabled(0));
+			EXPECT_TRUE(Particles.Disabled(1));
+			EXPECT_TRUE(Particles.Disabled(2));
+			EXPECT_TRUE(Particles.Disabled(3));
+			EXPECT_TRUE(Particles.Disabled(4));
+			EXPECT_TRUE(Particles.Disabled(5));
+			EXPECT_FALSE(Particles.Disabled(6));
 #endif
 
 			Solver->AdvanceSolverBy(1 / 24.);
 			FinalizeSolver(*Solver);
 
 #if TODO_REIMPLEMENT_RIGID_CLUSTERING
-			R.ExpectTrue(ClusterMap.Num() == 2);
-			R.ExpectTrue(ClusterMap[4]->Num() == 2);
-			R.ExpectTrue(ClusterMap[4]->Contains(2));
-			R.ExpectTrue(ClusterMap[4]->Contains(3));
-			R.ExpectTrue(ClusterMap[5]->Num() == 2);
-			R.ExpectTrue(ClusterMap[5]->Contains(0));
-			R.ExpectTrue(ClusterMap[5]->Contains(1));
+			EXPECT_EQ(ClusterMap.Num(), 2);
+			EXPECT_EQ(ClusterMap[4]->Num(), 2);
+			EXPECT_TRUE(ClusterMap[4]->Contains(2));
+			EXPECT_TRUE(ClusterMap[4]->Contains(3));
+			EXPECT_EQ(ClusterMap[5]->Num(), 2);
+			EXPECT_TRUE(ClusterMap[5]->Contains(0));
+			EXPECT_TRUE(ClusterMap[5]->Contains(1));
 #endif
 
 #if TODO_REIMPLEMENT_GET_RIGID_PARTICLES
-			R.ExpectTrue(Particles.Disabled(0));
-			R.ExpectTrue(Particles.Disabled(1));
-			R.ExpectTrue(Particles.Disabled(2));
-			R.ExpectTrue(Particles.Disabled(3));
-			R.ExpectTrue(!Particles.Disabled(4));
-			R.ExpectTrue(!Particles.Disabled(5));
-			R.ExpectTrue(Particles.Disabled(6));
+			EXPECT_TRUE(Particles.Disabled(0));
+			EXPECT_TRUE(Particles.Disabled(1));
+			EXPECT_TRUE(Particles.Disabled(2));
+			EXPECT_TRUE(Particles.Disabled(3));
+			EXPECT_FALSE(Particles.Disabled(4));
+			EXPECT_FALSE(Particles.Disabled(5));
+			EXPECT_TRUE(Particles.Disabled(6));
 #endif
 		}
 
@@ -1224,12 +1064,10 @@ namespace GeometryCollectionExample
 
 		delete PhysObject;
 		delete FalloffField;
+		delete FieldObject;
 #endif
-
-		return !R.HasError();
 	}
-	template bool RigidBodies_Field_ClusterBreak_StrainModel_Test4<float>(ExampleResponse&& R);
+	template void RigidBodies_Field_ClusterBreak_StrainModel_Test4<float>();
 
 
 }
-
diff --git a/Engine/Plugins/Experimental/GeometryCollectionPlugin/Source/GeometryCollectionExampleCore/Private/GeometryCollection/GeometryCollectionExampleSimulationSolver.cpp b/Engine/Plugins/Experimental/GeometryCollectionPlugin/Source/GeometryCollectionExampleCore/Private/GeometryCollection/GeometryCollectionExampleSimulationSolver.cpp
--- a/Engine/Plugins/Experimental/GeometryCollectionPlugin/Source/GeometryCollectionExampleCore/Private/GeometryCollection/GeometryCollectionExampleSimulationSolver.cpp
+++ b/Engine/Plugins/Experimental/GeometryCollectionPlugin/Source/GeometryCollectionExampleCore/Private/GeometryCollection/GeometryCollectionExampleSimulationSolver.cpp
@@ -35,7 +35,7 @@ namespace GeometryCollectionExample
 {
 
 	template<class T>
-	bool Solver_AdvanceNoObjects(ExampleResponse&& R)
+	void Solver_AdvanceNoObjects()
 	{
 #if INCLUDE_CHAOS
 		Chaos::FPBDRigidsSolver* Solver = FChaosSolversModule::GetModule()->CreateSolver(true);
@@ -46,14 +46,12 @@ namespace GeometryCollectionExample
 
 		FChaosSolversModule::GetModule()->DestroySolver(Solver);
 #endif
-
-		return !R.HasError();
 	}
-	template bool Solver_AdvanceNoObjects<float>(ExampleResponse&& R);
+	template void Solver_AdvanceNoObjects<float>();
 
 
 	template<class T>
-	bool Solver_AdvanceDisabledObjects(ExampleResponse&& R)
+	void Solver_AdvanceDisabledObjects()
 	{
 #if INCLUDE_CHAOS
 		TSharedPtr<FGeometryCollection> RestCollection = GeometryCollection::MakeCubeElement(FTransform::Identity, FVector(1.0));
@@ -86,33 +84,31 @@ namespace GeometryCollectionExample
 
 		// never touched
 		TManagedArray<FTransform>& RestTransform = RestCollection->Transform;
-		R.ExpectTrue(FMath::Abs(RestTransform[0].GetTranslation().Z) < SMALL_THRESHOLD);
+		EXPECT_LT(FMath::Abs(RestTransform[0].GetTranslation().Z), SMALL_THRESHOLD);
 
 		// simulated
 		TManagedArray<FTransform>& Transform = DynamicCollection->Transform;
-		R.ExpectTrue(Transform.Num() == 1);
-		R.ExpectTrue(FMath::Abs(Transform[0].GetTranslation().Z) < SMALL_THRESHOLD);
+		EXPECT_EQ(Transform.Num(), 1);
+		EXPECT_LT(FMath::Abs(Transform[0].GetTranslation().Z), SMALL_THRESHOLD);
 
 		FChaosSolversModule::GetModule()->DestroySolver(Solver);
 
 		delete PhysObject;
-
 #endif
-		return !R.HasError();
 	}
-	template bool Solver_AdvanceDisabledObjects<float>(ExampleResponse&& R);
+	template void Solver_AdvanceDisabledObjects<float>();
 
 	template<class T>
-	bool Solver_AdvanceDisabledClusteredObjects(ExampleResponse&& R)
+	void Solver_AdvanceDisabledClusteredObjects()
 	{
 #if INCLUDE_CHAOS
 		TSharedPtr<FGeometryCollection> RestCollection = GeometryCollection::MakeCubeElement(FTransform(FQuat::MakeFromEuler(FVector(0, 0, 0.)), FVector(0, -10, 10)), FVector(1.0));
 		RestCollection->AppendGeometry(*GeometryCollection::MakeCubeElement(FTransform(FQuat::MakeFromEuler(FVector(0, 0, 0.)), FVector(0, 10, 10)), FVector(1.0)));
-		R.ExpectTrue(RestCollection->Transform.Num() == 2);
+		EXPECT_EQ(RestCollection->Transform.Num(), 2);
 
 
 		FGeometryCollectionClusteringUtility::ClusterAllBonesUnderNewRoot(RestCollection.Get());
-		R.ExpectTrue(RestCollection->Transform.Num() == 3);
+		EXPECT_EQ(RestCollection->Transform.Num(), 3);
 		RestCollection->Transform[2] = FTransform(FQuat::MakeFromEuler(FVector(90.f, 0, 0.)), FVector(0, 0, 40));
 
 		//GeometryCollectionAlgo::PrintParentHierarchy(RestCollection.Get());
@@ -156,24 +152,16 @@ namespace GeometryCollectionExample
 		delete PhysObject;
 
 #endif
-
-
-		return !R.HasError();
 	}
-	template bool Solver_AdvanceDisabledClusteredObjects<float>(ExampleResponse&& R);
+	template void Solver_AdvanceDisabledClusteredObjects<float>();
 
 
 	template<class T>
-	bool Solver_ValidateReverseMapping(ExampleResponse&& R)
+	void Solver_ValidateReverseMapping()
 	{
 #if INCLUDE_CHAOS
 		TUniquePtr<Chaos::TChaosPhysicsMaterial<T>> PhysicalMaterial = MakeUnique<Chaos::TChaosPhysicsMaterial<T>>();
-		PhysicalMaterial->Friction = 0;
-		PhysicalMaterial->Restitution = 0;
-		PhysicalMaterial->SleepingLinearThreshold = 0;
-		PhysicalMaterial->SleepingAngularThreshold = 0;
-		PhysicalMaterial->DisabledLinearThreshold = 0;
-		PhysicalMaterial->DisabledAngularThreshold = 0;
+		InitMaterialToZero(PhysicalMaterial);
 
 		Chaos::FPBDRigidsSolver* Solver = FChaosSolversModule::GetModule()->CreateSolver(true);
 		Solver->SetHasFloor(true);
@@ -183,22 +171,13 @@ namespace GeometryCollectionExample
 
 		for (int32 i = 0; i < 10; i++)
 		{
-			TSharedPtr<FGeometryCollection> RestCollection = GeometryCollection::MakeCubeElement(FTransform::Identity, FVector(1.0));
-			TSharedPtr<FGeometryDynamicCollection> DynamicCollection = GeometryCollectionToGeometryDynamicCollection(RestCollection.Get());
+			TSharedPtr<FGeometryCollection> RestCollection = nullptr;
+			TSharedPtr<FGeometryDynamicCollection> DynamicCollection = nullptr;
 
-			auto InitFunc = [&RestCollection, &DynamicCollection, &PhysicalMaterial](FSimulationParameters& InParams)
-			{
-				InParams.RestCollection = RestCollection.Get();
-				InParams.DynamicCollection = DynamicCollection.Get();
-				InParams.PhysicalMaterial  = MakeSerializable(PhysicalMaterial);
-				InParams.Shared.SizeSpecificData[0].CollisionType = ECollisionTypeEnum::Chaos_Volumetric;
-				InParams.Simulating = true;
-				Chaos::FErrorReporter ErrorReporter;
-				BuildSimulationData(ErrorReporter, *RestCollection, InParams.Shared);
-			};
+			InitCollectionsParameters InitParams = { FTransform::Identity, FVector(1.0), nullptr, (int32)EObjectStateTypeEnum::Chaos_Object_Kinematic };
+			InitCollections(PhysicalMaterial, RestCollection, DynamicCollection, InitParams);
 
-			FGeometryCollectionPhysicsProxy* PhysObject = new FGeometryCollectionPhysicsProxy(nullptr, DynamicCollection.Get(), InitFunc, nullptr, nullptr);;
-			PhysObject->Initialize();
+			FGeometryCollectionPhysicsProxy* PhysObject = RigidBodySetup(PhysicalMaterial, RestCollection, DynamicCollection);
 
 #if CHAOS_PARTICLEHANDLE_TODO
 			Solver->RegisterObject(PhysObject);
@@ -216,42 +195,36 @@ namespace GeometryCollectionExample
 		const Chaos::TArrayCollectionArray<PhysicsProxyWrapper> & SolverMapping = Solver->GetPhysicsProxyReverseMapping();
 		const Chaos::TArrayCollectionArray<int32> & ParticleMapping = Solver->GetParticleIndexReverseMapping();
 
-		R.ExpectTrue(SolverMapping.Num() == 11);
-		R.ExpectTrue(ParticleMapping.Num() == 11);
+		EXPECT_EQ(SolverMapping.Num(), 11);
+		EXPECT_EQ(ParticleMapping.Num(), 11);
 
-		R.ExpectTrue(ParticleMapping[0] == INDEX_NONE);
-		R.ExpectTrue(ParticleMapping[1] == 0);
+		EXPECT_EQ(ParticleMapping[0], INDEX_NONE);
+		EXPECT_EQ(ParticleMapping[1], 0);
 
-		R.ExpectTrue(SolverMapping[0].PhysicsProxy == nullptr);
-		R.ExpectTrue(SolverMapping[0].Type == EPhysicsProxyType::NoneType);
+		EXPECT_EQ(SolverMapping[0].PhysicsProxy, nullptr);
+		EXPECT_EQ(SolverMapping[0].Type, EPhysicsProxyType::NoneType);
 
-		R.ExpectTrue(SolverMapping[5].PhysicsProxy != nullptr);
-		R.ExpectTrue(SolverMapping[5].Type == EPhysicsProxyType::GeometryCollectionType);
+		EXPECT_EQ(SolverMapping[5].PhysicsProxy != nullptr);
+		EXPECT_EQ(SolverMapping[5].Type, EPhysicsProxyType::GeometryCollectionType);
 
 		const TManagedArray<int32> & RigidBodyID = static_cast<FGeometryCollectionPhysicsProxy*>(SolverMapping[5].PhysicsProxy)->
 			GetGeometryDynamicCollection_PhysicsThread()->GetAttribute<int32>("RigidBodyID", FGeometryCollection::TransformGroup);
-		R.ExpectTrue(RigidBodyID.Num() == 1);
-		R.ExpectTrue(RigidBodyID[0] == 5);
+		EXPECT_EQ(RigidBodyID.Num(), 1);
+		EXPECT_EQ(RigidBodyID[0], 5);
 #endif
 
 		FChaosSolversModule::GetModule()->DestroySolver(Solver);
 
 #endif
-		return !R.HasError();
 	}
-	template bool Solver_ValidateReverseMapping<float>(ExampleResponse&& R);
+	template void Solver_ValidateReverseMapping<float>();
 
 #if INCLUDE_CHAOS
 	template<class T>
 	void CommonInit(int32 NumObjects, bool UseClusters, Chaos::FPBDRigidsSolver** SolverInOut, TUniquePtr<Chaos::TChaosPhysicsMaterial<T>>& PhysicalMaterial, TArray<TSharedPtr<FGeometryCollection>>& RestArray, TArray<TSharedPtr<FGeometryDynamicCollection>>& DynamicArray)
 	{
 		PhysicalMaterial = MakeUnique<Chaos::TChaosPhysicsMaterial<T>>();
-		PhysicalMaterial->Friction = 0;
-		PhysicalMaterial->Restitution = 0;
-		PhysicalMaterial->SleepingLinearThreshold = 0;
-		PhysicalMaterial->SleepingAngularThreshold = 0;
-		PhysicalMaterial->DisabledLinearThreshold = 0;
-		PhysicalMaterial->DisabledAngularThreshold = 0;
+		InitMaterialToZero(PhysicalMaterial);
 
 		*SolverInOut = FChaosSolversModule::GetModule()->CreateSolver(true);
 		Chaos::FPBDRigidsSolver* Solver = *SolverInOut;
@@ -279,21 +252,12 @@ namespace GeometryCollectionExample
 			
 			TSharedPtr<FGeometryDynamicCollection> DynamicCollection = GeometryCollectionToGeometryDynamicCollection(RestCollection.Get());
 
-			auto InitFunc = [&RestCollection, &DynamicCollection, &PhysicalMaterial](FSimulationParameters& InParams)
+			auto CustomFunc = [&RestCollection, &DynamicCollection, &PhysicalMaterial](FSimulationParameters& InParams)
 			{
-				InParams.RestCollection = RestCollection.Get();
-				InParams.DynamicCollection = DynamicCollection.Get();
-				InParams.PhysicalMaterial = MakeSerializable(PhysicalMaterial);
-				InParams.Shared.SizeSpecificData[0].CollisionType = ECollisionTypeEnum::Chaos_Volumetric;
 				InParams.Shared.SizeSpecificData[0].ImplicitType = EImplicitTypeEnum::Chaos_Implicit_Box;
-				InParams.DamageThreshold = { 0.1f };
-				InParams.Simulating = true;
-				Chaos::FErrorReporter ErrorReporter;
-				BuildSimulationData(ErrorReporter, *RestCollection, InParams.Shared);
 			};
 
-			FGeometryCollectionPhysicsProxy* PhysObject = new FGeometryCollectionPhysicsProxy(nullptr, DynamicCollection.Get(), InitFunc, nullptr, nullptr);;
-			PhysObject->Initialize();
+			FGeometryCollectionPhysicsProxy* PhysObject = RigidBodySetup(PhysicalMaterial, RestCollection, DynamicCollection, CustomFunc);
 			PhysObject->SetCollisionParticlesPerObjectFraction(1.0);
 
 #if CHAOS_PARTICLEHANDLE_TODO
@@ -338,7 +302,7 @@ namespace GeometryCollectionExample
 	};
 
 	template<class T>
-	bool Solver_CollisionEventFilter(ExampleResponse&& R)
+	void Solver_CollisionEventFilter()
 	{
 #if INCLUDE_CHAOS
 
@@ -383,15 +347,15 @@ namespace GeometryCollectionExample
 			if (Impact)
 			{
 				// any objects with a mass of less than 6 are removed from returned collision data
-				R.ExpectTrue(AllCollisionsArray.Num() == 4);
+				EXPECT_EQ(AllCollisionsArray.Num(), 4);
 
 				for (const auto& Collision : AllCollisionsArray)
 				{ 
-					R.ExpectTrue(Particles.M(Collision.ParticleIndex) >= TestMassThreshold);
-					R.ExpectTrue(Collision.Mass1 >= TestMassThreshold);
-					R.ExpectTrue(Collision.Velocity1.Z < 0.0f);
-					R.ExpectTrue(Collision.Mass2 == 0.0f);
-					R.ExpectTrue(Collision.Velocity2.Z == 0.0f);
+					EXPECT_GE(Particles.M(Collision.ParticleIndex), TestMassThreshold);
+					EXPECT_GE(Collision.Mass1, TestMassThreshold);
+					EXPECT_LT(Collision.Velocity1.Z, 0.0f);
+					EXPECT_EQ(Collision.Mass2, 0.0f);
+					EXPECT_EQ(Collision.Velocity2.Z, 0.0f);
 				}
 			}
 		} while (!Impact);
@@ -400,12 +364,11 @@ namespace GeometryCollectionExample
 		FChaosSolversModule::GetModule()->DestroySolver(Solver);
 #endif
 
-		return !R.HasError();
 	}
-	template bool Solver_CollisionEventFilter<float>(ExampleResponse&& R);
+	template void Solver_CollisionEventFilter<float>();
 
 	template<class T>
-	bool Solver_BreakingEventFilter(ExampleResponse&& R)
+	void Solver_BreakingEventFilter()
 	{
 #if INCLUDE_CHAOS
 
@@ -451,20 +414,20 @@ namespace GeometryCollectionExample
 
 			if (Impact)
 			{
-				R.ExpectTrue(Particles.Disabled(0) == false); // ground
-				R.ExpectTrue(Particles.Disabled(1) == false); // piece1 active 6 mass
-				R.ExpectTrue(Particles.Disabled(2) == false); // piece2 active 0.5 mass
-				R.ExpectTrue(Particles.Disabled(3) == false); // piece3 active 0.5 mass
-				R.ExpectTrue(Particles.Disabled(4) == false); // cluster active 7 mass
-				R.ExpectTrue(Particles.Disabled(5) == true); // cluster
+				EXPECT_EQ(Particles.Disabled(0), false); // ground
+				EXPECT_EQ(Particles.Disabled(1), false); // piece1 active 6 mass
+				EXPECT_EQ(Particles.Disabled(2), false); // piece2 active 0.5 mass
+				EXPECT_EQ(Particles.Disabled(3), false); // piece3 active 0.5 mass
+				EXPECT_EQ(Particles.Disabled(4), false); // cluster active 7 mass
+				EXPECT_EQ(Particles.Disabled(5), true); // cluster
 
 				// breaking data
-				R.ExpectTrue(AllBreakingsArray.Num() == 2); // 2 pieces filtered out of 4
+				EXPECT_EQ(AllBreakingsArray.Num(), 2); // 2 pieces filtered out of 4
 
-				R.ExpectTrue(AllBreakingsArray[0].ParticleIndex == 4);
-				R.ExpectTrue(AllBreakingsArray[0].Mass == TestMass + 2.0f);
-				R.ExpectTrue(AllBreakingsArray[1].ParticleIndex == 1);
-				R.ExpectTrue(AllBreakingsArray[1].Mass == TestMass + 1.0f);
+				EXPECT_EQ(AllBreakingsArray[0].ParticleIndex, 4);
+				EXPECT_EQ(AllBreakingsArray[0].Mass, TestMass + 2.0f);
+				EXPECT_EQ(AllBreakingsArray[1].ParticleIndex, 1);
+				EXPECT_EQ(AllBreakingsArray[1].Mass, TestMass + 1.0f);
 
 			}
 		} while (!Impact);
@@ -472,9 +435,8 @@ namespace GeometryCollectionExample
 #endif
 		FChaosSolversModule::GetModule()->DestroySolver(Solver);
 #endif
-		return !R.HasError();
 	}
-	template bool Solver_BreakingEventFilter<float>(ExampleResponse&& R);
+	template void Solver_BreakingEventFilter<float>();
 }
 
 
diff --git a/Engine/Plugins/Experimental/GeometryCollectionPlugin/Source/GeometryCollectionExampleCore/Private/GeometryCollection/GeometryCollectionExampleSimulationStreaming.cpp b/Engine/Plugins/Experimental/GeometryCollectionPlugin/Source/GeometryCollectionExampleCore/Private/GeometryCollection/GeometryCollectionExampleSimulationStreaming.cpp
--- a/Engine/Plugins/Experimental/GeometryCollectionPlugin/Source/GeometryCollectionExampleCore/Private/GeometryCollection/GeometryCollectionExampleSimulationStreaming.cpp
+++ b/Engine/Plugins/Experimental/GeometryCollectionPlugin/Source/GeometryCollectionExampleCore/Private/GeometryCollection/GeometryCollectionExampleSimulationStreaming.cpp
@@ -24,7 +24,7 @@
 namespace GeometryCollectionExample
 {
 	template<class T>
-	bool RigidBodies_Streaming_StartSolverEmpty(ExampleResponse&& R)
+	void RigidBodies_Streaming_StartSolverEmpty()
 	{
 #if INCLUDE_CHAOS
 		Chaos::FPBDRigidsSolver* Solver = FChaosSolversModule::GetModule()->CreateSolver(true);
@@ -48,9 +48,9 @@ namespace GeometryCollectionExample
 			}
 		}
 
-		R.ExpectTrue(Particles.Size() == 9);
+		EXPECT_EQ(Particles.Size(), 9);
 		for (uint32 i = 0; i < Particles.Size() - 1; i++)
-			R.ExpectTrue(Particles.X(i).Z < Particles.X(i + 1).Z);
+			EXPECT_LT(Particles.X(i).Z, Particles.X(i + 1).Z);
 
 		// cleanup
 		for (auto Obj : Collections)
@@ -58,13 +58,12 @@ namespace GeometryCollectionExample
 #endif
 		delete Solver;
 #endif
-		return !R.HasError();
 	}
-	template bool RigidBodies_Streaming_StartSolverEmpty<float>(ExampleResponse&& R);
+	template void RigidBodies_Streaming_StartSolverEmpty<float>();
 
 
 	template<class T>
-	bool RigidBodies_Streaming_BulkInitialization(ExampleResponse&& R)
+	void RigidBodies_Streaming_BulkInitialization()
 	{
 #if INCLUDE_CHAOS
 		typename SimulationObjects<T>::FParameters P;
@@ -96,9 +95,9 @@ namespace GeometryCollectionExample
 			Solver->AdvanceSolverBy(1 / 24.);
 		}
 
-		R.ExpectTrue(Particles.Size() == 9);
+		EXPECT_EQ(Particles.Size(), 9);
 		for (uint32 i = 0; i < Particles.Size() - 1; i++)
-			R.ExpectTrue( FMath::Abs(Particles.X(i).Z - Particles.X(i + 1).Z)<KINDA_SMALL_NUMBER );
+			EXPECT_LT( FMath::Abs(Particles.X(i).Z - Particles.X(i + 1).Z), KINDA_SMALL_NUMBER );
 
 
 		// cleanup
@@ -107,14 +106,13 @@ namespace GeometryCollectionExample
 #endif
 		delete Solver;
 #endif
-		return !R.HasError();
 	}
-	template bool RigidBodies_Streaming_BulkInitialization<float>(ExampleResponse&& R); 
+	template void RigidBodies_Streaming_BulkInitialization<float>(); 
 
 
 
 	template<class T>
-	bool RigidBodies_Streaming_DeferedClusteringInitialization(ExampleResponse&& R)
+	void RigidBodies_Streaming_DeferedClusteringInitialization()
 	{
 #if INCLUDE_CHAOS
 		typename SimulationObjects<T>::FParameters P;
@@ -142,20 +140,20 @@ namespace GeometryCollectionExample
 		}
 
 		// all particles should be disabled
-		R.ExpectTrue(Particles.Size() == 9);
+		EXPECT_EQ(Particles.Size(), 9);
 		for (uint32 i = 0; i < Particles.Size(); i++)
 		{
-			R.ExpectTrue(Particles.Disabled(i) == true);
+			EXPECT_EQ(Particles.Disabled(i), true);
 		}
 
 		for (auto Obj : Collections)
 			Obj->PhysicsProxy->ActivateBodies();
 
 		// all particles should be enabled
-		R.ExpectTrue(Particles.Size() == 9);
+		EXPECT_EQ(Particles.Size(), 9);
 		for (uint32 i = 0; i < Particles.Size(); i++)
 		{
-			R.ExpectTrue(Particles.Disabled(i) == false);
+			EXPECT_EQ(Particles.Disabled(i), false);
 		}
 
 		for (int32 Frame = 1; Frame < 100; Frame++)
@@ -164,13 +162,13 @@ namespace GeometryCollectionExample
 		}
 
 		// new cluster parent should be falling
-		R.ExpectTrue(Particles.Size() == 10);
+		EXPECT_EQ(Particles.Size(), 10);
 		for (uint32 i = 0; i < Particles.Size()-1; i++)
 		{
-			R.ExpectTrue(Particles.Disabled(i) == true);
+			EXPECT_EQ(Particles.Disabled(i), true);
 		}
-		R.ExpectTrue(Particles.Disabled(Particles.Size()-1) == false);
-		R.ExpectTrue(Particles.X(Particles.Size() - 1).Z < -1.f);
+		EXPECT_EQ(Particles.Disabled(Particles.Size()-1), false);
+		EXPECT_LT(Particles.X(Particles.Size() - 1).Z, -1.f);
 
 		// cleanup
 		for (auto Obj : Collections)
@@ -178,9 +176,8 @@ namespace GeometryCollectionExample
 #endif
 		delete Solver;
 #endif
-		return !R.HasError();
 	}
-	template bool RigidBodies_Streaming_DeferedClusteringInitialization<float>(ExampleResponse&& R);
+	template void RigidBodies_Streaming_DeferedClusteringInitialization<float>();
 
 
 }
diff --git a/Engine/Plugins/Experimental/GeometryCollectionPlugin/Source/GeometryCollectionExampleCore/Private/GeometryCollection/GeometryCollectionExampleSpatialHash.cpp b/Engine/Plugins/Experimental/GeometryCollectionPlugin/Source/GeometryCollectionExampleCore/Private/GeometryCollection/GeometryCollectionExampleSpatialHash.cpp
--- a/Engine/Plugins/Experimental/GeometryCollectionPlugin/Source/GeometryCollectionExampleCore/Private/GeometryCollection/GeometryCollectionExampleSpatialHash.cpp
+++ b/Engine/Plugins/Experimental/GeometryCollectionPlugin/Source/GeometryCollectionExampleCore/Private/GeometryCollection/GeometryCollectionExampleSpatialHash.cpp
@@ -1014,7 +1014,7 @@ namespace GeometryCollectionExample
 	};
 
 	template<class T>
-	bool GetClosestPointsTest1(ExampleResponse&& R)
+	void GetClosestPointsTest1()
 	{
 		TSpatialHash<float> SpatialHash(Particles_1000, 18.f);
 		// TSpatialHash<T>::Init() Time is 0.000123
@@ -1024,22 +1024,22 @@ namespace GeometryCollectionExample
 		TArray<int32> ClosestPoints = SpatialHash.GetClosestPoints(Particle_LookUp, 122.0);
 		// TSpatialHash<T>::GetClosestPoints() Time is 0.000096
 
-		R.ExpectTrue(ClosestPoints.Num() == 7);
+		EXPECT_EQ(ClosestPoints.Num(), 7);
 
 		TSet<int32> CPSet(ClosestPoints);
-		R.ExpectTrue(CPSet.Contains(767));
-		R.ExpectTrue(CPSet.Contains(499));
-		R.ExpectTrue(CPSet.Contains(412));
-		R.ExpectTrue(CPSet.Contains(754));
-		R.ExpectTrue(CPSet.Contains(898));
-		R.ExpectTrue(CPSet.Contains(55));
-		R.ExpectTrue(CPSet.Contains(802));
-
-		R.ExpectTrue(!CPSet.Contains(800));
-		R.ExpectTrue(!CPSet.Contains(700));
-		R.ExpectTrue(!CPSet.Contains(1));
-		R.ExpectTrue(!CPSet.Contains(900));
-		R.ExpectTrue(!CPSet.Contains(498));
+		EXPECT_TRUE(CPSet.Contains(767));
+		EXPECT_TRUE(CPSet.Contains(499));
+		EXPECT_TRUE(CPSet.Contains(412));
+		EXPECT_TRUE(CPSet.Contains(754));
+		EXPECT_TRUE(CPSet.Contains(898));
+		EXPECT_TRUE(CPSet.Contains(55));
+		EXPECT_TRUE(CPSet.Contains(802));
+
+		EXPECT_FALSE(CPSet.Contains(800));
+		EXPECT_FALSE(CPSet.Contains(700));
+		EXPECT_FALSE(CPSet.Contains(1));
+		EXPECT_FALSE(CPSet.Contains(900));
+		EXPECT_FALSE(CPSet.Contains(498));
 
 		// ----------------------------------------------------------
 
@@ -1049,35 +1049,35 @@ namespace GeometryCollectionExample
 		TArray<int32> ClosestPoints1 = SpatialHash.GetClosestPoints(Particle_LookUp1, 199.0);
 		// TSpatialHash<T>::GetClosestPoints() Time is 0.000121
 
-		R.ExpectTrue(ClosestPoints1.Num() == 24);
+		EXPECT_EQ(ClosestPoints1.Num(), 24);
 
 		TSet<int32> CPSet1(ClosestPoints1);
-		R.ExpectTrue(CPSet1.Contains(402));
-		R.ExpectTrue(CPSet1.Contains(53));
-		R.ExpectTrue(CPSet1.Contains(862));
-		R.ExpectTrue(CPSet1.Contains(403));
-		R.ExpectTrue(CPSet1.Contains(21));
-		R.ExpectTrue(CPSet1.Contains(844));
-		R.ExpectTrue(CPSet1.Contains(952));
-		R.ExpectTrue(CPSet1.Contains(612));
-		R.ExpectTrue(CPSet1.Contains(908));
-		R.ExpectTrue(CPSet1.Contains(531));
-		R.ExpectTrue(CPSet1.Contains(353));
-		R.ExpectTrue(CPSet1.Contains(915));
-		R.ExpectTrue(CPSet1.Contains(270));
-		R.ExpectTrue(CPSet1.Contains(867));
-		R.ExpectTrue(CPSet1.Contains(551));
-		R.ExpectTrue(CPSet1.Contains(101));
-		R.ExpectTrue(CPSet1.Contains(244));
-		R.ExpectTrue(CPSet1.Contains(399));
-		R.ExpectTrue(CPSet1.Contains(854));
-		R.ExpectTrue(CPSet1.Contains(498));
-		R.ExpectTrue(CPSet1.Contains(462));
-		R.ExpectTrue(CPSet1.Contains(139));
-		R.ExpectTrue(CPSet1.Contains(464));
-		R.ExpectTrue(CPSet1.Contains(963));
-
-		R.ExpectTrue(!CPSet1.Contains(1));
+		EXPECT_TRUE(CPSet1.Contains(402));
+		EXPECT_TRUE(CPSet1.Contains(53));
+		EXPECT_TRUE(CPSet1.Contains(862));
+		EXPECT_TRUE(CPSet1.Contains(403));
+		EXPECT_TRUE(CPSet1.Contains(21));
+		EXPECT_TRUE(CPSet1.Contains(844));
+		EXPECT_TRUE(CPSet1.Contains(952));
+		EXPECT_TRUE(CPSet1.Contains(612));
+		EXPECT_TRUE(CPSet1.Contains(908));
+		EXPECT_TRUE(CPSet1.Contains(531));
+		EXPECT_TRUE(CPSet1.Contains(353));
+		EXPECT_TRUE(CPSet1.Contains(915));
+		EXPECT_TRUE(CPSet1.Contains(270));
+		EXPECT_TRUE(CPSet1.Contains(867));
+		EXPECT_TRUE(CPSet1.Contains(551));
+		EXPECT_TRUE(CPSet1.Contains(101));
+		EXPECT_TRUE(CPSet1.Contains(244));
+		EXPECT_TRUE(CPSet1.Contains(399));
+		EXPECT_TRUE(CPSet1.Contains(854));
+		EXPECT_TRUE(CPSet1.Contains(498));
+		EXPECT_TRUE(CPSet1.Contains(462));
+		EXPECT_TRUE(CPSet1.Contains(139));
+		EXPECT_TRUE(CPSet1.Contains(464));
+		EXPECT_TRUE(CPSet1.Contains(963));
+
+		EXPECT_FALSE(CPSet1.Contains(1));
 
 		// ----------------------------------------------------------
 
@@ -1087,12 +1087,12 @@ namespace GeometryCollectionExample
 		TArray<int32> ClosestPoints2 = SpatialHash.GetClosestPoints(Particle_LookUp2, 52.0);
 		// TSpatialHash<T>::GetClosestPoints() Time is 0.000008
 
-		R.ExpectTrue(ClosestPoints2.Num() == 1);
+		EXPECT_EQ(ClosestPoints2.Num(), 1);
 
 		TSet<int32> CPSet2(ClosestPoints2);
-		R.ExpectTrue(CPSet2.Contains(238));
+		EXPECT_TRUE(CPSet2.Contains(238));
 
-		R.ExpectTrue(!CPSet2.Contains(403));
+		EXPECT_FALSE(CPSet2.Contains(403));
 
 		// ----------------------------------------------------------
 
@@ -1102,23 +1102,21 @@ namespace GeometryCollectionExample
 		TArray<int32> ClosestPoints3 = SpatialHash.GetClosestPoints(Particle_LookUp3, 122.0);
 		// TSpatialHash<T>::GetClosestPoints() Time is 0.000035
 
-		R.ExpectTrue(ClosestPoints3.Num() == 5);
+		EXPECT_EQ(ClosestPoints3.Num(), 5);
 
 		TSet<int32> CPSet3(ClosestPoints3);
-		R.ExpectTrue(CPSet3.Contains(229));
-		R.ExpectTrue(CPSet3.Contains(78));
-		R.ExpectTrue(CPSet3.Contains(379));
-		R.ExpectTrue(CPSet3.Contains(843));
-		R.ExpectTrue(CPSet3.Contains(700));
+		EXPECT_TRUE(CPSet3.Contains(229));
+		EXPECT_TRUE(CPSet3.Contains(78));
+		EXPECT_TRUE(CPSet3.Contains(379));
+		EXPECT_TRUE(CPSet3.Contains(843));
+		EXPECT_TRUE(CPSet3.Contains(700));
 
-		R.ExpectTrue(!CPSet3.Contains(1));
-
-		return !R.HasError();
+		EXPECT_FALSE(CPSet3.Contains(1));
 	}
-	template bool GetClosestPointsTest1<float>(ExampleResponse&& R);
+	template void GetClosestPointsTest1<float>();
 
 	template<class T>
-	bool GetClosestPointsTest2(ExampleResponse&& R)
+	void GetClosestPointsTest2()
 	{
 		TSpatialHash<float> SpatialHash(Particles_1000);
 		// TSpatialHash<T>::Init() Time is 0.000080
@@ -1128,19 +1126,19 @@ namespace GeometryCollectionExample
 		TArray<int32> ClosestPoints = SpatialHash.GetClosestPoints(Particle_LookUp, 122.0);
 		// TSpatialHash<T>::GetClosestPoints() Time is 0.000039
 
-		R.ExpectTrue(ClosestPoints.Num() == 7);
+		EXPECT_EQ(ClosestPoints.Num(), 7);
 
 		TSet<int32> CPSet(ClosestPoints);
-		R.ExpectTrue(CPSet.Contains(767));
-		R.ExpectTrue(CPSet.Contains(499));
-		R.ExpectTrue(CPSet.Contains(412));
-		R.ExpectTrue(CPSet.Contains(754));
-		R.ExpectTrue(CPSet.Contains(898));
-		R.ExpectTrue(CPSet.Contains(55));
-		R.ExpectTrue(CPSet.Contains(802));
+		EXPECT_TRUE(CPSet.Contains(767));
+		EXPECT_TRUE(CPSet.Contains(499));
+		EXPECT_TRUE(CPSet.Contains(412));
+		EXPECT_TRUE(CPSet.Contains(754));
+		EXPECT_TRUE(CPSet.Contains(898));
+		EXPECT_TRUE(CPSet.Contains(55));
+		EXPECT_TRUE(CPSet.Contains(802));
 
-		R.ExpectTrue(!CPSet.Contains(900));
-		R.ExpectTrue(!CPSet.Contains(498));
+		EXPECT_FALSE(CPSet.Contains(900));
+		EXPECT_FALSE(CPSet.Contains(498));
 
 		// ----------------------------------------------------------
 
@@ -1150,35 +1148,35 @@ namespace GeometryCollectionExample
 		TArray<int32> ClosestPoints1 = SpatialHash.GetClosestPoints(Particle_LookUp1, 199.0);
 		// TSpatialHash<T>::GetClosestPoints() Time is 0.000046
 
-		R.ExpectTrue(ClosestPoints1.Num() == 24);
+		EXPECT_EQ(ClosestPoints1.Num(), 24);
 
 		TSet<int32> CPSet1(ClosestPoints1);
-		R.ExpectTrue(CPSet1.Contains(402));
-		R.ExpectTrue(CPSet1.Contains(53));
-		R.ExpectTrue(CPSet1.Contains(862));
-		R.ExpectTrue(CPSet1.Contains(403));
-		R.ExpectTrue(CPSet1.Contains(21));
-		R.ExpectTrue(CPSet1.Contains(844));
-		R.ExpectTrue(CPSet1.Contains(952));
-		R.ExpectTrue(CPSet1.Contains(612));
-		R.ExpectTrue(CPSet1.Contains(908));
-		R.ExpectTrue(CPSet1.Contains(531));
-		R.ExpectTrue(CPSet1.Contains(353));
-		R.ExpectTrue(CPSet1.Contains(915));
-		R.ExpectTrue(CPSet1.Contains(270));
-		R.ExpectTrue(CPSet1.Contains(867));
-		R.ExpectTrue(CPSet1.Contains(551));
-		R.ExpectTrue(CPSet1.Contains(101));
-		R.ExpectTrue(CPSet1.Contains(244));
-		R.ExpectTrue(CPSet1.Contains(399));
-		R.ExpectTrue(CPSet1.Contains(854));
-		R.ExpectTrue(CPSet1.Contains(498));
-		R.ExpectTrue(CPSet1.Contains(462));
-		R.ExpectTrue(CPSet1.Contains(139));
-		R.ExpectTrue(CPSet1.Contains(464));
-		R.ExpectTrue(CPSet1.Contains(963));
-
-		R.ExpectTrue(!CPSet1.Contains(1));
+		EXPECT_TRUE(CPSet1.Contains(402));
+		EXPECT_TRUE(CPSet1.Contains(53));
+		EXPECT_TRUE(CPSet1.Contains(862));
+		EXPECT_TRUE(CPSet1.Contains(403));
+		EXPECT_TRUE(CPSet1.Contains(21));
+		EXPECT_TRUE(CPSet1.Contains(844));
+		EXPECT_TRUE(CPSet1.Contains(952));
+		EXPECT_TRUE(CPSet1.Contains(612));
+		EXPECT_TRUE(CPSet1.Contains(908));
+		EXPECT_TRUE(CPSet1.Contains(531));
+		EXPECT_TRUE(CPSet1.Contains(353));
+		EXPECT_TRUE(CPSet1.Contains(915));
+		EXPECT_TRUE(CPSet1.Contains(270));
+		EXPECT_TRUE(CPSet1.Contains(867));
+		EXPECT_TRUE(CPSet1.Contains(551));
+		EXPECT_TRUE(CPSet1.Contains(101));
+		EXPECT_TRUE(CPSet1.Contains(244));
+		EXPECT_TRUE(CPSet1.Contains(399));
+		EXPECT_TRUE(CPSet1.Contains(854));
+		EXPECT_TRUE(CPSet1.Contains(498));
+		EXPECT_TRUE(CPSet1.Contains(462));
+		EXPECT_TRUE(CPSet1.Contains(139));
+		EXPECT_TRUE(CPSet1.Contains(464));
+		EXPECT_TRUE(CPSet1.Contains(963));
+
+		EXPECT_FALSE(CPSet1.Contains(1));
 
 		// ----------------------------------------------------------
 
@@ -1188,12 +1186,12 @@ namespace GeometryCollectionExample
 		TArray<int32> ClosestPoints2 = SpatialHash.GetClosestPoints(Particle_LookUp2, 52.0);
 		// TSpatialHash<T>::GetClosestPoints() Time is 0.000003
 
-		R.ExpectTrue(ClosestPoints2.Num() == 1);
+		EXPECT_EQ(ClosestPoints2.Num(), 1);
 
 		TSet<int32> CPSet2(ClosestPoints2);
-		R.ExpectTrue(CPSet2.Contains(238));
+		EXPECT_TRUE(CPSet2.Contains(238));
 
-		R.ExpectTrue(!CPSet2.Contains(403));
+		EXPECT_FALSE(CPSet2.Contains(403));
 
 		// ----------------------------------------------------------
 
@@ -1203,24 +1201,22 @@ namespace GeometryCollectionExample
 		TArray<int32> ClosestPoints3 = SpatialHash.GetClosestPoints(Particle_LookUp3, 122.0);
 		// TSpatialHash<T>::GetClosestPoints() Time is 0.000020
 
-		R.ExpectTrue(ClosestPoints3.Num() == 5);
+		EXPECT_EQ(ClosestPoints3.Num(), 5);
 
 		TSet<int32> CPSet3(ClosestPoints3);
-		R.ExpectTrue(CPSet3.Contains(229));
-		R.ExpectTrue(CPSet3.Contains(78));
-		R.ExpectTrue(CPSet3.Contains(379));
-		R.ExpectTrue(CPSet3.Contains(843));
-		R.ExpectTrue(CPSet3.Contains(700));
-
-		R.ExpectTrue(!CPSet3.Contains(1));
+		EXPECT_TRUE(CPSet3.Contains(229));
+		EXPECT_TRUE(CPSet3.Contains(78));
+		EXPECT_TRUE(CPSet3.Contains(379));
+		EXPECT_TRUE(CPSet3.Contains(843));
+		EXPECT_TRUE(CPSet3.Contains(700));
 
-		return !R.HasError();
+		EXPECT_FALSE(CPSet3.Contains(1));
 
 	}
-	template bool GetClosestPointsTest2<float>(ExampleResponse&& R);
+	template void GetClosestPointsTest2<float>();
 
 	template<class T>
-	bool GetClosestPointsTest3(ExampleResponse&& R)
+	void GetClosestPointsTest3()
 	{
 		TSpatialHash<float> SpatialHash(Particles_1000);
 		// TSpatialHash<T>::Init() Time is 0.000080
@@ -1230,15 +1226,15 @@ namespace GeometryCollectionExample
 		TArray<int32> ClosestPoints = SpatialHash.GetClosestPoints(Particle_LookUp, 122.0, 12);
 		// TSpatialHash<T>::GetClosestPoints() Time is 0.000050
 
-		R.ExpectTrue(ClosestPoints.Num() == 7);
+		EXPECT_EQ(ClosestPoints.Num(), 7);
 
-		R.ExpectTrue(ClosestPoints[0] == 767);
-		R.ExpectTrue(ClosestPoints[1] == 499);
-		R.ExpectTrue(ClosestPoints[2] == 412);
-		R.ExpectTrue(ClosestPoints[3] == 754);
-		R.ExpectTrue(ClosestPoints[4] == 898);
-		R.ExpectTrue(ClosestPoints[5] == 55);
-		R.ExpectTrue(ClosestPoints[6] == 802);
+		EXPECT_EQ(ClosestPoints[0], 767);
+		EXPECT_EQ(ClosestPoints[1], 499);
+		EXPECT_EQ(ClosestPoints[2], 412);
+		EXPECT_EQ(ClosestPoints[3], 754);
+		EXPECT_EQ(ClosestPoints[4], 898);
+		EXPECT_EQ(ClosestPoints[5], 55);
+		EXPECT_EQ(ClosestPoints[6], 802);
 
 		// ----------------------------------------------------------
 
@@ -1248,28 +1244,28 @@ namespace GeometryCollectionExample
 		TArray<int32> ClosestPoints1 = SpatialHash.GetClosestPoints(Particle_LookUp1, 199.0, 20);
 		// TSpatialHash<T>::GetClosestPoints() Time is 0.000065
 
-		R.ExpectTrue(ClosestPoints1.Num() == 20);
-
-		R.ExpectTrue(ClosestPoints1[0] == 402);
-		R.ExpectTrue(ClosestPoints1[1] == 53);
-		R.ExpectTrue(ClosestPoints1[2] == 862);
-		R.ExpectTrue(ClosestPoints1[3] == 403);
-		R.ExpectTrue(ClosestPoints1[4] == 21);
-		R.ExpectTrue(ClosestPoints1[5] == 844);
-		R.ExpectTrue(ClosestPoints1[6] == 952);
-		R.ExpectTrue(ClosestPoints1[7] == 612);
-		R.ExpectTrue(ClosestPoints1[8] == 908);
-		R.ExpectTrue(ClosestPoints1[9] == 531);
-		R.ExpectTrue(ClosestPoints1[10] == 353);
-		R.ExpectTrue(ClosestPoints1[11] == 915);
-		R.ExpectTrue(ClosestPoints1[12] == 270);
-		R.ExpectTrue(ClosestPoints1[13] == 867);
-		R.ExpectTrue(ClosestPoints1[14] == 551);
-		R.ExpectTrue(ClosestPoints1[15] == 101);
-		R.ExpectTrue(ClosestPoints1[16] == 244);
-		R.ExpectTrue(ClosestPoints1[17] == 399);
-		R.ExpectTrue(ClosestPoints1[18] == 854);
-		R.ExpectTrue(ClosestPoints1[19] == 498);
+		EXPECT_EQ(ClosestPoints1.Num(), 20);
+
+		EXPECT_EQ(ClosestPoints1[0], 402);
+		EXPECT_EQ(ClosestPoints1[1], 53);
+		EXPECT_EQ(ClosestPoints1[2], 862);
+		EXPECT_EQ(ClosestPoints1[3], 403);
+		EXPECT_EQ(ClosestPoints1[4], 21);
+		EXPECT_EQ(ClosestPoints1[5], 844);
+		EXPECT_EQ(ClosestPoints1[6], 952);
+		EXPECT_EQ(ClosestPoints1[7], 612);
+		EXPECT_EQ(ClosestPoints1[8], 908);
+		EXPECT_EQ(ClosestPoints1[9], 531);
+		EXPECT_EQ(ClosestPoints1[10], 353);
+		EXPECT_EQ(ClosestPoints1[11], 915);
+		EXPECT_EQ(ClosestPoints1[12], 270);
+		EXPECT_EQ(ClosestPoints1[13], 867);
+		EXPECT_EQ(ClosestPoints1[14], 551);
+		EXPECT_EQ(ClosestPoints1[15], 101);
+		EXPECT_EQ(ClosestPoints1[16], 244);
+		EXPECT_EQ(ClosestPoints1[17], 399);
+		EXPECT_EQ(ClosestPoints1[18], 854);
+		EXPECT_EQ(ClosestPoints1[19], 498);
 
 		// ----------------------------------------------------------
 
@@ -1279,13 +1275,13 @@ namespace GeometryCollectionExample
 		TArray<int32> ClosestPoints2 = SpatialHash.GetClosestPoints(Particle_LookUp2, 200.0, 5);
 		// TSpatialHash<T>::GetClosestPoints() Time is 0.000048
 
-		R.ExpectTrue(ClosestPoints2.Num() == 5);
+		EXPECT_EQ(ClosestPoints2.Num(), 5);
 
-		R.ExpectTrue(ClosestPoints2[0] == 238);
-		R.ExpectTrue(ClosestPoints2[1] == 807);
-		R.ExpectTrue(ClosestPoints2[2] == 149);
-		R.ExpectTrue(ClosestPoints2[3] == 284);
-		R.ExpectTrue(ClosestPoints2[4] == 245);
+		EXPECT_EQ(ClosestPoints2[0], 238);
+		EXPECT_EQ(ClosestPoints2[1], 807);
+		EXPECT_EQ(ClosestPoints2[2], 149);
+		EXPECT_EQ(ClosestPoints2[3], 284);
+		EXPECT_EQ(ClosestPoints2[4], 245);
 
 		// ----------------------------------------------------------
 
@@ -1295,13 +1291,13 @@ namespace GeometryCollectionExample
 		TArray<int32> ClosestPoints3 = SpatialHash.GetClosestPoints(Particle_LookUp3, 122.0, 12);
 		// TSpatialHash<T>::GetClosestPoints() Time is 0.000024
 
-		R.ExpectTrue(ClosestPoints3.Num() == 5);
+		EXPECT_EQ(ClosestPoints3.Num(), 5);
 
-		R.ExpectTrue(ClosestPoints3[0] == 229);
-		R.ExpectTrue(ClosestPoints3[1] == 78);
-		R.ExpectTrue(ClosestPoints3[2] == 379);
-		R.ExpectTrue(ClosestPoints3[3] == 843);
-		R.ExpectTrue(ClosestPoints3[4] == 700);
+		EXPECT_EQ(ClosestPoints3[0], 229);
+		EXPECT_EQ(ClosestPoints3[1], 78);
+		EXPECT_EQ(ClosestPoints3[2], 379);
+		EXPECT_EQ(ClosestPoints3[3], 843);
+		EXPECT_EQ(ClosestPoints3[4], 700);
 
 		// ----------------------------------------------------------
 
@@ -1311,7 +1307,7 @@ namespace GeometryCollectionExample
 		TArray<int32> ClosestPoints4 = SpatialHash.GetClosestPoints(Particle_LookUp4, 250.0, 100000);
 		// TSpatialHash<T>::GetClosestPoints() Time is 0.000147
 
-		R.ExpectTrue(ClosestPoints4.Num() == 63);
+		EXPECT_EQ(ClosestPoints4.Num(), 63);
 
 		// ----------------------------------------------------------
 
@@ -1321,15 +1317,13 @@ namespace GeometryCollectionExample
 		TArray<int32> ClosestPoints5 = SpatialHash.GetClosestPoints(Particle_LookUp5, 1250.0, 100000);
 		// TSpatialHash<T>::GetClosestPoints() Time is 0.000782
 
-		R.ExpectTrue(ClosestPoints5.Num() == 1000);
-
-		return !R.HasError();
+		EXPECT_EQ(ClosestPoints5.Num(), 1000);
 
 	}
-	template bool GetClosestPointsTest3<float>(ExampleResponse&& R);
+	template void GetClosestPointsTest3<float>();
 
 	template<class T>
-	bool GetClosestPointTest(ExampleResponse&& R)
+	void GetClosestPointTest()
 	{
 		TSpatialHash<float> SpatialHash(Particles_1000);
 		// TSpatialHash<T>::Init() Time is 0.000087
@@ -1339,7 +1333,7 @@ namespace GeometryCollectionExample
 		int32 ClosestPoint = SpatialHash.GetClosestPoint(Particle_LookUp);
 		// TSpatialHash<T>::GetClosestPoints() Time is 0.000005
 
-		R.ExpectTrue(ClosestPoint == 568);
+		EXPECT_EQ(ClosestPoint, 568);
 
 		// ----------------------------------------------------------
 
@@ -1349,7 +1343,7 @@ namespace GeometryCollectionExample
 		int32 ClosestPoint1 = SpatialHash.GetClosestPoint(Particle_LookUp1);
 		// TSpatialHash<T>::GetClosestPoints() Time is 0.000005
 
-		R.ExpectTrue(ClosestPoint1 == 854);
+		EXPECT_EQ(ClosestPoint1, 854);
 
 		// ----------------------------------------------------------
 
@@ -1359,7 +1353,7 @@ namespace GeometryCollectionExample
 		int32 ClosestPoint2 = SpatialHash.GetClosestPoint(Particle_LookUp2);
 		// TSpatialHash<T>::GetClosestPoints() Time is 0.000005
 
-		R.ExpectTrue(ClosestPoint2 == 500);
+		EXPECT_EQ(ClosestPoint2, 500);
 
 		// ----------------------------------------------------------
 
@@ -1368,7 +1362,7 @@ namespace GeometryCollectionExample
 		int32 ClosestPoint3 = SpatialHash.GetClosestPoint(Particle_LookUp3);
 		// TSpatialHash<T>::GetClosestPoints() Time is 0.000005
 
-		R.ExpectTrue(ClosestPoint3 == 161);
+		EXPECT_EQ(ClosestPoint3, 161);
 
 		// ----------------------------------------------------------
 
@@ -1377,15 +1371,13 @@ namespace GeometryCollectionExample
 		int32 ClosestPoint4 = SpatialHash.GetClosestPoint(Particle_LookUp4);
 		// TSpatialHash<T>::GetClosestPoints() Time is 0.000005
 
-		R.ExpectTrue(ClosestPoint4 == 431);
-
-		return !R.HasError();
+		EXPECT_EQ(ClosestPoint4, 431);
 
 	}
-	template bool GetClosestPointTest<float>(ExampleResponse&& R);
+	template void GetClosestPointTest<float>();
 
 	template<class T>
-	bool HashTableUpdateTest(ExampleResponse&& R)
+	void HashTableUpdateTest()
 	{
 		TSpatialHash<float> SpatialHash(Particles_1000);
 		// TSpatialHash<T>::Init() Time is 0.000087
@@ -1394,13 +1386,13 @@ namespace GeometryCollectionExample
 		int32 ClosestPoint = SpatialHash.GetClosestPoint(Particle_LookUp);
 		// TSpatialHash<T>::GetClosestPoints() Time is 0.000005
 
-		R.ExpectTrue(ClosestPoint == 568);
+		EXPECT_EQ(ClosestPoint, 568);
 
 		SpatialHash.Update(15.0);
 		int32 ClosestPoint1 = SpatialHash.GetClosestPoint(Particle_LookUp);
 		// TSpatialHash<T>::GetClosestPoints() Time is 0.000005
 
-		R.ExpectTrue(ClosestPoint1 == 568);
+		EXPECT_EQ(ClosestPoint1, 568);
 
 		// ----------------------------------------------------------
 
@@ -1411,33 +1403,31 @@ namespace GeometryCollectionExample
 		TArray<int32> ClosestPoints2 = SpatialHash2.GetClosestPoints(Particle_LookUp2, 122.0, 12);
 		// TSpatialHash<T>::GetClosestPoints() Time is 0.000024
 
-		R.ExpectTrue(ClosestPoints2.Num() == 5);
+		EXPECT_EQ(ClosestPoints2.Num(), 5);
 
-		R.ExpectTrue(ClosestPoints2[0] == 229);
-		R.ExpectTrue(ClosestPoints2[1] == 78);
-		R.ExpectTrue(ClosestPoints2[2] == 379);
-		R.ExpectTrue(ClosestPoints2[3] == 843);
-		R.ExpectTrue(ClosestPoints2[4] == 700);
+		EXPECT_EQ(ClosestPoints2[0], 229);
+		EXPECT_EQ(ClosestPoints2[1], 78);
+		EXPECT_EQ(ClosestPoints2[2], 379);
+		EXPECT_EQ(ClosestPoints2[3], 843);
+		EXPECT_EQ(ClosestPoints2[4], 700);
 
 		SpatialHash2.Update(35.0);
 		TArray<int32> ClosestPoints3 = SpatialHash2.GetClosestPoints(Particle_LookUp2, 122.0, 12);
 		// TSpatialHash<T>::GetClosestPoints() Time is 0.000005
 
-		R.ExpectTrue(ClosestPoints3.Num() == 5);
-
-		R.ExpectTrue(ClosestPoints3[0] == 229);
-		R.ExpectTrue(ClosestPoints3[1] == 78);
-		R.ExpectTrue(ClosestPoints3[2] == 379);
-		R.ExpectTrue(ClosestPoints3[3] == 843);
-		R.ExpectTrue(ClosestPoints3[4] == 700);
+		EXPECT_EQ(ClosestPoints3.Num(), 5);
 
-		return !R.HasError();
+		EXPECT_EQ(ClosestPoints3[0], 229);
+		EXPECT_EQ(ClosestPoints3[1], 78);
+		EXPECT_EQ(ClosestPoints3[2], 379);
+		EXPECT_EQ(ClosestPoints3[3], 843);
+		EXPECT_EQ(ClosestPoints3[4], 700);
 
 	}
-	template bool HashTableUpdateTest<float>(ExampleResponse&& R);
+	template void HashTableUpdateTest<float>();
 
 	template<class T>
-	bool HashTablePressureTest(ExampleResponse&& R)
+	void HashTablePressureTest()
 	{
 		const int32 NUM_PARTICLES = 1000000;
 		const float BOUNDARY_MIN = -1000.0;
@@ -1458,12 +1448,10 @@ namespace GeometryCollectionExample
 		TArray<int32> ClosestPoints2 = SpatialHash.GetClosestPoints(Particle_LookUp, 500.0);
 		// TSpatialHash<T>::GetClosestPoints() Time is 0.043953
 
-		R.ExpectTrue(ClosestPoints.Num() == 150);
-
-		return !R.HasError();
+		EXPECT_EQ(ClosestPoints.Num(), 150);
 
 	}
-	template bool HashTablePressureTest<float>(ExampleResponse&& R);
+	template void HashTablePressureTest<float>();
 
 }
 
diff --git a/Engine/Plugins/Experimental/GeometryCollectionPlugin/Source/GeometryCollectionExampleCore/Private/GeometryCollection/GeometryCollectionExampleUtility.cpp b/Engine/Plugins/Experimental/GeometryCollectionPlugin/Source/GeometryCollectionExampleCore/Private/GeometryCollection/GeometryCollectionExampleUtility.cpp
--- a/Engine/Plugins/Experimental/GeometryCollectionPlugin/Source/GeometryCollectionExampleCore/Private/GeometryCollection/GeometryCollectionExampleUtility.cpp
+++ b/Engine/Plugins/Experimental/GeometryCollectionPlugin/Source/GeometryCollectionExampleCore/Private/GeometryCollection/GeometryCollectionExampleUtility.cpp
@@ -7,11 +7,15 @@
 #include "GeometryCollection/GeometryCollection.h"
 #include "GeometryCollection/GeometryCollectionUtility.h"
 #include "GeometryCollection/GeometryCollectionAlgo.h"
+#include "GeometryCollection/GeometryCollectionSimulationCoreTypes.h"
 
 #include "PBDRigidsSolver.h"
 #include "PhysicsProxy/PhysicsProxies.h"
+#include "PhysicsProxy/GeometryCollectionPhysicsProxy.h"
 #include "../Resource/FracturedGeometry.h"
 
+#include "Chaos/ErrorReporter.h"
+
 
 namespace GeometryCollectionExample {
 
@@ -213,5 +217,95 @@ namespace GeometryCollectionExample {
 		return RestCollection;
 	}
 
-}
+	template<class T>
+	void InitMaterialToZero(TUniquePtr<Chaos::TChaosPhysicsMaterial<T>> const &PhysicalMaterial)
+	{
+		PhysicalMaterial->Friction = 0;
+		PhysicalMaterial->Restitution = 0;
+		PhysicalMaterial->SleepingLinearThreshold = 0;
+		PhysicalMaterial->SleepingAngularThreshold = 0;
+		PhysicalMaterial->DisabledLinearThreshold = 0;
+		PhysicalMaterial->DisabledAngularThreshold = 0;
+	}
+	template void InitMaterialToZero(TUniquePtr<Chaos::TChaosPhysicsMaterial<float>> const &PhysicalMaterial);
+
+
+	template<class T>
+	void InitCollections(
+		TUniquePtr<Chaos::TChaosPhysicsMaterial<T>> &PhysicalMaterial,
+		TSharedPtr<FGeometryCollection>& RestCollection,
+		TSharedPtr<FGeometryDynamicCollection>& DynamicCollection,
+		InitCollectionsParameters& InitParams
+	)
+	{
+		// Allow for customized initialization of these objects in the calling function. 
+		if (PhysicalMaterial == nullptr)
+		{
+			PhysicalMaterial = MakeUnique<Chaos::TChaosPhysicsMaterial<T>>();
+			InitMaterialToZero(PhysicalMaterial);
+		}
+
+		if (RestCollection == nullptr)
+		{
+			//Default initialization is a cube of the specified center and size. 
+			RestCollection = GeometryCollection::MakeCubeElement(InitParams.RestCenter, InitParams.RestScale);
+			if (InitParams.RestInitFunc != nullptr)
+			{
+				InitParams.RestInitFunc(RestCollection);
+			}
+		}
+
+		if (DynamicCollection == nullptr)
+		{
+			DynamicCollection = GeometryCollectionToGeometryDynamicCollection(RestCollection.Get(), InitParams.DynamicStateDefault);
+		}
+	}
+	template void InitCollections(
+		TUniquePtr<Chaos::TChaosPhysicsMaterial<float>> &PhysicalMaterial,
+		TSharedPtr<FGeometryCollection>& RestCollection,
+		TSharedPtr<FGeometryDynamicCollection>& DynamicCollection,
+		InitCollectionsParameters& InitParams
+	);
+
+	template<class T>
+	FGeometryCollectionPhysicsProxy* RigidBodySetup(
+		TUniquePtr<Chaos::TChaosPhysicsMaterial<T>> &PhysicalMaterial,
+		TSharedPtr<FGeometryCollection>& RestCollection,
+		TSharedPtr<FGeometryDynamicCollection>& DynamicCollection,
+		FInitFunc CustomFunc
+	)
+	{
+#if INCLUDE_CHAOS
+		auto InitFunc = [&RestCollection, &DynamicCollection, &PhysicalMaterial, &CustomFunc](FSimulationParameters& InParams)
+		{
+			InParams.RestCollection = RestCollection.Get();
+			InParams.DynamicCollection = DynamicCollection.Get();
+			InParams.PhysicalMaterial = MakeSerializable(PhysicalMaterial);
+			InParams.Shared.SizeSpecificData[0].CollisionType = ECollisionTypeEnum::Chaos_Volumetric;
+
+			if (CustomFunc != nullptr)
+			{
+				CustomFunc(InParams);
+			}
+
+			InParams.Simulating = true;
+			Chaos::FErrorReporter ErrorReporter;
+			BuildSimulationData(ErrorReporter, *RestCollection, InParams.Shared);
+		};
+
+		FGeometryCollectionPhysicsProxy* PhysObject = new FGeometryCollectionPhysicsProxy(nullptr, DynamicCollection.Get(), InitFunc, nullptr, nullptr);;
+		PhysObject->Initialize();
+		return PhysObject;
+#else
+		return nullptr;
+#endif
+	}
+	template FGeometryCollectionPhysicsProxy* RigidBodySetup(
+		TUniquePtr<Chaos::TChaosPhysicsMaterial<float>> & PhysicalMaterial,
+		TSharedPtr<FGeometryCollection>& RestCollection,
+		TSharedPtr<FGeometryDynamicCollection>& DynamicCollection,
+		FInitFunc CustomFunc
+	);
+
+} // end namespace GeometryCollectionExample
 
