Remove spare #endif
#rnx

diff --git a/Engine/Plugins/Experimental/GeometryCollectionPlugin/Source/GeometryCollectionExampleCore/Private/GeometryCollection/GeometryCollectionExampleSimulationField.cpp b/Engine/Plugins/Experimental/GeometryCollectionPlugin/Source/GeometryCollectionExampleCore/Private/GeometryCollection/GeometryCollectionExampleSimulationField.cpp
--- a/Engine/Plugins/Experimental/GeometryCollectionPlugin/Source/GeometryCollectionExampleCore/Private/GeometryCollection/GeometryCollectionExampleSimulationField.cpp
+++ b/Engine/Plugins/Experimental/GeometryCollectionPlugin/Source/GeometryCollectionExampleCore/Private/GeometryCollection/GeometryCollectionExampleSimulationField.cpp
@@ -646,7 +646,6 @@ namespace GeometryCollectionExample
 
 		delete PhysObject;
 		delete RadialMask;
-#endif
 	}
 	template void RigidBodies_Field_CollisionGroup<float>();
 
