Remove additional spare #endif
#rnx

diff --git a/Engine/Plugins/Experimental/GeometryCollectionPlugin/Source/GeometryCollectionExampleCore/Private/GeometryCollection/GeometryCollectionExampleSimulationField.cpp b/Engine/Plugins/Experimental/GeometryCollectionPlugin/Source/GeometryCollectionExampleCore/Private/GeometryCollection/GeometryCollectionExampleSimulationField.cpp
--- a/Engine/Plugins/Experimental/GeometryCollectionPlugin/Source/GeometryCollectionExampleCore/Private/GeometryCollection/GeometryCollectionExampleSimulationField.cpp
+++ b/Engine/Plugins/Experimental/GeometryCollectionPlugin/Source/GeometryCollectionExampleCore/Private/GeometryCollection/GeometryCollectionExampleSimulationField.cpp
@@ -737,7 +737,6 @@ namespace GeometryCollectionExample
 		delete PhysObject;
 		delete FalloffField;
 		delete FieldObject;
-#endif
 	}
 	template void RigidBodies_Field_ClusterBreak_StrainModel_Test1<float>();
 
@@ -843,7 +842,6 @@ namespace GeometryCollectionExample
 		delete PhysObject;
 		delete FalloffField;
 		delete FieldObject;
-#endif
 	}
 	template void RigidBodies_Field_ClusterBreak_StrainModel_Test2<float>();
 
@@ -948,7 +946,6 @@ namespace GeometryCollectionExample
 		delete PhysObject;
 		delete FalloffField;
 		delete FieldObject;
-#endif
 	}
 	template void RigidBodies_Field_ClusterBreak_StrainModel_Test3<float>();
 
@@ -1045,7 +1042,6 @@ namespace GeometryCollectionExample
 		delete PhysObject;
 		delete FalloffField;
 		delete FieldObject;
-#endif
 	}
 	template void RigidBodies_Field_ClusterBreak_StrainModel_Test4<float>();
 
