Add in missing #endif
#rnx

diff --git a/Engine/Plugins/Experimental/GeometryCollectionPlugin/Source/GeometryCollectionExampleCore/Private/GeometryCollection/GeometryCollectionExampleUtility.cpp b/Engine/Plugins/Experimental/GeometryCollectionPlugin/Source/GeometryCollectionExampleCore/Private/GeometryCollection/GeometryCollectionExampleUtility.cpp
--- a/Engine/Plugins/Experimental/GeometryCollectionPlugin/Source/GeometryCollectionExampleCore/Private/GeometryCollection/GeometryCollectionExampleUtility.cpp
+++ b/Engine/Plugins/Experimental/GeometryCollectionPlugin/Source/GeometryCollectionExampleCore/Private/GeometryCollection/GeometryCollectionExampleUtility.cpp
@@ -284,6 +284,7 @@ namespace GeometryCollectionExample {
 		return PhysObject;
 #else
 		return nullptr;
+#endif
 	}
 	template FGeometryCollectionPhysicsProxy* RigidBodySetup(
 		TUniquePtr<Chaos::TChaosPhysicsMaterial<float>> & PhysicalMaterial,
