Fix static analysis.

diff --git a/Engine/Plugins/Experimental/GeometryCollectionPlugin/Source/GeometryCollectionExampleCore/Private/GeometryCollection/GeometryCollectionExampleSimulation.cpp b/Engine/Plugins/Experimental/GeometryCollectionPlugin/Source/GeometryCollectionExampleCore/Private/GeometryCollection/GeometryCollectionExampleSimulation.cpp
--- a/Engine/Plugins/Experimental/GeometryCollectionPlugin/Source/GeometryCollectionExampleCore/Private/GeometryCollection/GeometryCollectionExampleSimulation.cpp
+++ b/Engine/Plugins/Experimental/GeometryCollectionPlugin/Source/GeometryCollectionExampleCore/Private/GeometryCollection/GeometryCollectionExampleSimulation.cpp
@@ -290,7 +290,8 @@ namespace GeometryCollectionExample
 		//
 		auto RestInitFunc = [](TSharedPtr<FGeometryCollection>& RestCollection)
 		{
-			RestCollection->AppendGeometry(*RestCollection);
+			RestCollection->AppendGeometry(*GeometryCollection::MakeCubeElement(FTransform(FVector(0, 0, 0)), FVector(100.0)));
+			RestCollection->AppendGeometry(*GeometryCollection::MakeCubeElement(FTransform(FVector(0, 0, 0)), FVector(100.0)));
 			RestCollection->Transform[1].SetTranslation(FVector(0.f, 0.f, 5.f));
 		};
 
