Missing files

diff --git a/Engine/Plugins/Experimental/GeometryCollectionPlugin/Source/GeometryCollectionExampleCore/Private/GeometryCollection/GeometryCollectionExampleBoneHierarchy.cpp b/Engine/Plugins/Experimental/GeometryCollectionPlugin/Source/GeometryCollectionExampleCore/Private/GeometryCollection/GeometryCollectionExampleBoneHierarchy.cpp
--- a/Engine/Plugins/Experimental/GeometryCollectionPlugin/Source/GeometryCollectionExampleCore/Private/GeometryCollection/GeometryCollectionExampleBoneHierarchy.cpp
+++ b/Engine/Plugins/Experimental/GeometryCollectionPlugin/Source/GeometryCollectionExampleCore/Private/GeometryCollection/GeometryCollectionExampleBoneHierarchy.cpp
@@ -29,7 +29,7 @@ namespace GeometryCollectionExample
 		const TArray<Chaos::TVector<float,3>> *Points = Group1->BuildSamplePoints(1.0, 1, 1000);
 		check(Points);
 		check(Points->Num());
-		Chaos::TImplicitObject<float, 3> *Implicit = Group1->BuildSimImplicitObject();
+		Chaos::FImplicitObject *Implicit = Group1->BuildSimImplicitObject();
 		AllOnSurface(Implicit, *Points);
 	}
 	void TestSphere2(
@@ -50,7 +50,7 @@ namespace GeometryCollectionExample
 		const TArray<Chaos::TVector<float,3>> *Points = Group1->BuildSamplePoints(1.0, 1, 1000);
 		check(Points);
 		check(Points->Num());
-		Chaos::TImplicitObject<float, 3> *Implicit = Group1->BuildSimImplicitObject();
+		Chaos::FImplicitObject *Implicit = Group1->BuildSimImplicitObject();
 		AllOnSurface(Implicit, *Points);
 	}
 
diff --git a/Engine/Plugins/Experimental/GeometryCollectionPlugin/Source/GeometryCollectionExampleCore/Private/GeometryCollection/GeometryCollectionExampleClustering.cpp b/Engine/Plugins/Experimental/GeometryCollectionPlugin/Source/GeometryCollectionExampleCore/Private/GeometryCollection/GeometryCollectionExampleClustering.cpp
--- a/Engine/Plugins/Experimental/GeometryCollectionPlugin/Source/GeometryCollectionExampleCore/Private/GeometryCollection/GeometryCollectionExampleClustering.cpp
+++ b/Engine/Plugins/Experimental/GeometryCollectionPlugin/Source/GeometryCollectionExampleCore/Private/GeometryCollection/GeometryCollectionExampleClustering.cpp
@@ -1978,7 +1978,7 @@ namespace GeometryCollectionExample
 		Object->PhysicsProxy->Initialize();
 		Object->PhysicsProxy->ActivateBodies();
 
-		typedef TUniquePtr<Chaos::TImplicitObject<float, 3>> FImplicitPointer;
+		typedef TUniquePtr<Chaos::FImplicitObject> FImplicitPointer;
 		const TManagedArray<FImplicitPointer> & Implicits = Object->RestCollection->template GetAttribute<FImplicitPointer>(FGeometryCollectionPhysicsProxy::ImplicitsAttribute, FTransformCollection::TransformGroup);
 
 		typedef TUniquePtr< FCollisionStructureManager::FSimplicial > FSimplicialPointer;
diff --git a/Engine/Plugins/Experimental/GeometryCollectionPlugin/Source/GeometryCollectionExampleCore/Private/GeometryCollection/GeometryCollectionExampleSimulation.cpp b/Engine/Plugins/Experimental/GeometryCollectionPlugin/Source/GeometryCollectionExampleCore/Private/GeometryCollection/GeometryCollectionExampleSimulation.cpp
--- a/Engine/Plugins/Experimental/GeometryCollectionPlugin/Source/GeometryCollectionExampleCore/Private/GeometryCollection/GeometryCollectionExampleSimulation.cpp
+++ b/Engine/Plugins/Experimental/GeometryCollectionPlugin/Source/GeometryCollectionExampleCore/Private/GeometryCollection/GeometryCollectionExampleSimulation.cpp
@@ -444,11 +444,11 @@ namespace GeometryCollectionExample
 		Object->PhysicsProxy->Initialize();
 
 		// check implicit domain
-		typedef TUniquePtr<Chaos::TImplicitObject<float, 3>> FImplicitPointer;
+		typedef TUniquePtr<Chaos::FImplicitObject> FImplicitPointer;
 		const TManagedArray<FImplicitPointer> & Implicits = Object->RestCollection->template GetAttribute<FImplicitPointer>("Implicits", FTransformCollection::TransformGroup);
 		EXPECT_EQ(Implicits.Num(), 1);
 
-		const Chaos::TImplicitObject<float, 3> & Implicit = *Implicits[0];
+		const Chaos::FImplicitObject & Implicit = *Implicits[0];
 		for (float x = -1.05; x < 1.0; x += 0.1)
 		{
 			Vec Normal;
