Removed T from TChaosPhysicsMaterial. Rename to FChaosPhysicsMaterial. Clean up old code in Defines.h

diff --git a/Engine/Plugins/Experimental/GeometryCollectionPlugin/Source/GeometryCollectionExampleCore/Private/GeometryCollection/GeometryCollectionExampleClustering.cpp b/Engine/Plugins/Experimental/GeometryCollectionPlugin/Source/GeometryCollectionExampleCore/Private/GeometryCollection/GeometryCollectionExampleClustering.cpp
--- a/Engine/Plugins/Experimental/GeometryCollectionPlugin/Source/GeometryCollectionExampleCore/Private/GeometryCollection/GeometryCollectionExampleClustering.cpp
+++ b/Engine/Plugins/Experimental/GeometryCollectionPlugin/Source/GeometryCollectionExampleCore/Private/GeometryCollection/GeometryCollectionExampleClustering.cpp
@@ -53,7 +53,7 @@ namespace GeometryCollectionExample
 	template<class T>
 	void RigidBodies_ClusterTest_SingleLevelNonBreaking()
 	{
-		TUniquePtr<Chaos::TChaosPhysicsMaterial<T>> PhysicalMaterial = MakeUnique<Chaos::TChaosPhysicsMaterial<T>>();
+		TUniquePtr<Chaos::FChaosPhysicsMaterial> PhysicalMaterial = MakeUnique<Chaos::FChaosPhysicsMaterial>();
 		InitMaterialToZero(PhysicalMaterial);
 
 		TSharedPtr<FGeometryCollection> RestCollection = GeometryCollection::MakeCubeElement(FTransform(FQuat::MakeFromEuler(FVector(0, 0, 0.)), FVector(0, -10, 10)),FVector(1.0));
@@ -134,7 +134,7 @@ namespace GeometryCollectionExample
 	template<class T>
 	void RigidBodies_ClusterTest_DeactivateClusterParticle()
 	{
-		TUniquePtr<Chaos::TChaosPhysicsMaterial<T>> PhysicalMaterial = MakeUnique<Chaos::TChaosPhysicsMaterial<T>>();
+		TUniquePtr<Chaos::FChaosPhysicsMaterial> PhysicalMaterial = MakeUnique<Chaos::FChaosPhysicsMaterial>();
 		InitMaterialToZero(PhysicalMaterial);
 
 		TSharedPtr<FGeometryCollection> RestCollection = GeometryCollection::MakeCubeElement(FTransform(FQuat::MakeFromEuler(FVector(0.f)), FVector(20.f)),FVector(1.0));
@@ -287,7 +287,7 @@ namespace GeometryCollectionExample
 		// Wait until the cluster hits the ground.
 		// Ensure that the cluster breaks and that the children have the correct states from then on.
 		//
-		TUniquePtr<Chaos::TChaosPhysicsMaterial<T>> PhysicalMaterial = MakeUnique<Chaos::TChaosPhysicsMaterial<T>>();
+		TUniquePtr<Chaos::FChaosPhysicsMaterial> PhysicalMaterial = MakeUnique<Chaos::FChaosPhysicsMaterial>();
 		InitMaterialToZero(PhysicalMaterial);
 
 		TSharedPtr<FGeometryCollection> RestCollection = GeometryCollection::MakeCubeElement(FTransform(FQuat::MakeFromEuler(FVector(0, 0, 0.)), FVector(0, -10, 10)),FVector(1.0));
@@ -423,7 +423,7 @@ namespace GeometryCollectionExample
 	template<class T>
 	void RigidBodies_ClusterTest_NestedCluster()
 	{
-		TUniquePtr<Chaos::TChaosPhysicsMaterial<T>> PhysicalMaterial = MakeUnique<Chaos::TChaosPhysicsMaterial<T>>();
+		TUniquePtr<Chaos::FChaosPhysicsMaterial> PhysicalMaterial = MakeUnique<Chaos::FChaosPhysicsMaterial>();
 		InitMaterialToZero(PhysicalMaterial);
 
 		TSharedPtr<FGeometryCollection> RestCollection = GeometryCollection::MakeCubeElement(FTransform(FQuat::MakeFromEuler(FVector(0, 0, 0.)), FVector(0, -10, 10)),FVector(1.0));
@@ -542,7 +542,7 @@ namespace GeometryCollectionExample
 	template<class T>
 	void RigidBodies_ClusterTest_NestedCluster_MultiStrain()
 	{
-		TUniquePtr<Chaos::TChaosPhysicsMaterial<T>> PhysicalMaterial = MakeUnique<Chaos::TChaosPhysicsMaterial<T>>();
+		TUniquePtr<Chaos::FChaosPhysicsMaterial> PhysicalMaterial = MakeUnique<Chaos::FChaosPhysicsMaterial>();
 		InitMaterialToZero(PhysicalMaterial);
 
 		TSharedPtr<FGeometryCollection> RestCollection = GeometryCollection::MakeCubeElement(FTransform(FQuat::MakeFromEuler(FVector(0.f)), FVector(20.f)),FVector(1.0));
@@ -696,7 +696,7 @@ namespace GeometryCollectionExample
 	template<class T>
 	void RigidBodies_ClusterTest_NestedCluster_Halt()
 	{
-		TUniquePtr<Chaos::TChaosPhysicsMaterial<T>> PhysicalMaterial = MakeUnique<Chaos::TChaosPhysicsMaterial<T>>();
+		TUniquePtr<Chaos::FChaosPhysicsMaterial> PhysicalMaterial = MakeUnique<Chaos::FChaosPhysicsMaterial>();
 		InitMaterialToZero(PhysicalMaterial);
 
 		TSharedPtr<FGeometryCollection> RestCollection = GeometryCollection::MakeCubeElement(FTransform(FQuat::MakeFromEuler(FVector(0.f)), FVector(20.f)),FVector(1.0));
@@ -833,7 +833,7 @@ namespace GeometryCollectionExample
 		// Test : Set one element kinematic. When the cluster breaks the elements that do not contain the kinematic
 		//        rigid body should be dynamic, while the clusters that contain the kinematic body should remain 
 		//        kinematic. 
-		TUniquePtr<Chaos::TChaosPhysicsMaterial<T>> PhysicalMaterial = MakeUnique<Chaos::TChaosPhysicsMaterial<T>>();
+		TUniquePtr<Chaos::FChaosPhysicsMaterial> PhysicalMaterial = MakeUnique<Chaos::FChaosPhysicsMaterial>();
 		InitMaterialToZero(PhysicalMaterial);
 
 		TSharedPtr<FGeometryCollection> RestCollection = GeometryCollection::MakeCubeElement(FTransform(FQuat::MakeFromEuler(FVector(0.f)), FVector(20.f)), FVector(1.0));
@@ -1071,7 +1071,7 @@ namespace GeometryCollectionExample
 		// Test : Set one element static. When the cluster breaks the elements that do not contain the static
 		//        rigid body should be dynamic, while the clusters that contain the static body should remain 
 		//        static. 
-		TUniquePtr<Chaos::TChaosPhysicsMaterial<T>> PhysicalMaterial = MakeUnique<Chaos::TChaosPhysicsMaterial<T>>();
+		TUniquePtr<Chaos::FChaosPhysicsMaterial> PhysicalMaterial = MakeUnique<Chaos::FChaosPhysicsMaterial>();
 		InitMaterialToZero(PhysicalMaterial);
 
 		TSharedPtr<FGeometryCollection> RestCollection = GeometryCollection::MakeCubeElement(FTransform(FQuat::MakeFromEuler(FVector(0.f)), FVector(20.f)), FVector(1.0));
@@ -1283,7 +1283,7 @@ namespace GeometryCollectionExample
 	{
 		// Test : Set one element kinematic. When the cluster breaks the elements that do not contain the kinematic
 		//        Rigid body should be dynamic, while the clusters that contain the kinematic body should remain kinematic. 
-		TUniquePtr<Chaos::TChaosPhysicsMaterial<T>> PhysicalMaterial = MakeUnique<Chaos::TChaosPhysicsMaterial<T>>();
+		TUniquePtr<Chaos::FChaosPhysicsMaterial> PhysicalMaterial = MakeUnique<Chaos::FChaosPhysicsMaterial>();
 		InitMaterialToZero(PhysicalMaterial);
 
 		TSharedPtr<FGeometryCollection> RestCollection = CreateClusteredBody( FVector(0,0,100) );
@@ -1414,7 +1414,7 @@ namespace GeometryCollectionExample
 	template<class T>
 	void RigidBodies_ClusterTest_ReleaseClusterParticle_ClusteredNode()
 	{
-		TUniquePtr<Chaos::TChaosPhysicsMaterial<T>> PhysicalMaterial = MakeUnique<Chaos::TChaosPhysicsMaterial<T>>();
+		TUniquePtr<Chaos::FChaosPhysicsMaterial> PhysicalMaterial = MakeUnique<Chaos::FChaosPhysicsMaterial>();
 		InitMaterialToZero(PhysicalMaterial);
 
 		TSharedPtr<FGeometryCollection> RestCollection = CreateClusteredBody(FVector(0, 0, 100));
@@ -1551,7 +1551,7 @@ namespace GeometryCollectionExample
 	template<class T>
 	void RigidBodies_ClusterTest_ReleaseClusterParticle_ClusteredKinematicNode()
 	{
-		TUniquePtr<Chaos::TChaosPhysicsMaterial<T>> PhysicalMaterial = MakeUnique<Chaos::TChaosPhysicsMaterial<T>>();
+		TUniquePtr<Chaos::FChaosPhysicsMaterial> PhysicalMaterial = MakeUnique<Chaos::FChaosPhysicsMaterial>();
 		InitMaterialToZero(PhysicalMaterial);
 
 		TSharedPtr<FGeometryCollection> RestCollection = CreateClusteredBody(FVector(0, 0, 100));
@@ -1685,7 +1685,7 @@ namespace GeometryCollectionExample
 	template<class T>
 	void RigidBodies_ClusterTest_ReleaseClusterParticles_AllLeafNodes()
 	{
-		TUniquePtr<Chaos::TChaosPhysicsMaterial<T>> PhysicalMaterial = MakeUnique<Chaos::TChaosPhysicsMaterial<T>>();
+		TUniquePtr<Chaos::FChaosPhysicsMaterial> PhysicalMaterial = MakeUnique<Chaos::FChaosPhysicsMaterial>();
 		InitMaterialToZero(PhysicalMaterial);
 
 		TSharedPtr<FGeometryCollection> RestCollection = CreateClusteredBody(FVector(0, 0, 100));
@@ -1773,7 +1773,7 @@ namespace GeometryCollectionExample
 	template<class T>
 	void RigidBodies_ClusterTest_ReleaseClusterParticles_ClusterNodeAndSubClusterNode()
 	{
-		TUniquePtr<Chaos::TChaosPhysicsMaterial<T>> PhysicalMaterial = MakeUnique<Chaos::TChaosPhysicsMaterial<T>>();
+		TUniquePtr<Chaos::FChaosPhysicsMaterial> PhysicalMaterial = MakeUnique<Chaos::FChaosPhysicsMaterial>();
 		InitMaterialToZero(PhysicalMaterial);
 
 		TSharedPtr<FGeometryCollection> RestCollection = CreateClusteredBody_TwoParents_TwoBodies(FVector(0, 0, 100));
@@ -1868,7 +1868,7 @@ namespace GeometryCollectionExample
 	template<class T>
 	void RigidBodies_ClusterTest_RemoveOnFracture()
 	{
-		TUniquePtr<Chaos::TChaosPhysicsMaterial<T>> PhysicalMaterial = MakeUnique<Chaos::TChaosPhysicsMaterial<T>>();
+		TUniquePtr<Chaos::FChaosPhysicsMaterial> PhysicalMaterial = MakeUnique<Chaos::FChaosPhysicsMaterial>();
 		InitMaterialToZero(PhysicalMaterial);
 
 		TSharedPtr<FGeometryCollection> RestCollection = GeometryCollection::MakeCubeElement(FTransform(FQuat::MakeFromEuler(FVector(0, 0, 0.)), FVector(0, -10, 10)), FVector(1.0));
diff --git a/Engine/Plugins/Experimental/GeometryCollectionPlugin/Source/GeometryCollectionExampleCore/Private/GeometryCollection/GeometryCollectionExampleSimulation.cpp b/Engine/Plugins/Experimental/GeometryCollectionPlugin/Source/GeometryCollectionExampleCore/Private/GeometryCollection/GeometryCollectionExampleSimulation.cpp
--- a/Engine/Plugins/Experimental/GeometryCollectionPlugin/Source/GeometryCollectionExampleCore/Private/GeometryCollection/GeometryCollectionExampleSimulation.cpp
+++ b/Engine/Plugins/Experimental/GeometryCollectionPlugin/Source/GeometryCollectionExampleCore/Private/GeometryCollection/GeometryCollectionExampleSimulation.cpp
@@ -29,7 +29,7 @@ namespace GeometryCollectionExample
 	template<class T>
 	void RigidBodiesFallingUnderGravity()
 	{
-		TUniquePtr<Chaos::TChaosPhysicsMaterial<T>> PhysicalMaterial = nullptr;
+		TUniquePtr<Chaos::FChaosPhysicsMaterial> PhysicalMaterial = nullptr;
 		TSharedPtr<FGeometryCollection> RestCollection = nullptr;
 		TSharedPtr<FGeometryDynamicCollection> DynamicCollection = nullptr;
 
@@ -73,7 +73,7 @@ namespace GeometryCollectionExample
 	template<class T>
 	void RigidBodiesCollidingWithSolverFloor()
 	{
-		TUniquePtr<Chaos::TChaosPhysicsMaterial<T>> PhysicalMaterial = nullptr;
+		TUniquePtr<Chaos::FChaosPhysicsMaterial> PhysicalMaterial = nullptr;
 		TSharedPtr<FGeometryCollection> RestCollection = nullptr;
 		TSharedPtr<FGeometryDynamicCollection> DynamicCollection = nullptr;
 
@@ -123,7 +123,7 @@ namespace GeometryCollectionExample
 	template<class T>
 	void RigidBodiesSingleSphereCollidingWithSolverFloor()
 	{
-		TUniquePtr<Chaos::TChaosPhysicsMaterial<T>> PhysicalMaterial = nullptr;
+		TUniquePtr<Chaos::FChaosPhysicsMaterial> PhysicalMaterial = nullptr;
 		TSharedPtr<FGeometryCollection> RestCollection = nullptr;
 		TSharedPtr<FGeometryDynamicCollection> DynamicCollection = nullptr;
 
@@ -183,7 +183,7 @@ namespace GeometryCollectionExample
 	template<class T>
 	void RigidBodiesSingleSphereIntersectingWithSolverFloor()
 	{
-		TUniquePtr<Chaos::TChaosPhysicsMaterial<T>> PhysicalMaterial = nullptr;
+		TUniquePtr<Chaos::FChaosPhysicsMaterial> PhysicalMaterial = nullptr;
 		TSharedPtr<FGeometryCollection> RestCollection = nullptr;
 		TSharedPtr<FGeometryDynamicCollection> DynamicCollection = nullptr;
 
@@ -236,7 +236,7 @@ namespace GeometryCollectionExample
 	template<class T>
 	void RigidBodiesKinematic()
 	{
-		TUniquePtr<Chaos::TChaosPhysicsMaterial<T>> PhysicalMaterial = nullptr;
+		TUniquePtr<Chaos::FChaosPhysicsMaterial> PhysicalMaterial = nullptr;
 		TSharedPtr<FGeometryCollection> RestCollection = nullptr;
 		TSharedPtr<FGeometryDynamicCollection> DynamicCollection = nullptr;
 
@@ -281,7 +281,7 @@ namespace GeometryCollectionExample
 	template<class T>
 	void RigidBodiesSleepingActivation()
 	{
-		TUniquePtr<Chaos::TChaosPhysicsMaterial<T>> PhysicalMaterial = nullptr;
+		TUniquePtr<Chaos::FChaosPhysicsMaterial> PhysicalMaterial = nullptr;
 		TSharedPtr<FGeometryCollection> RestCollection = nullptr;
 		TSharedPtr<FGeometryDynamicCollection> DynamicCollection = nullptr;
 
@@ -339,7 +339,7 @@ namespace GeometryCollectionExample
 	template<class T>
 	void RigidBodies_CollisionGroup()
 	{
-		TUniquePtr<Chaos::TChaosPhysicsMaterial<T>> PhysicalMaterial = nullptr;
+		TUniquePtr<Chaos::FChaosPhysicsMaterial> PhysicalMaterial = nullptr;
 		TSharedPtr<FGeometryCollection> RestCollection = nullptr;
 		TSharedPtr<FGeometryDynamicCollection> DynamicCollection = nullptr;
 
diff --git a/Engine/Plugins/Experimental/GeometryCollectionPlugin/Source/GeometryCollectionExampleCore/Private/GeometryCollection/GeometryCollectionExampleSimulationField.cpp b/Engine/Plugins/Experimental/GeometryCollectionPlugin/Source/GeometryCollectionExampleCore/Private/GeometryCollection/GeometryCollectionExampleSimulationField.cpp
--- a/Engine/Plugins/Experimental/GeometryCollectionPlugin/Source/GeometryCollectionExampleCore/Private/GeometryCollection/GeometryCollectionExampleSimulationField.cpp
+++ b/Engine/Plugins/Experimental/GeometryCollectionPlugin/Source/GeometryCollectionExampleCore/Private/GeometryCollection/GeometryCollectionExampleSimulationField.cpp
@@ -31,7 +31,7 @@ namespace GeometryCollectionExample
 	template<class T>
 	void RigidBodies_Field_KinematicActivation()
 	{
-		TUniquePtr<Chaos::TChaosPhysicsMaterial<T>> PhysicalMaterial = nullptr;
+		TUniquePtr<Chaos::FChaosPhysicsMaterial> PhysicalMaterial = nullptr;
 		TSharedPtr<FGeometryCollection> RestCollection = nullptr;
 		TSharedPtr<FGeometryDynamicCollection> DynamicCollection = nullptr;
 
@@ -96,7 +96,7 @@ namespace GeometryCollectionExample
 	template<class T>
 	void RigidBodies_Field_InitialLinearVelocity()
 	{
-		TUniquePtr<Chaos::TChaosPhysicsMaterial<T>> PhysicalMaterial = nullptr;
+		TUniquePtr<Chaos::FChaosPhysicsMaterial> PhysicalMaterial = nullptr;
 		TSharedPtr<FGeometryCollection> RestCollection = nullptr;
 		TSharedPtr<FGeometryDynamicCollection> DynamicCollection = nullptr;
 
@@ -179,7 +179,7 @@ namespace GeometryCollectionExample
 	template<class T>
 	void RigidBodies_Field_StayDynamic()
 	{
-		TUniquePtr<Chaos::TChaosPhysicsMaterial<T>> PhysicalMaterial = nullptr;
+		TUniquePtr<Chaos::FChaosPhysicsMaterial> PhysicalMaterial = nullptr;
 		TSharedPtr<FGeometryCollection> RestCollection = nullptr;
 		TSharedPtr<FGeometryDynamicCollection> DynamicCollection = nullptr;
 
@@ -265,7 +265,7 @@ namespace GeometryCollectionExample
 	template<class T>
 	void RigidBodies_Field_LinearForce()
 	{
-		TUniquePtr<Chaos::TChaosPhysicsMaterial<T>> PhysicalMaterial = nullptr;
+		TUniquePtr<Chaos::FChaosPhysicsMaterial> PhysicalMaterial = nullptr;
 		TSharedPtr<FGeometryCollection> RestCollection = nullptr;
 		TSharedPtr<FGeometryDynamicCollection> DynamicCollection = nullptr;
 
@@ -342,7 +342,7 @@ namespace GeometryCollectionExample
 	template<class T>
 	void RigidBodies_Field_Torque()
 	{
-		TUniquePtr<Chaos::TChaosPhysicsMaterial<T>> PhysicalMaterial = nullptr;
+		TUniquePtr<Chaos::FChaosPhysicsMaterial> PhysicalMaterial = nullptr;
 		TSharedPtr<FGeometryCollection> RestCollection = nullptr;
 		TSharedPtr<FGeometryDynamicCollection> DynamicCollection = nullptr;
 
@@ -427,7 +427,7 @@ namespace GeometryCollectionExample
 	template<class T>
 	void RigidBodies_Field_Kill()
 	{
-		TUniquePtr<Chaos::TChaosPhysicsMaterial<T>> PhysicalMaterial = nullptr;
+		TUniquePtr<Chaos::FChaosPhysicsMaterial> PhysicalMaterial = nullptr;
 		TSharedPtr<FGeometryCollection> RestCollection = nullptr;
 		TSharedPtr<FGeometryDynamicCollection> DynamicCollection = nullptr;
 
@@ -495,7 +495,7 @@ namespace GeometryCollectionExample
 	template<class T>
 	void RigidBodies_Field_LinearVelocity()
 	{
-		TUniquePtr<Chaos::TChaosPhysicsMaterial<T>> PhysicalMaterial = nullptr;
+		TUniquePtr<Chaos::FChaosPhysicsMaterial> PhysicalMaterial = nullptr;
 		TSharedPtr<FGeometryCollection> RestCollection = nullptr;
 		TSharedPtr<FGeometryDynamicCollection> DynamicCollection = nullptr;
 
@@ -568,7 +568,7 @@ namespace GeometryCollectionExample
 	template<class T>
 	void RigidBodies_Field_CollisionGroup()
 	{
-		TUniquePtr<Chaos::TChaosPhysicsMaterial<T>> PhysicalMaterial = nullptr;
+		TUniquePtr<Chaos::FChaosPhysicsMaterial> PhysicalMaterial = nullptr;
 		TSharedPtr<FGeometryCollection> RestCollection = nullptr;
 		TSharedPtr<FGeometryDynamicCollection> DynamicCollection = nullptr;
 
@@ -652,7 +652,7 @@ namespace GeometryCollectionExample
 	template<class T>
 	void RigidBodies_Field_ClusterBreak_StrainModel_Test1()
 	{
-		TUniquePtr<Chaos::TChaosPhysicsMaterial<T>> PhysicalMaterial = nullptr;
+		TUniquePtr<Chaos::FChaosPhysicsMaterial> PhysicalMaterial = nullptr;
 		TSharedPtr<FGeometryCollection> RestCollection = CreateClusteredBody_TwoByTwo_ThreeTransform(FVector(0));
 		TSharedPtr<FGeometryDynamicCollection> DynamicCollection = nullptr;
 
@@ -745,7 +745,7 @@ namespace GeometryCollectionExample
 	template<class T>
 	void RigidBodies_Field_ClusterBreak_StrainModel_Test2()
 	{
-		TUniquePtr<Chaos::TChaosPhysicsMaterial<T>> PhysicalMaterial = nullptr;
+		TUniquePtr<Chaos::FChaosPhysicsMaterial> PhysicalMaterial = nullptr;
 		TSharedPtr<FGeometryCollection> RestCollection = CreateClusteredBody_ThreeByTwo_ThreeTransform(FVector(0));
 		TSharedPtr<FGeometryDynamicCollection> DynamicCollection = nullptr;
 
@@ -848,7 +848,7 @@ namespace GeometryCollectionExample
 	template<class T>
 	void RigidBodies_Field_ClusterBreak_StrainModel_Test3()
 	{
-		TUniquePtr<Chaos::TChaosPhysicsMaterial<T>> PhysicalMaterial = nullptr;
+		TUniquePtr<Chaos::FChaosPhysicsMaterial> PhysicalMaterial = nullptr;
 		TSharedPtr<FGeometryCollection> RestCollection = CreateClusteredBody_ThreeByTwo_ThreeTransform(FVector(0));
 		TSharedPtr<FGeometryDynamicCollection> DynamicCollection = nullptr;
 
@@ -953,7 +953,7 @@ namespace GeometryCollectionExample
 	template<class T>
 	void RigidBodies_Field_ClusterBreak_StrainModel_Test4()
 	{
-		TUniquePtr<Chaos::TChaosPhysicsMaterial<T>> PhysicalMaterial = nullptr;
+		TUniquePtr<Chaos::FChaosPhysicsMaterial> PhysicalMaterial = nullptr;
 		TSharedPtr<FGeometryCollection> RestCollection = CreateClusteredBody_TwoByTwo_ThreeTransform(FVector(0));
 		TSharedPtr<FGeometryDynamicCollection> DynamicCollection = nullptr;
 
diff --git a/Engine/Plugins/Experimental/GeometryCollectionPlugin/Source/GeometryCollectionExampleCore/Private/GeometryCollection/GeometryCollectionExampleSimulationSolver.cpp b/Engine/Plugins/Experimental/GeometryCollectionPlugin/Source/GeometryCollectionExampleCore/Private/GeometryCollection/GeometryCollectionExampleSimulationSolver.cpp
--- a/Engine/Plugins/Experimental/GeometryCollectionPlugin/Source/GeometryCollectionExampleCore/Private/GeometryCollection/GeometryCollectionExampleSimulationSolver.cpp
+++ b/Engine/Plugins/Experimental/GeometryCollectionPlugin/Source/GeometryCollectionExampleCore/Private/GeometryCollection/GeometryCollectionExampleSimulationSolver.cpp
@@ -153,7 +153,7 @@ namespace GeometryCollectionExample
 	template<class T>
 	void Solver_ValidateReverseMapping()
 	{
-		TUniquePtr<Chaos::TChaosPhysicsMaterial<T>> PhysicalMaterial = MakeUnique<Chaos::TChaosPhysicsMaterial<T>>();
+		TUniquePtr<Chaos::FChaosPhysicsMaterial> PhysicalMaterial = MakeUnique<Chaos::FChaosPhysicsMaterial>();
 		InitMaterialToZero(PhysicalMaterial);
 
 		Chaos::FPBDRigidsSolver* Solver = FChaosSolversModule::GetModule()->CreateSolver(true);
@@ -212,9 +212,9 @@ namespace GeometryCollectionExample
 	template void Solver_ValidateReverseMapping<float>();
 
 	template<class T>
-	void CommonInit(int32 NumObjects, bool UseClusters, Chaos::FPBDRigidsSolver** SolverInOut, TUniquePtr<Chaos::TChaosPhysicsMaterial<T>>& PhysicalMaterial, TArray<TSharedPtr<FGeometryCollection>>& RestArray, TArray<TSharedPtr<FGeometryDynamicCollection>>& DynamicArray)
+	void CommonInit(int32 NumObjects, bool UseClusters, Chaos::FPBDRigidsSolver** SolverInOut, TUniquePtr<Chaos::FChaosPhysicsMaterial>& PhysicalMaterial, TArray<TSharedPtr<FGeometryCollection>>& RestArray, TArray<TSharedPtr<FGeometryDynamicCollection>>& DynamicArray)
 	{
-		PhysicalMaterial = MakeUnique<Chaos::TChaosPhysicsMaterial<T>>();
+		PhysicalMaterial = MakeUnique<Chaos::FChaosPhysicsMaterial>();
 		InitMaterialToZero(PhysicalMaterial);
 
 		*SolverInOut = FChaosSolversModule::GetModule()->CreateSolver(true);
@@ -294,7 +294,7 @@ namespace GeometryCollectionExample
 		float TestMassThreshold = 6.0f;
 
 		Chaos::FPBDRigidsSolver* Solver;
-		TUniquePtr<Chaos::TChaosPhysicsMaterial<T>> PhysicalMaterial;
+		TUniquePtr<Chaos::FChaosPhysicsMaterial> PhysicalMaterial;
 		TArray<TSharedPtr<FGeometryCollection> > RestArray;
 		TArray<TSharedPtr<FGeometryDynamicCollection> > DynamicArray;
 
@@ -357,7 +357,7 @@ namespace GeometryCollectionExample
 		float TestMass = 6.0f;
 
 		Chaos::FPBDRigidsSolver* Solver;
-		TUniquePtr<Chaos::TChaosPhysicsMaterial<T>> PhysicalMaterial;
+		TUniquePtr<Chaos::FChaosPhysicsMaterial> PhysicalMaterial;
 		TArray<TSharedPtr<FGeometryCollection> > RestArray;
 		TArray<TSharedPtr<FGeometryDynamicCollection> > DynamicArray;
 
diff --git a/Engine/Plugins/Experimental/GeometryCollectionPlugin/Source/GeometryCollectionExampleCore/Public/GeometryCollection/GeometryCollectionExampleSimulationObject.h b/Engine/Plugins/Experimental/GeometryCollectionPlugin/Source/GeometryCollectionExampleCore/Public/GeometryCollection/GeometryCollectionExampleSimulationObject.h
--- a/Engine/Plugins/Experimental/GeometryCollectionPlugin/Source/GeometryCollectionExampleCore/Public/GeometryCollection/GeometryCollectionExampleSimulationObject.h
+++ b/Engine/Plugins/Experimental/GeometryCollectionPlugin/Source/GeometryCollectionExampleCore/Public/GeometryCollection/GeometryCollectionExampleSimulationObject.h
@@ -48,7 +48,7 @@ namespace GeometryCollectionExample {
 			: Parameters(InParameters)
 			, RestCollection(RestCollectionIn)
 			, DynamicCollection(GeometryCollectionToGeometryDynamicCollection(RestCollection.Get()))
-			, PhysicalMaterial(new Chaos::TChaosPhysicsMaterial<T>())
+			, PhysicalMaterial(new Chaos::FChaosPhysicsMaterial())
 			, PhysicsProxy(new FGeometryCollectionPhysicsProxy(nullptr, DynamicCollection.Get(), [&](FSimulationParameters& P) {Init(P); }, nullptr, nullptr))
 		{
 			PhysicalMaterial->Friction = 0;
@@ -63,7 +63,7 @@ namespace GeometryCollectionExample {
 		SimulationObjects(FParameters InParameters,
 			TSharedPtr<FGeometryCollection> InRestCollection, 
 			TSharedPtr<FGeometryDynamicCollection> InDynamicCollection, 
-			TUniquePtr<Chaos::TChaosPhysicsMaterial<T>> InPhysicalMaterial)
+			TUniquePtr<Chaos::FChaosPhysicsMaterial> InPhysicalMaterial)
 			: Parameters(InParameters)
 			, RestCollection(InRestCollection)
 			, DynamicCollection(InDynamicCollection)
@@ -96,7 +96,7 @@ namespace GeometryCollectionExample {
 		FParameters Parameters;
 		TSharedPtr<FGeometryCollection> RestCollection;
 		TSharedPtr<FGeometryDynamicCollection> DynamicCollection;
-		TUniquePtr<Chaos::TChaosPhysicsMaterial<T>> PhysicalMaterial;
+		TUniquePtr<Chaos::FChaosPhysicsMaterial> PhysicalMaterial;
 		TSharedPtr<FGeometryCollectionPhysicsProxy> PhysicsProxy;
 		FSharedSimulationSizeSpecificData SimulationData;
 	};
diff --git a/Engine/Source/Programs/UnrealBuildTool/Configuration/ModuleRules.cs b/Engine/Source/Programs/UnrealBuildTool/Configuration/ModuleRules.cs
--- a/Engine/Source/Programs/UnrealBuildTool/Configuration/ModuleRules.cs
+++ b/Engine/Source/Programs/UnrealBuildTool/Configuration/ModuleRules.cs
@@ -1159,7 +1159,6 @@ namespace UnrealBuildTool
 				{
 					PublicDefinitions.Add("WITH_CHAOS=1");
 					PublicDefinitions.Add("WITH_CHAOS_NEEDS_TO_BE_FIXED=1");
-					PublicDefinitions.Add("COMPILE_ID_TYPES_AS_INTS=0");
 					PublicDefinitions.Add("WITH_CHAOS_CLOTHING=1");
 					PublicDefinitions.Add("WITH_CLOTH_COLLISION_DETECTION=1");
 					
diff --git a/Engine/Source/Runtime/Experimental/Chaos/Private/Chaos/PBDCollisionConstraint.cpp b/Engine/Source/Runtime/Experimental/Chaos/Private/Chaos/PBDCollisionConstraint.cpp
--- a/Engine/Source/Runtime/Experimental/Chaos/Private/Chaos/PBDCollisionConstraint.cpp
+++ b/Engine/Source/Runtime/Experimental/Chaos/Private/Chaos/PBDCollisionConstraint.cpp
@@ -73,7 +73,7 @@ void ConstructConstraintsImpl(TGeometryParticleHandle<T, d>* Particle0, TGeometr
 //
 
 template<typename T, int d>
-TPBDCollisionConstraint<T, d>::TPBDCollisionConstraint(const TPBDRigidsSOAs<T,d>& InParticles, TArrayCollectionArray<bool>& Collided, const TArrayCollectionArray<TSerializablePtr<TChaosPhysicsMaterial<T>>>& InPerParticleMaterials, const int32 PairIterations /*= 1*/, const T Thickness /*= (T)0*/)
+TPBDCollisionConstraint<T, d>::TPBDCollisionConstraint(const TPBDRigidsSOAs<T,d>& InParticles, TArrayCollectionArray<bool>& Collided, const TArrayCollectionArray<TSerializablePtr<FChaosPhysicsMaterial>>& InPerParticleMaterials, const int32 PairIterations /*= 1*/, const T Thickness /*= (T)0*/)
 	: Particles(InParticles)
 	, SpatialAcceleration(nullptr)
 	, MCollided(Collided)
@@ -596,8 +596,8 @@ void TPBDCollisionConstraint<T, d>::Apply(const T Dt, FRigidBodyContactConstrain
 	const TVector<T, d>& V1 = Particle1->V();
 	const TVector<T, d>& W0 = Particle0->W();
 	const TVector<T, d>& W1 = Particle1->W();
-	TSerializablePtr<TChaosPhysicsMaterial<T>> PhysicsMaterial0 = Particle0->AuxilaryValue(MPhysicsMaterials);
-	TSerializablePtr<TChaosPhysicsMaterial<T>> PhysicsMaterial1 = Particle1->AuxilaryValue(MPhysicsMaterials);
+	TSerializablePtr<FChaosPhysicsMaterial> PhysicsMaterial0 = Particle0->AuxilaryValue(MPhysicsMaterials);
+	TSerializablePtr<FChaosPhysicsMaterial> PhysicsMaterial1 = Particle1->AuxilaryValue(MPhysicsMaterials);
 
 	TVector<T, d> VectorToPoint1 = Constraint.Location - P0;
 	TVector<T, d> VectorToPoint2 = Constraint.Location - P1;
@@ -819,8 +819,8 @@ void TPBDCollisionConstraint<T, d>::ApplyPushOut(const T Dt, FRigidBodyContactCo
 	const TVector<T, d>& V1 = PBDRigid1 ? PBDRigid1->V() : ZeroVector;
 	const TVector<T, d>& W0 = PBDRigid0 ? PBDRigid0->W() : ZeroVector;
 	const TVector<T, d>& W1 = PBDRigid1 ? PBDRigid1->W() : ZeroVector;
-	TSerializablePtr<TChaosPhysicsMaterial<T>> PhysicsMaterial0 = Particle0->AuxilaryValue(MPhysicsMaterials);
-	TSerializablePtr<TChaosPhysicsMaterial<T>> PhysicsMaterial1 = Particle1->AuxilaryValue(MPhysicsMaterials);
+	TSerializablePtr<FChaosPhysicsMaterial> PhysicsMaterial0 = Particle0->AuxilaryValue(MPhysicsMaterials);
+	TSerializablePtr<FChaosPhysicsMaterial> PhysicsMaterial1 = Particle1->AuxilaryValue(MPhysicsMaterials);
 	const bool IsTemporarilyStatic0 = IsTemporarilyStatic.Contains(Particle0);
 	const bool IsTemporarilyStatic1 = IsTemporarilyStatic.Contains(Particle1);
 
diff --git a/Engine/Source/Runtime/Experimental/Chaos/Private/Chaos/PBDCollisionConstraintPGS.cpp b/Engine/Source/Runtime/Experimental/Chaos/Private/Chaos/PBDCollisionConstraintPGS.cpp
--- a/Engine/Source/Runtime/Experimental/Chaos/Private/Chaos/PBDCollisionConstraintPGS.cpp
+++ b/Engine/Source/Runtime/Experimental/Chaos/Private/Chaos/PBDCollisionConstraintPGS.cpp
@@ -73,7 +73,7 @@ TRigidTransform<T, d> GetTransformPGS(const TPBDRigidParticles<T, d>& InParticle
 }
 
 template<class T, int d>
-TPBDCollisionConstraintPGS<T, d>::TPBDCollisionConstraintPGS(TPBDRigidParticles<T, d>& InParticles, const TArray<int32>& InIndices, TArrayCollectionArray<bool>& Collided, TArrayCollectionArray<TSerializablePtr<TChaosPhysicsMaterial<T>>>& InPhysicsMaterials, const T Thickness /*= (T)0*/)
+TPBDCollisionConstraintPGS<T, d>::TPBDCollisionConstraintPGS(TPBDRigidParticles<T, d>& InParticles, const TArray<int32>& InIndices, TArrayCollectionArray<bool>& Collided, TArrayCollectionArray<TSerializablePtr<FChaosPhysicsMaterial>>& InPhysicsMaterials, const T Thickness /*= (T)0*/)
     : MCollided(Collided), MPhysicsMaterials(InPhysicsMaterials), MThickness(Thickness), Tolerance(KINDA_SMALL_NUMBER), MaxIterations(10), bUseCCD(false)
 {
 }
diff --git a/Engine/Source/Runtime/Experimental/Chaos/Private/Chaos/PBDConstraintGraph.cpp b/Engine/Source/Runtime/Experimental/Chaos/Private/Chaos/PBDConstraintGraph.cpp
--- a/Engine/Source/Runtime/Experimental/Chaos/Private/Chaos/PBDConstraintGraph.cpp
+++ b/Engine/Source/Runtime/Experimental/Chaos/Private/Chaos/PBDConstraintGraph.cpp
@@ -506,7 +506,7 @@ void FPBDConstraintGraph::ComputeIsland(const int32 InNode, const int32 Island,
 }
 
 
-bool FPBDConstraintGraph::SleepInactive(const int32 Island, const TArrayCollectionArray<TSerializablePtr<TChaosPhysicsMaterial<FReal>>>& PerParticleMaterialAttributes)
+bool FPBDConstraintGraph::SleepInactive(const int32 Island, const TArrayCollectionArray<TSerializablePtr<FChaosPhysicsMaterial>>& PerParticleMaterialAttributes)
 {
 	// @todo(ccaulfield): should be able to eliminate this when island is already sleeping
 
@@ -535,7 +535,7 @@ bool FPBDConstraintGraph::SleepInactive(const int32 Island, const TArrayCollecti
 			M += PBDRigid->M();
 			V += PBDRigid->V() * PBDRigid->M();
 
-			if (TSerializablePtr<TChaosPhysicsMaterial<FReal>> PhysicsMaterial = Particle->AuxilaryValue(PerParticleMaterialAttributes))
+			if (TSerializablePtr<FChaosPhysicsMaterial> PhysicsMaterial = Particle->AuxilaryValue(PerParticleMaterialAttributes))
 			{
 				LinearSleepingThreshold = FMath::Min(LinearSleepingThreshold, PhysicsMaterial->SleepingLinearThreshold);
 				AngularSleepingThreshold = FMath::Min(AngularSleepingThreshold, PhysicsMaterial->SleepingAngularThreshold);
diff --git a/Engine/Source/Runtime/Experimental/Chaos/Public/Chaos/Defines.h b/Engine/Source/Runtime/Experimental/Chaos/Public/Chaos/Defines.h
--- a/Engine/Source/Runtime/Experimental/Chaos/Public/Chaos/Defines.h
+++ b/Engine/Source/Runtime/Experimental/Chaos/Public/Chaos/Defines.h
@@ -1,11 +1,6 @@
 // Copyright 1998-2019 Epic Games, Inc. All Rights Reserved.
 #pragma once
 
-#ifndef COMPILE_ID_TYPES_AS_INTS
-#define COMPILE_ID_TYPES_AS_INTS 0
-#endif
-
-#include <functional>
 #if COMPILE_WITHOUT_UNREAL_SUPPORT
 	#include <stdint.h>
 #else
@@ -21,79 +16,57 @@
 typedef int32_t int32;
 #endif
 
+#include "Chaos/Core.h"
 
 namespace Chaos
 {
-#if COMPILE_ID_TYPES_AS_INTS
-typedef uint32 IslandId;
+	class FChaosPhysicsMaterial
+	{
+	public:
 
-static FORCEINLINE uint32 ToValue(uint32 Id) { return Id; }
-#else
-#define CREATEIDTYPE(IDNAME) \
-    class IDNAME \
-    { \
-      public: \
-        IDNAME() {} \
-        IDNAME(const uint32 InValue) : Value(InValue) {} \
-        bool operator==(const IDNAME& Other) const { return Value == Other.Value; } \
-        uint32 Value; \
-    }
+		Chaos::FReal Friction;
+		Chaos::FReal Restitution;
+		Chaos::FReal SleepingLinearThreshold;
+		Chaos::FReal SleepingAngularThreshold;
+		Chaos::FReal DisabledLinearThreshold;
+		Chaos::FReal DisabledAngularThreshold;
 
-CREATEIDTYPE(IslandId);
+		FChaosPhysicsMaterial()
+			: Friction(0.5)
+			, Restitution(0.1)
+			, SleepingLinearThreshold(1)
+			, SleepingAngularThreshold(1)
+			, DisabledLinearThreshold(0)
+			, DisabledAngularThreshold(0)
+		{
+		}
 
-template<class T_ID>
-static uint32 ToValue(T_ID Id)
-{
-    return Id.Value;
-}
-#endif
-template<class T>
-class TChaosPhysicsMaterial
-{
-public:
-	T Friction;
-	T Restitution;
-	T SleepingLinearThreshold;
-	T SleepingAngularThreshold;
-	T DisabledLinearThreshold;
-	T DisabledAngularThreshold;
+		static constexpr bool IsSerializablePtr = true;
 
-	TChaosPhysicsMaterial()
-		: Friction((T)0.5)
-		, Restitution((T)0.1)
-		, SleepingLinearThreshold((T)1)
-		, SleepingAngularThreshold((T)1)
-		, DisabledLinearThreshold((T)0)
-		, DisabledAngularThreshold((T)0)
-	{
-	}
+		static void StaticSerialize(FArchive& Ar, TSerializablePtr<FChaosPhysicsMaterial>& Serializable)
+		{
+			FChaosPhysicsMaterial* Material = const_cast<FChaosPhysicsMaterial*>(Serializable.Get());
 
-	static constexpr bool IsSerializablePtr = true;
+			if(Ar.IsLoading())
+			{
+				Material = new FChaosPhysicsMaterial();
+				Serializable.SetFromRawLowLevel(Material);
+			}
 
-	static void StaticSerialize(FArchive& Ar, TSerializablePtr<TChaosPhysicsMaterial<T>>& Serializable)
-	{
-		TChaosPhysicsMaterial<T>* Material = const_cast<TChaosPhysicsMaterial<T>*>(Serializable.Get());
-		
-		if (Ar.IsLoading())
+			Material->Serialize(Ar);
+		}
+
+		void Serialize(FArchive& Ar)
 		{
-			Material = new TChaosPhysicsMaterial<T>();
-			Serializable.SetFromRawLowLevel(Material);
+			Ar << Friction << Restitution << SleepingLinearThreshold << SleepingAngularThreshold << DisabledLinearThreshold << DisabledAngularThreshold;
 		}
-		
-		Material->Serialize(Ar);
-	}
+	};
+
 
-	void Serialize(FArchive& Ar)
+	template <typename T>
+	FORCEINLINE FArchive& operator<<(FArchive& Ar, FChaosPhysicsMaterial& Value)
 	{
-		Ar << Friction << Restitution << SleepingLinearThreshold << SleepingAngularThreshold << DisabledLinearThreshold << DisabledAngularThreshold;
+		Value.Serialize(Ar);
+		return Ar;
 	}
-};
-
-
-template <typename T>
-FORCEINLINE FArchive& operator<<(FArchive& Ar, TChaosPhysicsMaterial<T>& Value)
-{
-	Value.Serialize(Ar);
-	return Ar;
-}
 }
\ No newline at end of file
diff --git a/Engine/Source/Runtime/Experimental/Chaos/Public/Chaos/PBDCollisionConstraint.h b/Engine/Source/Runtime/Experimental/Chaos/Public/Chaos/PBDCollisionConstraint.h
--- a/Engine/Source/Runtime/Experimental/Chaos/Public/Chaos/PBDCollisionConstraint.h
+++ b/Engine/Source/Runtime/Experimental/Chaos/Public/Chaos/PBDCollisionConstraint.h
@@ -87,7 +87,7 @@ public:
 	using FRigidBodyContactConstraint = TRigidBodyContactConstraint<T, d>;
 	using FConstraintHandleID = TSymmetricPairKey<const TGeometryParticleHandle<T, d>*>;
 
-	TPBDCollisionConstraint(const TPBDRigidsSOAs<T,d>& InParticles, TArrayCollectionArray<bool>& Collided, const TArrayCollectionArray<TSerializablePtr<TChaosPhysicsMaterial<T>>>& PerParticleMaterials, const int32 PairIterations = 1, const T Thickness = (T)0);
+	TPBDCollisionConstraint(const TPBDRigidsSOAs<T,d>& InParticles, TArrayCollectionArray<bool>& Collided, const TArrayCollectionArray<TSerializablePtr<FChaosPhysicsMaterial>>& PerParticleMaterials, const int32 PairIterations = 1, const T Thickness = (T)0);
 	virtual ~TPBDCollisionConstraint() {}
 
 	//
@@ -270,7 +270,7 @@ private:
 
 	TArray<FRigidBodyContactConstraint> Constraints;
 	TArrayCollectionArray<bool>& MCollided;
-	const TArrayCollectionArray<TSerializablePtr<TChaosPhysicsMaterial<T>>>& MPhysicsMaterials;
+	const TArrayCollectionArray<TSerializablePtr<FChaosPhysicsMaterial>>& MPhysicsMaterials;
 	bool bEnableVelocitySolve;
 	int32 MPairIterations;
 	T MThickness;
diff --git a/Engine/Source/Runtime/Experimental/Chaos/Public/Chaos/PBDCollisionConstraintPGS.h b/Engine/Source/Runtime/Experimental/Chaos/Public/Chaos/PBDCollisionConstraintPGS.h
--- a/Engine/Source/Runtime/Experimental/Chaos/Public/Chaos/PBDCollisionConstraintPGS.h
+++ b/Engine/Source/Runtime/Experimental/Chaos/Public/Chaos/PBDCollisionConstraintPGS.h
@@ -28,7 +28,7 @@ class TPBDCollisionConstraintPGS
 	typedef TRigidBodyContactConstraintPGS<T, d> FRigidBodyContactConstraint;
 	typedef TPBDConstraintGraph<T, d> FConstraintGraph;
 
-	TPBDCollisionConstraintPGS(TPBDRigidParticles<T, d>& InParticles, const TArray<int32>& InIndices, TArrayCollectionArray<bool>& Collided, TArrayCollectionArray<TSerializablePtr<TChaosPhysicsMaterial<T>>>& InPhysicsMaterials, const T Thickness = (T)0);
+	TPBDCollisionConstraintPGS(TPBDRigidParticles<T, d>& InParticles, const TArray<int32>& InIndices, TArrayCollectionArray<bool>& Collided, TArrayCollectionArray<TSerializablePtr<FChaosPhysicsMaterial>>& InPhysicsMaterials, const T Thickness = (T)0);
 	virtual ~TPBDCollisionConstraintPGS() {}
 
 	void ComputeConstraints(const TPBDRigidParticles<T, d>& InParticles, const TArray<int32>& InIndices, const T Dt);
@@ -81,7 +81,7 @@ class TPBDCollisionConstraintPGS
 
 	TArray<FRigidBodyContactConstraint> Constraints;
 	TArrayCollectionArray<bool>& MCollided;
-	const TArrayCollectionArray<TSerializablePtr<TChaosPhysicsMaterial<T>>>& MPhysicsMaterials;
+	const TArrayCollectionArray<TSerializablePtr<FChaosPhysicsMaterial>>& MPhysicsMaterials;
 	const T MThickness;
 	const T Tolerance;
 	const T MaxIterations;
diff --git a/Engine/Source/Runtime/Experimental/Chaos/Public/Chaos/PBDConstraintGraph.h b/Engine/Source/Runtime/Experimental/Chaos/Public/Chaos/PBDConstraintGraph.h
--- a/Engine/Source/Runtime/Experimental/Chaos/Public/Chaos/PBDConstraintGraph.h
+++ b/Engine/Source/Runtime/Experimental/Chaos/Public/Chaos/PBDConstraintGraph.h
@@ -94,7 +94,7 @@ namespace Chaos
 		/**
 		 * Put particles in inactive islands to sleep.
 		 */
-		bool SleepInactive(const int32 Island, const TArrayCollectionArray<TSerializablePtr<TChaosPhysicsMaterial<FReal>>>& PerParticleMaterialAttributes);
+		bool SleepInactive(const int32 Island, const TArrayCollectionArray<TSerializablePtr<FChaosPhysicsMaterial>>& PerParticleMaterialAttributes);
 
 		/**
 		 * Wake all particles in an Island.
diff --git a/Engine/Source/Runtime/Experimental/Chaos/Public/Chaos/PBDRigidsEvolution.h b/Engine/Source/Runtime/Experimental/Chaos/Public/Chaos/PBDRigidsEvolution.h
--- a/Engine/Source/Runtime/Experimental/Chaos/Public/Chaos/PBDRigidsEvolution.h
+++ b/Engine/Source/Runtime/Experimental/Chaos/Public/Chaos/PBDRigidsEvolution.h
@@ -371,8 +371,8 @@ class TPBDRigidsEvolutionBase
 	const auto& GetActiveClusteredArray() const { return Particles.GetActiveClusteredArray(); }
 	const auto& GetNonDisabledClusteredArray() const { return Particles.GetNonDisabledClusteredArray(); }
 
-	CHAOS_API TSerializablePtr<TChaosPhysicsMaterial<T>> GetPhysicsMaterial(const TGeometryParticleHandle<T, d>* Particle) const { return Particle->AuxilaryValue(PhysicsMaterials); }
-	CHAOS_API void SetPhysicsMaterial(TGeometryParticleHandle<T,d>* Particle, TSerializablePtr<TChaosPhysicsMaterial<T>> InMaterial)
+	CHAOS_API TSerializablePtr<FChaosPhysicsMaterial> GetPhysicsMaterial(const TGeometryParticleHandle<T, d>* Particle) const { return Particle->AuxilaryValue(PhysicsMaterials); }
+	CHAOS_API void SetPhysicsMaterial(TGeometryParticleHandle<T,d>* Particle, TSerializablePtr<FChaosPhysicsMaterial> InMaterial)
 	{
 		check(!Particle->AuxilaryValue(PerParticlePhysicsMaterials)); //shouldn't be setting non unique material if a unique one already exists
 		Particle->AuxilaryValue(PhysicsMaterials) = InMaterial;
@@ -605,8 +605,8 @@ protected:
 	FKinematicUpdateRule KinematicUpdate;
 	TArray<FPBDConstraintGraphRule*> ConstraintRules;
 	FPBDConstraintGraph ConstraintGraph;
-	TArrayCollectionArray<TSerializablePtr<TChaosPhysicsMaterial<T>>> PhysicsMaterials;
-	TArrayCollectionArray<TUniquePtr<TChaosPhysicsMaterial<T>>> PerParticlePhysicsMaterials;
+	TArrayCollectionArray<TSerializablePtr<FChaosPhysicsMaterial>> PhysicsMaterials;
+	TArrayCollectionArray<TUniquePtr<FChaosPhysicsMaterial>> PerParticlePhysicsMaterials;
 	TArrayCollectionArray<int32> ParticleDisableCount;
 	TArrayCollectionArray<bool> Collided;
 
diff --git a/Engine/Source/Runtime/Experimental/ChaosCore/Public/Chaos/Array.h b/Engine/Source/Runtime/Experimental/ChaosCore/Public/Chaos/Array.h
--- a/Engine/Source/Runtime/Experimental/ChaosCore/Public/Chaos/Array.h
+++ b/Engine/Source/Runtime/Experimental/ChaosCore/Public/Chaos/Array.h
@@ -1,7 +1,6 @@
 // Copyright 1998-2019 Epic Games, Inc. All Rights Reserved.
 #pragma once
 
-#include "Chaos/Defines.h"
 #include "Templates/TypeHash.h"
 
 #include <functional>
diff --git a/Engine/Source/Runtime/Experimental/ChaosCore/Public/Chaos/Vector.h b/Engine/Source/Runtime/Experimental/ChaosCore/Public/Chaos/Vector.h
--- a/Engine/Source/Runtime/Experimental/ChaosCore/Public/Chaos/Vector.h
+++ b/Engine/Source/Runtime/Experimental/ChaosCore/Public/Chaos/Vector.h
@@ -8,7 +8,6 @@
 #endif
 
 #include "Chaos/Array.h"
-#include "Chaos/Defines.h"
 #include "Chaos/Pair.h"
 
 #include "Containers/StaticArray.h"
diff --git a/Engine/Source/Runtime/Experimental/ChaosSolvers/Private/PhysicsProxy/FieldSystemPhysicsProxy.cpp b/Engine/Source/Runtime/Experimental/ChaosSolvers/Private/PhysicsProxy/FieldSystemPhysicsProxy.cpp
--- a/Engine/Source/Runtime/Experimental/ChaosSolvers/Private/PhysicsProxy/FieldSystemPhysicsProxy.cpp
+++ b/Engine/Source/Runtime/Experimental/ChaosSolvers/Private/PhysicsProxy/FieldSystemPhysicsProxy.cpp
@@ -457,10 +457,10 @@ void FFieldSystemPhysicsProxy::FieldParameterUpdateCallback(Chaos::FPhysicsSolve
 						for (const ContextIndex Index : IndicesArray)
 						{
 							const SolverObjectWrapper& ParticleObjectWrapper = SolverObjectMapping[Index.Result];
-							TSerializablePtr<Chaos::TChaosPhysicsMaterial<float>> Material = CurrentSolver->GetPhysicsMaterial(Index.Result);
+							TSerializablePtr<Chaos::FChaosPhysicsMaterial> Material = CurrentSolver->GetPhysicsMaterial(Index.Result);
 							if (ensure(Material) && ParticleObjectWrapper.SolverObject)
 							{
-								const TUniquePtr<Chaos::TChaosPhysicsMaterial<float>>& InstanceMaterial = CurrentSolver->GetPerParticlePhysicsMaterial(Index.Result);
+								const TUniquePtr<Chaos::FChaosPhysicsMaterial>& InstanceMaterial = CurrentSolver->GetPerParticlePhysicsMaterial(Index.Result);
 								if (InstanceMaterial)
 								{
 									Results[Index.Result] = InstanceMaterial->SleepingLinearThreshold;
@@ -483,7 +483,7 @@ void FFieldSystemPhysicsProxy::FieldParameterUpdateCallback(Chaos::FPhysicsSolve
 						{
 							const int32 i = Index.Result;
 							const PhysicsProxyWrapper& ParticleObjectWrapper = PhysicsProxyMapping[i];
-							TSerializablePtr<TChaosPhysicsMaterial<float>> Material = CurrentSolver->GetPhysicsMaterial(i);
+							TSerializablePtr<FChaosPhysicsMaterial> Material = CurrentSolver->GetPhysicsMaterial(i);
 							if (!ensure(Material) || !ParticleObjectWrapper.PhysicsProxy)	//question: do we actually need to check for solver object?
 							{
 								continue;
@@ -495,9 +495,9 @@ void FFieldSystemPhysicsProxy::FieldParameterUpdateCallback(Chaos::FPhysicsSolve
 								if (Results[i] != Material->SleepingLinearThreshold)
 								{
 									// value changed from shared material, make unique material.
-									TUniquePtr<Chaos::TChaosPhysicsMaterial<float>> NewMaterial = MakeUnique< Chaos::TChaosPhysicsMaterial<float>>(*Material);
+									TUniquePtr<Chaos::FChaosPhysicsMaterial> NewMaterial = MakeUnique< Chaos::FChaosPhysicsMaterial>(*Material);
 									CurrentSolver->SetPerParticlePhysicsMaterial(Index.Result, MoveTemp(NewMaterial));
-									const TUniquePtr<TChaosPhysicsMaterial<float>>& InstanceMaterial = CurrentSolver->GetPerParticlePhysicsMaterial(i);
+									const TUniquePtr<FChaosPhysicsMaterial>& InstanceMaterial = CurrentSolver->GetPerParticlePhysicsMaterial(i);
 
 									InstanceMaterial->SleepingLinearThreshold = Results[i];
 									InstanceMaterial->SleepingAngularThreshold = Results[i];
@@ -505,7 +505,7 @@ void FFieldSystemPhysicsProxy::FieldParameterUpdateCallback(Chaos::FPhysicsSolve
 							}
 							else
 							{
-								const TUniquePtr<TChaosPhysicsMaterial<float>>& InstanceMaterial = CurrentSolver->GetPerParticlePhysicsMaterial(i);
+								const TUniquePtr<FChaosPhysicsMaterial>& InstanceMaterial = CurrentSolver->GetPerParticlePhysicsMaterial(i);
 
 								if (InstanceMaterial->SleepingLinearThreshold != Results[i])
 								{
@@ -548,10 +548,10 @@ void FFieldSystemPhysicsProxy::FieldParameterUpdateCallback(Chaos::FPhysicsSolve
 						for (const ContextIndex Index : IndicesArray)
 						{
 							const SolverObjectWrapper& ParticleObjectWrapper = SolverObjectMapping[Index.Result];
-							TSerializablePtr<Chaos::TChaosPhysicsMaterial<float>> Material = CurrentSolver->GetPhysicsMaterial(Index.Result);
+							TSerializablePtr<Chaos::FChaosPhysicsMaterial> Material = CurrentSolver->GetPhysicsMaterial(Index.Result);
 							if (ensure(Material) && ParticleObjectWrapper.SolverObject)
 							{
-								const TUniquePtr<Chaos::TChaosPhysicsMaterial<float>>& InstanceMaterial = CurrentSolver->GetPerParticlePhysicsMaterial(Index.Result);
+								const TUniquePtr<Chaos::FChaosPhysicsMaterial>& InstanceMaterial = CurrentSolver->GetPerParticlePhysicsMaterial(Index.Result);
 								if (InstanceMaterial)
 								{
 									Results[Index.Result] = InstanceMaterial->DisabledLinearThreshold;
@@ -574,7 +574,7 @@ void FFieldSystemPhysicsProxy::FieldParameterUpdateCallback(Chaos::FPhysicsSolve
 						{
 							const int32 i = Index.Result;
 							const PhysicsProxyWrapper& ParticleObjectWrapper = PhysicsProxyMapping[i];
-							TSerializablePtr<Chaos::TChaosPhysicsMaterial<float>> Material = CurrentSolver->GetPhysicsMaterial(i);
+							TSerializablePtr<Chaos::FChaosPhysicsMaterial> Material = CurrentSolver->GetPhysicsMaterial(i);
 							if (!ensure(Material) || !ParticleObjectWrapper.PhysicsProxy)	//question: do we actually need to check for solver object?
 							{
 								continue;
@@ -586,9 +586,9 @@ void FFieldSystemPhysicsProxy::FieldParameterUpdateCallback(Chaos::FPhysicsSolve
 								if (Results[i] != Material->DisabledLinearThreshold)
 								{
 									// value changed from shared material, make unique material.
-									TUniquePtr<Chaos::TChaosPhysicsMaterial<float>> NewMaterial = MakeUnique< Chaos::TChaosPhysicsMaterial<float>>(*Material);
+									TUniquePtr<Chaos::FChaosPhysicsMaterial> NewMaterial = MakeUnique< Chaos::FChaosPhysicsMaterial>(*Material);
 									CurrentSolver->SetPerParticlePhysicsMaterial(Index.Result, MoveTemp(NewMaterial));
-									const TUniquePtr<TChaosPhysicsMaterial<float>>& InstanceMaterial = CurrentSolver->GetPerParticlePhysicsMaterial(i);
+									const TUniquePtr<FChaosPhysicsMaterial>& InstanceMaterial = CurrentSolver->GetPerParticlePhysicsMaterial(i);
 
 									InstanceMaterial->DisabledLinearThreshold = Results[i];
 									InstanceMaterial->DisabledAngularThreshold = Results[i];
@@ -596,7 +596,7 @@ void FFieldSystemPhysicsProxy::FieldParameterUpdateCallback(Chaos::FPhysicsSolve
 							}
 							else
 							{
-								const TUniquePtr<TChaosPhysicsMaterial<float>>& InstanceMaterial = CurrentSolver->GetPerParticlePhysicsMaterial(i);
+								const TUniquePtr<FChaosPhysicsMaterial>& InstanceMaterial = CurrentSolver->GetPerParticlePhysicsMaterial(i);
 
 								if (InstanceMaterial->DisabledLinearThreshold != Results[i])
 								{
diff --git a/Engine/Source/Runtime/Experimental/ChaosSolvers/Public/PhysicsProxy/SkeletalMeshPhysicsProxy.h b/Engine/Source/Runtime/Experimental/ChaosSolvers/Public/PhysicsProxy/SkeletalMeshPhysicsProxy.h
--- a/Engine/Source/Runtime/Experimental/ChaosSolvers/Public/PhysicsProxy/SkeletalMeshPhysicsProxy.h
+++ b/Engine/Source/Runtime/Experimental/ChaosSolvers/Public/PhysicsProxy/SkeletalMeshPhysicsProxy.h
@@ -64,7 +64,7 @@ struct CHAOSSOLVERS_API FSkeletalMeshPhysicsProxyParams
 	FVector InitialLinearVelocity;
 	FVector InitialAngularVelocity;
 
-	Chaos::TSerializablePtr<Chaos::TChaosPhysicsMaterial<float>> PhysicalMaterial;	// @todo(ccaulfield): should be per-shape
+	Chaos::TSerializablePtr<Chaos::FChaosPhysicsMaterial> PhysicalMaterial;	// @todo(ccaulfield): should be per-shape
 	EObjectStateTypeEnum ObjectType;												// @todo(ccaulfield): should be per-body
 
 	ECollisionTypeEnum CollisionType;
diff --git a/Engine/Source/Runtime/Experimental/ChaosSolvers/Public/PhysicsProxy/StaticMeshPhysicsProxy.h b/Engine/Source/Runtime/Experimental/ChaosSolvers/Public/PhysicsProxy/StaticMeshPhysicsProxy.h
--- a/Engine/Source/Runtime/Experimental/ChaosSolvers/Public/PhysicsProxy/StaticMeshPhysicsProxy.h
+++ b/Engine/Source/Runtime/Experimental/ChaosSolvers/Public/PhysicsProxy/StaticMeshPhysicsProxy.h
@@ -60,7 +60,7 @@ public:
 		EImplicitTypeEnum ShapeType;
 		bool bSimulating;
 		FTransform* TargetTransform;
-		Chaos::TSerializablePtr<Chaos::TChaosPhysicsMaterial<float>> PhysicalMaterial;
+		Chaos::TSerializablePtr<Chaos::FChaosPhysicsMaterial> PhysicalMaterial;
 		float Mass;
 		int32 MinRes;
 		int32 MaxRes;
diff --git a/Engine/Source/Runtime/Experimental/GeometryCollectionEngine/Private/GeometryCollection/GeometryCollectionComponent.cpp b/Engine/Source/Runtime/Experimental/GeometryCollectionEngine/Private/GeometryCollection/GeometryCollectionComponent.cpp
--- a/Engine/Source/Runtime/Experimental/GeometryCollectionEngine/Private/GeometryCollection/GeometryCollectionComponent.cpp
+++ b/Engine/Source/Runtime/Experimental/GeometryCollectionEngine/Private/GeometryCollection/GeometryCollectionComponent.cpp
@@ -141,7 +141,7 @@ UGeometryCollectionComponent::UGeometryCollectionComponent(const FObjectInitiali
 	GlobalNavMeshInvalidationCounter += 3;
 	NavmeshInvalidationTimeSliceIndex = GlobalNavMeshInvalidationCounter;
 
-	ChaosMaterial = MakeUnique<Chaos::TChaosPhysicsMaterial<float>>();
+	ChaosMaterial = MakeUnique<Chaos::FChaosPhysicsMaterial>();
 
 	WorldBounds = FBoxSphereBounds(FBox(ForceInit));	
 
diff --git a/Engine/Source/Runtime/Experimental/GeometryCollectionEngine/Private/GeometryCollection/SkeletalMeshSimulationComponent.cpp b/Engine/Source/Runtime/Experimental/GeometryCollectionEngine/Private/GeometryCollection/SkeletalMeshSimulationComponent.cpp
--- a/Engine/Source/Runtime/Experimental/GeometryCollectionEngine/Private/GeometryCollection/SkeletalMeshSimulationComponent.cpp
+++ b/Engine/Source/Runtime/Experimental/GeometryCollectionEngine/Private/GeometryCollection/SkeletalMeshSimulationComponent.cpp
@@ -65,7 +65,7 @@ USkeletalMeshSimulationComponent::USkeletalMeshSimulationComponent(const FObject
 {
 	// Enable calls to TickComponent()
 	UActorComponent::PrimaryComponentTick.bCanEverTick = true;	
-	ChaosMaterial = MakeUnique<Chaos::TChaosPhysicsMaterial<float>>();
+	ChaosMaterial = MakeUnique<Chaos::FChaosPhysicsMaterial>();
 }
 
 Chaos::FPhysicsSolver* GetSolver(const USkeletalMeshSimulationComponent& SkeletalMeshSimulationComponent)
diff --git a/Engine/Source/Runtime/Experimental/GeometryCollectionEngine/Private/GeometryCollection/StaticMeshSimulationComponent.cpp b/Engine/Source/Runtime/Experimental/GeometryCollectionEngine/Private/GeometryCollection/StaticMeshSimulationComponent.cpp
--- a/Engine/Source/Runtime/Experimental/GeometryCollectionEngine/Private/GeometryCollection/StaticMeshSimulationComponent.cpp
+++ b/Engine/Source/Runtime/Experimental/GeometryCollectionEngine/Private/GeometryCollection/StaticMeshSimulationComponent.cpp
@@ -38,7 +38,7 @@ UStaticMeshSimulationComponent::UStaticMeshSimulationComponent(const FObjectInit
 	PrimaryComponentTick.bStartWithTickEnabled = true;
 	PrimaryComponentTick.bCanEverTick = true;
 	PrimaryComponentTick.SetTickFunctionEnable(true);
-	ChaosMaterial = MakeUnique<Chaos::TChaosPhysicsMaterial<float>>();
+	ChaosMaterial = MakeUnique<Chaos::FChaosPhysicsMaterial>();
 }
 
 // We tick to detect components that unreal has moved, so we can update the solver.
diff --git a/Engine/Source/Runtime/Experimental/GeometryCollectionEngine/Public/GeometryCollection/GeometryCollectionComponent.h b/Engine/Source/Runtime/Experimental/GeometryCollectionEngine/Public/GeometryCollection/GeometryCollectionComponent.h
--- a/Engine/Source/Runtime/Experimental/GeometryCollectionEngine/Public/GeometryCollection/GeometryCollectionComponent.h
+++ b/Engine/Source/Runtime/Experimental/GeometryCollectionEngine/Public/GeometryCollection/GeometryCollectionComponent.h
@@ -607,7 +607,7 @@ private:
 	TArray<bool> EventsPlayed;
 
 	//@todo(mlentine): Don't have one per geo collection
-	TUniquePtr<Chaos::TChaosPhysicsMaterial<float>> ChaosMaterial;
+	TUniquePtr<Chaos::FChaosPhysicsMaterial> ChaosMaterial;
 
 	FGeometryCollectionPhysicsProxy* PhysicsProxy;
 	TUniquePtr<FGeometryDynamicCollection> DynamicCollection;
diff --git a/Engine/Source/Runtime/Experimental/GeometryCollectionEngine/Public/GeometryCollection/SkeletalMeshSimulationComponent.h b/Engine/Source/Runtime/Experimental/GeometryCollectionEngine/Public/GeometryCollection/SkeletalMeshSimulationComponent.h
--- a/Engine/Source/Runtime/Experimental/GeometryCollectionEngine/Public/GeometryCollection/SkeletalMeshSimulationComponent.h
+++ b/Engine/Source/Runtime/Experimental/GeometryCollectionEngine/Public/GeometryCollection/SkeletalMeshSimulationComponent.h
@@ -172,5 +172,5 @@ private:
 	FSkeletalMeshPhysicsProxy* PhysicsProxy;
 
 	//@todo(mlentine): Don't have one per static mesh
-	TUniquePtr<Chaos::TChaosPhysicsMaterial<float>> ChaosMaterial;
+	TUniquePtr<Chaos::FChaosPhysicsMaterial> ChaosMaterial;
 };
diff --git a/Engine/Source/Runtime/Experimental/GeometryCollectionEngine/Public/GeometryCollection/StaticMeshSimulationComponent.h b/Engine/Source/Runtime/Experimental/GeometryCollectionEngine/Public/GeometryCollection/StaticMeshSimulationComponent.h
--- a/Engine/Source/Runtime/Experimental/GeometryCollectionEngine/Public/GeometryCollection/StaticMeshSimulationComponent.h
+++ b/Engine/Source/Runtime/Experimental/GeometryCollectionEngine/Public/GeometryCollection/StaticMeshSimulationComponent.h
@@ -115,7 +115,7 @@ private :
 
 
 	//@todo(mlentine): Don't have one per static mesh
-	TUniquePtr<Chaos::TChaosPhysicsMaterial<float>> ChaosMaterial;
+	TUniquePtr<Chaos::FChaosPhysicsMaterial> ChaosMaterial;
 	
 protected:
 
diff --git a/Engine/Source/Runtime/Experimental/GeometryCollectionSimulationCore/Public/GeometryCollection/GeometryCollectionSimulationCoreTypes.h b/Engine/Source/Runtime/Experimental/GeometryCollectionSimulationCore/Public/GeometryCollection/GeometryCollectionSimulationCoreTypes.h
--- a/Engine/Source/Runtime/Experimental/GeometryCollectionSimulationCore/Public/GeometryCollection/GeometryCollectionSimulationCoreTypes.h
+++ b/Engine/Source/Runtime/Experimental/GeometryCollectionSimulationCore/Public/GeometryCollection/GeometryCollectionSimulationCoreTypes.h
@@ -291,7 +291,7 @@ struct FSimulationParameters
 	float ReverseCacheBeginTime;
 	bool bClearCache;
 
-	Chaos::TSerializablePtr<Chaos::TChaosPhysicsMaterial<float>> PhysicalMaterial;
+	Chaos::TSerializablePtr<Chaos::FChaosPhysicsMaterial> PhysicalMaterial;
 
 	FCollisionDataSimulationParameters CollisionData;
 	FBreakingDataSimulationParameters BreakingData;
