Convert geometry collection example utils to new material type.

diff --git a/Engine/Plugins/Experimental/GeometryCollectionPlugin/Source/GeometryCollectionExampleCore/Private/GeometryCollection/GeometryCollectionExampleUtility.cpp b/Engine/Plugins/Experimental/GeometryCollectionPlugin/Source/GeometryCollectionExampleCore/Private/GeometryCollection/GeometryCollectionExampleUtility.cpp
--- a/Engine/Plugins/Experimental/GeometryCollectionPlugin/Source/GeometryCollectionExampleCore/Private/GeometryCollection/GeometryCollectionExampleUtility.cpp
+++ b/Engine/Plugins/Experimental/GeometryCollectionPlugin/Source/GeometryCollectionExampleCore/Private/GeometryCollection/GeometryCollectionExampleUtility.cpp
@@ -203,8 +203,7 @@ namespace GeometryCollectionExample {
 		return RestCollection;
 	}
 
-	template<class T>
-	void InitMaterialToZero(TUniquePtr<Chaos::TChaosPhysicsMaterial<T>> const &PhysicalMaterial)
+	void InitMaterialToZero(TUniquePtr<Chaos::FChaosPhysicsMaterial> const &PhysicalMaterial)
 	{
 		PhysicalMaterial->Friction = 0;
 		PhysicalMaterial->Restitution = 0;
@@ -213,12 +212,9 @@ namespace GeometryCollectionExample {
 		PhysicalMaterial->DisabledLinearThreshold = 0;
 		PhysicalMaterial->DisabledAngularThreshold = 0;
 	}
-	template void InitMaterialToZero(TUniquePtr<Chaos::TChaosPhysicsMaterial<float>> const &PhysicalMaterial);
 
-
-	template<class T>
 	void InitCollections(
-		TUniquePtr<Chaos::TChaosPhysicsMaterial<T>> &PhysicalMaterial,
+		TUniquePtr<Chaos::FChaosPhysicsMaterial> &PhysicalMaterial,
 		TSharedPtr<FGeometryCollection>& RestCollection,
 		TSharedPtr<FGeometryDynamicCollection>& DynamicCollection,
 		InitCollectionsParameters& InitParams
@@ -227,7 +223,7 @@ namespace GeometryCollectionExample {
 		// Allow for customized initialization of these objects in the calling function. 
 		if (PhysicalMaterial == nullptr)
 		{
-			PhysicalMaterial = MakeUnique<Chaos::TChaosPhysicsMaterial<T>>();
+			PhysicalMaterial = MakeUnique<Chaos::FChaosPhysicsMaterial>();
 			InitMaterialToZero(PhysicalMaterial);
 		}
 
@@ -246,16 +242,9 @@ namespace GeometryCollectionExample {
 			DynamicCollection = GeometryCollectionToGeometryDynamicCollection(RestCollection.Get(), InitParams.DynamicStateDefault);
 		}
 	}
-	template void InitCollections(
-		TUniquePtr<Chaos::TChaosPhysicsMaterial<float>> &PhysicalMaterial,
-		TSharedPtr<FGeometryCollection>& RestCollection,
-		TSharedPtr<FGeometryDynamicCollection>& DynamicCollection,
-		InitCollectionsParameters& InitParams
-	);
 
-	template<class T>
 	FGeometryCollectionPhysicsProxy* RigidBodySetup(
-		TUniquePtr<Chaos::TChaosPhysicsMaterial<T>> &PhysicalMaterial,
+		TUniquePtr<Chaos::FChaosPhysicsMaterial> &PhysicalMaterial,
 		TSharedPtr<FGeometryCollection>& RestCollection,
 		TSharedPtr<FGeometryDynamicCollection>& DynamicCollection,
 		FInitFunc CustomFunc
@@ -286,12 +275,6 @@ namespace GeometryCollectionExample {
 		return nullptr;
 #endif
 	}
-	template FGeometryCollectionPhysicsProxy* RigidBodySetup(
-		TUniquePtr<Chaos::TChaosPhysicsMaterial<float>> & PhysicalMaterial,
-		TSharedPtr<FGeometryCollection>& RestCollection,
-		TSharedPtr<FGeometryDynamicCollection>& DynamicCollection,
-		FInitFunc CustomFunc
-	);
 
 } // end namespace GeometryCollectionExample
 
diff --git a/Engine/Plugins/Experimental/GeometryCollectionPlugin/Source/GeometryCollectionExampleCore/Public/GeometryCollection/GeometryCollectionExampleUtility.h b/Engine/Plugins/Experimental/GeometryCollectionPlugin/Source/GeometryCollectionExampleCore/Public/GeometryCollection/GeometryCollectionExampleUtility.h
--- a/Engine/Plugins/Experimental/GeometryCollectionPlugin/Source/GeometryCollectionExampleCore/Public/GeometryCollection/GeometryCollectionExampleUtility.h
+++ b/Engine/Plugins/Experimental/GeometryCollectionPlugin/Source/GeometryCollectionExampleCore/Public/GeometryCollection/GeometryCollectionExampleUtility.h
@@ -32,8 +32,7 @@ namespace GeometryCollectionExample {
 	TSharedPtr<FGeometryCollection>	CreateClusteredBody_ThreeByTwo_ThreeTransform(FVector Position);
 	TSharedPtr<FGeometryCollection>	CreateClusteredBody_FracturedGeometry(FVector Position = FVector(0));
 
-	template<class T>
-	void InitMaterialToZero(TUniquePtr<Chaos::TChaosPhysicsMaterial<T>> const &PhysicalMaterial);
+	void InitMaterialToZero(TUniquePtr<Chaos::FChaosPhysicsMaterial> const &PhysicalMaterial);
 
 	using FCollectionInitFunc = TFunction<void(TSharedPtr<FGeometryCollection>&)>;
 	struct InitCollectionsParameters
@@ -44,18 +43,16 @@ namespace GeometryCollectionExample {
 		int DynamicStateDefault;
 	};
 
-	template<class T>
 	void InitCollections(
-		TUniquePtr<Chaos::TChaosPhysicsMaterial<T>> &PhysicalMaterial,
+		TUniquePtr<Chaos::FChaosPhysicsMaterial> &PhysicalMaterial,
 		TSharedPtr<FGeometryCollection>& RestCollection,
 		TSharedPtr<FGeometryDynamicCollection>& DynamicCollection,
 		InitCollectionsParameters& InitParams
 	);
 
 	using FInitFunc = TFunction<void(FSimulationParameters&)>;
-	template<class T>
 	FGeometryCollectionPhysicsProxy* RigidBodySetup(
-		TUniquePtr<Chaos::TChaosPhysicsMaterial<T>> &PhysicalMaterial,
+		TUniquePtr<Chaos::FChaosPhysicsMaterial> &PhysicalMaterial,
 		TSharedPtr<FGeometryCollection> &RestCollection,
 		TSharedPtr<FGeometryDynamicCollection> &DynamicCollection,
 		FInitFunc CustomFunc = nullptr
