Integrate 9553110 from //UE4/DevPhysics to //UE4/Release-4.24.

diff --git a/Engine/Plugins/Experimental/GeometryCollectionPlugin/Source/GeometryCollectionEditor/Private/GeometryCollection/ActorFactoryGeometryCollection.cpp b/Engine/Plugins/Experimental/GeometryCollectionPlugin/Source/GeometryCollectionEditor/Private/GeometryCollection/ActorFactoryGeometryCollection.cpp
--- a/Engine/Plugins/Experimental/GeometryCollectionPlugin/Source/GeometryCollectionEditor/Private/GeometryCollection/ActorFactoryGeometryCollection.cpp
+++ b/Engine/Plugins/Experimental/GeometryCollectionPlugin/Source/GeometryCollectionEditor/Private/GeometryCollection/ActorFactoryGeometryCollection.cpp
@@ -4,6 +4,7 @@
 #include "GeometryCollection/GeometryCollectionActor.h"
 #include "GeometryCollection/GeometryCollection.h"
 #include "GeometryCollection/GeometryCollectionComponent.h"
+#include "GeometryCollection/GeometryCollectionFactory.h"
 #include "Misc/FileHelper.h"
 #include "ActorFactories/ActorFactory.h"
 #include "AssetData.h"
diff --git a/Engine/Plugins/Experimental/GeometryCollectionPlugin/Source/GeometryCollectionEditor/Private/GeometryCollection/GeometryCollectionCacheFactory.cpp b/Engine/Plugins/Experimental/GeometryCollectionPlugin/Source/GeometryCollectionEditor/Private/GeometryCollection/GeometryCollectionCacheFactory.cpp
new file mode 100644
--- /dev/null
+++ b/Engine/Plugins/Experimental/GeometryCollectionPlugin/Source/GeometryCollectionEditor/Private/GeometryCollection/GeometryCollectionCacheFactory.cpp
@@ -0,0 +1,86 @@
+// Copyright 1998-2019 Epic Games, Inc. All Rights Reserved.
+
+#include "GeometryCollection/GeometryCollectionCacheFactory.h"
+#include "GeometryCollection/GeometryCollectionCache.h"
+#include "GeometryCollection/GeometryCollectionObject.h"
+#include "IContentBrowserSingleton.h"
+#include "ContentBrowserModule.h"
+
+#include "Widgets/SWindow.h"
+#include "Widgets/Layout/SBorder.h"
+#include "EditorStyleSet.h"
+#include "Editor.h"
+
+#define LOCTEXT_NAMESPACE "GeomCollectionCacheFactory"
+
+UGeometryCollectionCacheFactory::UGeometryCollectionCacheFactory()
+{
+	SupportedClass = UGeometryCollectionCache::StaticClass();
+}
+
+bool UGeometryCollectionCacheFactory::CanCreateNew() const
+{
+	return true;
+}
+
+bool UGeometryCollectionCacheFactory::FactoryCanImport(const FString& Filename)
+{
+	return false;
+}
+
+bool UGeometryCollectionCacheFactory::ShouldShowInNewMenu() const
+{
+	return true;
+}
+
+UObject* UGeometryCollectionCacheFactory::FactoryCreateNew(UClass* InClass, UObject* InParent, FName InName, EObjectFlags Flags, UObject* Context, FFeedbackContext* Warn)
+{
+	if(TargetCollection)
+	{
+		UGeometryCollectionCache* NewCache = NewObject<UGeometryCollectionCache>(InParent, InClass, InName, Flags);
+
+		return NewCache;
+	}
+
+	return nullptr;
+}
+
+bool UGeometryCollectionCacheFactory::ConfigureProperties()
+{
+	TargetCollection = nullptr;
+
+	FContentBrowserModule& ContentBrowserModule = FModuleManager::LoadModuleChecked<FContentBrowserModule>("ContentBrowser");
+	
+	// Config for picker list of geometry collections
+	FAssetPickerConfig AssetPickerConfig;
+	AssetPickerConfig.Filter.ClassNames.Add(UGeometryCollection::StaticClass()->GetFName());
+	AssetPickerConfig.OnAssetSelected = FOnAssetSelected::CreateUObject(this, &UGeometryCollectionCacheFactory::OnConfigSelection);
+	AssetPickerConfig.InitialAssetViewType = EAssetViewType::List;
+
+	// Create and show the window
+	PickerWindow = SNew(SWindow)
+		.Title(LOCTEXT("CreateCacheOptions", "Pick Target Geometry Collection"))
+		.ClientSize(FVector2D(500, 600))
+		.SupportsMinimize(false)
+		.SupportsMaximize(false)
+		[
+			SNew(SBorder)
+			.BorderImage(FEditorStyle::GetBrush("Menu.Background"))
+			[
+				ContentBrowserModule.Get().CreateAssetPicker(AssetPickerConfig)
+			]
+		];
+
+	GEditor->EditorAddModalWindow(PickerWindow.ToSharedRef());
+	PickerWindow.Reset();
+
+	return !!TargetCollection;
+}
+
+void UGeometryCollectionCacheFactory::OnConfigSelection(const FAssetData& InSelectedAssetData)
+{
+	TargetCollection = Cast<UGeometryCollection>(InSelectedAssetData.GetAsset());
+	PickerWindow->RequestDestroyWindow();
+}
+
+#undef LOCTEXT_NAMESPACE
diff --git a/Engine/Plugins/Experimental/GeometryCollectionPlugin/Source/GeometryCollectionEditor/Private/GeometryCollection/GeometryCollectionConversion.cpp b/Engine/Plugins/Experimental/GeometryCollectionPlugin/Source/GeometryCollectionEditor/Private/GeometryCollection/GeometryCollectionConversion.cpp
--- a/Engine/Plugins/Experimental/GeometryCollectionPlugin/Source/GeometryCollectionEditor/Private/GeometryCollection/GeometryCollectionConversion.cpp
+++ b/Engine/Plugins/Experimental/GeometryCollectionPlugin/Source/GeometryCollectionEditor/Private/GeometryCollection/GeometryCollectionConversion.cpp
@@ -15,7 +15,7 @@
 #include "GeometryCollection/GeometryCollection.h"
 #include "GeometryCollection/GeometryCollectionActor.h"
 #include "GeometryCollection/GeometryCollectionAlgo.h"
-// #include "GeometryCollection/GeometryCollectionFactory.h"
+#include "GeometryCollection/GeometryCollectionFactory.h"
 #include "GeometryCollection/GeometryCollectionComponent.h"
 #include "GeometryCollection/GeometryCollectionClusteringUtility.h"
 #include "GeometryCollection/GeometryCollectionUtility.h"
@@ -431,3 +431,166 @@ void FGeometryCollectionConversion::AppendSkeletalMesh(const USkeletalMesh* Skel
 	}
 }
 
+void FGeometryCollectionConversion::CreateGeometryCollectionCommand(UWorld * World)
+{
+	UPackage* Package = CreatePackage(NULL, TEXT("/Game/GeometryCollectionAsset"));
+	auto GeometryCollectionFactory = NewObject<UGeometryCollectionFactory>();
+	UGeometryCollection* GeometryCollection = static_cast<UGeometryCollection*>(
+		GeometryCollectionFactory->FactoryCreateNew(UGeometryCollection::StaticClass(), Package,
+			FName("GeometryCollectionAsset"), RF_Standalone | RF_Public, NULL, GWarn));
+	FAssetRegistryModule::AssetCreated(GeometryCollection);
+	Package->SetDirtyFlag(true);
+}
+
+void FGeometryCollectionConversion::CreateFromSelectedActorsCommand(UWorld * World)
+{
+	//UE_LOG(UGeometryCollectionConversionLogging, Log, TEXT("FGeometryCollectionConversion::CreateCommand()"));
+	USelection* SelectedActors = GEditor->GetSelectedActors();
+	if (!SelectedActors)
+	{
+		return;
+	}
+
+	UPackage* Package = nullptr;
+	UGeometryCollection* GeometryCollection = nullptr;
+	for (FSelectionIterator Iter(*SelectedActors); Iter; ++Iter)
+	{
+		if (AActor* Actor = Cast<AActor>(*Iter))
+		{
+			TArray<UStaticMeshComponent *> StaticMeshComponents;
+			Actor->GetComponents<UStaticMeshComponent>(StaticMeshComponents);
+			if (StaticMeshComponents.Num() > 0)
+			{
+				for (int Index = 0; Index < StaticMeshComponents.Num(); Index++)
+				{
+					if (StaticMeshComponents[Index]->GetStaticMesh())
+					{
+						if (!Package || !GeometryCollection)
+						{
+							Package = CreatePackage(NULL, TEXT("/Game/GeometryCollectionAsset"));
+							auto GeometryCollectionFactory = NewObject<UGeometryCollectionFactory>();
+							GeometryCollection = static_cast<UGeometryCollection*>(
+								GeometryCollectionFactory->FactoryCreateNew(
+									UGeometryCollection::StaticClass(),
+									Package,
+									FName("GeometryCollectionAsset"),
+									RF_Standalone | RF_Public,
+									NULL,
+									GWarn));
+						}
+						FGeometryCollectionConversion::AppendStaticMesh(
+							StaticMeshComponents[Index]->GetStaticMesh(),
+							StaticMeshComponents[Index],
+							Actor->GetTransform(),
+							GeometryCollection);
+					}
+				}
+			}
+
+			TArray < USkeletalMeshComponent * > SkeletalMeshComponents;
+			Actor->GetComponents<USkeletalMeshComponent>(SkeletalMeshComponents);
+			if (SkeletalMeshComponents.Num() > 0)
+			{
+				for (int Index = 0; Index < SkeletalMeshComponents.Num(); Index++)
+				{
+					if (SkeletalMeshComponents[Index]->SkeletalMesh)
+					{
+						if (!Package || !GeometryCollection)
+						{
+							Package = CreatePackage(NULL, TEXT("/Game/GeometryCollectionAsset"));
+							auto GeometryCollectionFactory = NewObject<UGeometryCollectionFactory>();
+							GeometryCollection = static_cast<UGeometryCollection*>(
+								GeometryCollectionFactory->FactoryCreateNew(
+									UGeometryCollection::StaticClass(),
+									Package,
+									FName("GeometryCollectionAsset"),
+									RF_Standalone | RF_Public,
+									NULL,
+									GWarn));
+						}
+						FGeometryCollectionConversion::AppendSkeletalMesh(
+							SkeletalMeshComponents[Index]->SkeletalMesh,
+							SkeletalMeshComponents[Index],
+							Actor->GetTransform(),
+							GeometryCollection);
+					}
+				}
+			}
+		}
+	}
+
+	if(GeometryCollection != nullptr)
+	{
+		GeometryCollectionAlgo::PrepareForSimulation(GeometryCollection->GetGeometryCollection().Get());
+
+		FAssetRegistryModule::AssetCreated(GeometryCollection);
+		GeometryCollection->MarkPackageDirty();
+		Package->SetDirtyFlag(true);
+	}
+}
+
+void FGeometryCollectionConversion::CreateFromSelectedAssetsCommand(UWorld * World)
+{
+	//UE_LOG(UGeometryCollectionConversionLogging, Log, TEXT("FGeometryCollectionConversion::CreateCommand()"));
+	UPackage* Package = nullptr;
+	UGeometryCollection* GeometryCollection = nullptr;
+
+	TArray<FAssetData> SelectedAssets;
+	GEditor->GetContentBrowserSelections(SelectedAssets);
+	for (const FAssetData& AssetData : SelectedAssets)
+	{
+		if (AssetData.GetAsset()->IsA<UStaticMesh>())
+		{
+			UE_LOG(UGeometryCollectionConversionLogging, Log, TEXT("Static Mesh Content Browser : %s"), *AssetData.GetClass()->GetName());
+			if (!Package || !GeometryCollection)
+			{
+				Package = CreatePackage(NULL, TEXT("/Game/GeometryCollectionAsset"));
+				auto GeometryCollectionFactory = NewObject<UGeometryCollectionFactory>();
+				GeometryCollection = static_cast<UGeometryCollection*>(
+					GeometryCollectionFactory->FactoryCreateNew(
+						UGeometryCollection::StaticClass(), 
+						Package,
+						FName("GeometryCollectionAsset"), 
+						RF_Standalone | RF_Public, 
+						NULL, 
+						GWarn));
+			}
+			FGeometryCollectionConversion::AppendStaticMesh(
+				static_cast<const UStaticMesh *>(AssetData.GetAsset()), 
+				nullptr,
+				FTransform(), 
+				GeometryCollection);
+		}
+		else if (AssetData.GetAsset()->IsA<USkeletalMesh>())
+		{
+			UE_LOG(UGeometryCollectionConversionLogging, Log, TEXT("Skeletal Mesh Content Browser : %s"), *AssetData.GetClass()->GetName());
+			if (!Package || !GeometryCollection)
+			{
+				Package = CreatePackage(NULL, TEXT("/Game/GeometryCollectionAsset"));
+				auto GeometryCollectionFactory = NewObject<UGeometryCollectionFactory>();
+				GeometryCollection = static_cast<UGeometryCollection*>(
+					GeometryCollectionFactory->FactoryCreateNew(
+						UGeometryCollection::StaticClass(), 
+						Package,
+						FName("GeometryCollectionAsset"), 
+						RF_Standalone | RF_Public, 
+						NULL, 
+						GWarn));
+			}
+			FGeometryCollectionConversion::AppendSkeletalMesh(
+				static_cast<const USkeletalMesh *>(AssetData.GetAsset()), 
+				nullptr,
+				FTransform(),
+				GeometryCollection);
+		}
+	}
+
+	if(GeometryCollection != nullptr)
+	{
+		GeometryCollectionAlgo::PrepareForSimulation(GeometryCollection->GetGeometryCollection().Get());
+
+		FAssetRegistryModule::AssetCreated(GeometryCollection);
+		GeometryCollection->MarkPackageDirty();
+		Package->SetDirtyFlag(true);
+	}
+}
diff --git a/Engine/Plugins/Experimental/GeometryCollectionPlugin/Source/GeometryCollectionEditor/Private/GeometryCollection/GeometryCollectionEditorPlugin.cpp b/Engine/Plugins/Experimental/GeometryCollectionPlugin/Source/GeometryCollectionEditor/Private/GeometryCollection/GeometryCollectionEditorPlugin.cpp
--- a/Engine/Plugins/Experimental/GeometryCollectionPlugin/Source/GeometryCollectionEditor/Private/GeometryCollection/GeometryCollectionEditorPlugin.cpp
+++ b/Engine/Plugins/Experimental/GeometryCollectionPlugin/Source/GeometryCollectionEditor/Private/GeometryCollection/GeometryCollectionEditorPlugin.cpp
@@ -46,6 +46,18 @@ void IGeometryCollectionEditorPlugin::StartupModule()
 
 	if (GIsEditor && !IsRunningCommandlet())
 	{
+		EditorCommands.Add(IConsoleManager::Get().RegisterConsoleCommand(
+			TEXT("GeometryCollection.CreateFromSelectedActors"),
+			TEXT("Creates a GeometryCollection from the selected Actors that contain Skeletal and Statict Mesh Components"),
+			FConsoleCommandWithWorldDelegate::CreateStatic(&FGeometryCollectionConversion::CreateFromSelectedActorsCommand),
+			ECVF_Default
+		));
+		EditorCommands.Add(IConsoleManager::Get().RegisterConsoleCommand(
+			TEXT("GeometryCollection.CreateFromSelectedAssets"),
+			TEXT("Creates a GeometryCollection from the selected Skeletal Mesh and Static Mesh Assets"),
+			FConsoleCommandWithWorldDelegate::CreateStatic(&FGeometryCollectionConversion::CreateFromSelectedAssetsCommand),
+			ECVF_Default
+		));
 		EditorCommands.Add(IConsoleManager::Get().RegisterConsoleCommand(
 			TEXT("GeometryCollection.ToString"),
 			TEXT("Dump the contents of the collection to the log file. WARNING: The collection can be very large."),
diff --git a/Engine/Plugins/Experimental/GeometryCollectionPlugin/Source/GeometryCollectionEditor/Private/GeometryCollection/GeometryCollectionFactory.cpp b/Engine/Plugins/Experimental/GeometryCollectionPlugin/Source/GeometryCollectionEditor/Private/GeometryCollection/GeometryCollectionFactory.cpp
new file mode 100644
--- /dev/null
+++ b/Engine/Plugins/Experimental/GeometryCollectionPlugin/Source/GeometryCollectionEditor/Private/GeometryCollection/GeometryCollectionFactory.cpp
@@ -0,0 +1,195 @@
+// Copyright 1998-2019 Epic Games, Inc. All Rights Reserved.
+
+#include "GeometryCollection/GeometryCollectionFactory.h"
+
+#include "GeometryCollection/GeometryCollectionAlgo.h"
+#include "GeometryCollection/GeometryCollectionConversion.h"
+#include "GeometryCollection/GeometryCollectionObject.h"
+#include "GeometryCollection/GeometryCollection.h"
+#include "GeometryCollection/GeometryCollectionComponent.h"
+#include "Editor.h"
+#include "Editor/EditorEngine.h"
+#include "Engine/Selection.h"
+#include "Engine/StaticMesh.h"
+#include "Engine/SkeletalMesh.h"
+#include "Components/SkeletalMeshComponent.h"
+#include "Components/StaticMeshComponent.h"
+
+#define LOCTEXT_NAMESPACE "GeometryCollection"
+
+/////////////////////////////////////////////////////
+// GeometryCollectionFactory
+
+UGeometryCollectionFactory::UGeometryCollectionFactory(const FObjectInitializer& ObjectInitializer)
+	: Super(ObjectInitializer)
+{
+	bCreateNew = true;
+	bEditAfterNew = true;
+	SupportedClass = UGeometryCollection::StaticClass();
+}
+
+UGeometryCollection* UGeometryCollectionFactory::StaticFactoryCreateNew(UClass* Class, UObject* InParent, FName Name, EObjectFlags Flags, UObject* Context, FFeedbackContext* Warn)
+{
+	return static_cast<UGeometryCollection*>(NewObject<UGeometryCollection>(InParent, Class, Name, Flags | RF_Transactional | RF_Public | RF_Standalone));
+}
+
+void AppendActorComponentsRecursive(
+	const AActor* Actor,
+	TArray< GeometryCollectionStaticMeshConversionTuple >& StaticMeshList,
+	TArray< GeometryCollectionSkeletalMeshConversionTuple >& SkeletalMeshList,
+	TArray< GeometryCollectionTuple >& GeometryCollectionList)
+{
+	FTransform ActorTransform = Actor->GetTransform();
+
+	TArray<UStaticMeshComponent *> StaticMeshComponents;
+	Actor->GetComponents<UStaticMeshComponent>(StaticMeshComponents);
+	if (StaticMeshComponents.Num() > 0)
+	{
+		for (int Index = 0; Index < StaticMeshComponents.Num(); Index++)
+		{
+			if (StaticMeshComponents[Index]->GetStaticMesh())
+			{
+				StaticMeshList.Add(GeometryCollectionStaticMeshConversionTuple(
+					StaticMeshComponents[Index]->GetStaticMesh(),
+					StaticMeshComponents[Index],
+					StaticMeshComponents[Index]->GetComponentTransform()));
+			}
+		}
+	}
+
+	TArray < USkeletalMeshComponent * > SkeletalMeshComponents;
+	Actor->GetComponents<USkeletalMeshComponent>(SkeletalMeshComponents);
+	if (SkeletalMeshComponents.Num() > 0)
+	{
+		for (int Index = 0; Index < SkeletalMeshComponents.Num(); Index++)
+		{
+			if (SkeletalMeshComponents[Index]->SkeletalMesh)
+			{
+				SkeletalMeshList.Add(GeometryCollectionSkeletalMeshConversionTuple(
+					SkeletalMeshComponents[Index]->SkeletalMesh,
+					SkeletalMeshComponents[Index],
+					SkeletalMeshComponents[Index]->GetComponentTransform()));
+			}
+		}
+	}
+
+	TArray<UGeometryCollectionComponent *> GeometryCollectionComponents;
+	Actor->GetComponents<UGeometryCollectionComponent>(GeometryCollectionComponents);
+	if (GeometryCollectionComponents.Num() > 0)
+	{
+		for (int Index = 0; Index < GeometryCollectionComponents.Num(); Index++)
+		{
+			if (GeometryCollectionComponents[Index]->GetRestCollection())
+			{
+				GeometryCollectionList.Add(GeometryCollectionTuple(
+					GeometryCollectionComponents[Index]->GetRestCollection(),
+					GeometryCollectionComponents[Index],
+					GeometryCollectionComponents[Index]->GetComponentTransform()));
+			}
+		}
+	}
+
+	TArray<UChildActorComponent*> ChildActorComponents;
+	Actor->GetComponents<UChildActorComponent>(ChildActorComponents);
+	for (UChildActorComponent* ChildComponent : ChildActorComponents)
+	{
+		AActor* ChildActor = ChildComponent->GetChildActor();
+		if (ChildActor)
+		{
+			AppendActorComponentsRecursive(ChildActor, StaticMeshList, SkeletalMeshList, GeometryCollectionList);
+		}
+	}
+}
+
+UObject* UGeometryCollectionFactory::FactoryCreateNew(UClass* Class, UObject* InParent, FName Name, EObjectFlags Flags, UObject* Context, FFeedbackContext* Warn)
+{
+	FTransform LastTransform = FTransform::Identity;
+	TArray< GeometryCollectionStaticMeshConversionTuple > StaticMeshList;
+	TArray< GeometryCollectionSkeletalMeshConversionTuple > SkeletalMeshList;
+	TArray< GeometryCollectionTuple > GeometryCollectionList;
+
+	TArray<FAssetData> SelectedAssets;
+	GEditor->GetContentBrowserSelections(SelectedAssets);
+
+	for (const FAssetData& AssetData : SelectedAssets)
+	{
+		if (AssetData.GetAsset()->IsA<UStaticMesh>())
+		{
+			UStaticMeshComponent *DummyValue(NULL);
+			StaticMeshList.Add(GeometryCollectionStaticMeshConversionTuple(static_cast<const UStaticMesh *>(AssetData.GetAsset()), DummyValue, FTransform()));
+		}
+		else if (AssetData.GetAsset()->IsA<USkeletalMesh>())
+		{
+			USkeletalMeshComponent *DummyValue(NULL);
+			SkeletalMeshList.Add(GeometryCollectionSkeletalMeshConversionTuple(static_cast<const USkeletalMesh *>(AssetData.GetAsset()), DummyValue, FTransform()));
+		}
+		else if (AssetData.GetAsset()->IsA<UGeometryCollection>())
+		{
+			UGeometryCollectionComponent *DummyValue(NULL);
+			GeometryCollectionList.Add(GeometryCollectionTuple(static_cast<const UGeometryCollection *>(AssetData.GetAsset()), DummyValue, FTransform()));
+		}
+
+	}
+
+	if (USelection* SelectedActors = GEditor->GetSelectedActors())
+	{
+		for (FSelectionIterator Iter(*SelectedActors); Iter; ++Iter)
+		{
+			if (AActor* Actor = Cast<AActor>(*Iter))
+			{
+				AppendActorComponentsRecursive(Actor, StaticMeshList, SkeletalMeshList, GeometryCollectionList);
+
+				if (SelectedActors->GetBottom<AActor>() == Actor)
+				{
+					LastTransform = Actor->GetTransform();
+					LastTransform.SetScale3D(FVector(1.f, 1.f, 1.f));
+				}
+			}
+		}
+	}
+
+	UGeometryCollection* NewGeometryCollection = StaticFactoryCreateNew(Class, InParent, Name, Flags, Context, Warn);
+
+	for (GeometryCollectionStaticMeshConversionTuple & StaticMeshData : StaticMeshList)
+	{
+		FGeometryCollectionConversion::AppendStaticMesh(StaticMeshData.Get<0>(), StaticMeshData.Get<1>(), StaticMeshData.Get<2>(), NewGeometryCollection, false);
+	}
+
+	for (GeometryCollectionSkeletalMeshConversionTuple & SkeletalMeshData : SkeletalMeshList)
+	{
+		FGeometryCollectionConversion::AppendSkeletalMesh(SkeletalMeshData.Get<0>(), SkeletalMeshData.Get<1>(), SkeletalMeshData.Get<2>(), NewGeometryCollection, false);
+	}
+
+	for (GeometryCollectionTuple & GeometryCollectionData : GeometryCollectionList)
+	{
+		NewGeometryCollection->AppendGeometry(*GeometryCollectionData.Get<0>(), false, GeometryCollectionData.Get<2>());
+	}
+
+	// Add internal material and selection material
+	NewGeometryCollection->InitializeMaterials();
+	GeometryCollectionAlgo::PrepareForSimulation(NewGeometryCollection->GetGeometryCollection().Get());
+
+	// Initial pivot : 
+	// Offset everything from the last selected element so the transform will align with the null space. 
+	if (TSharedPtr<FGeometryCollection, ESPMode::ThreadSafe> Collection = NewGeometryCollection->GetGeometryCollection())
+	{
+		TManagedArray<int32> & Parent = Collection->Parent;
+		TManagedArray<FTransform>& Transform = Collection->Transform;
+
+		for(int TransformGroupIndex =0; TransformGroupIndex<Collection->Transform.Num(); TransformGroupIndex++)
+		{
+			if (Parent[TransformGroupIndex] == FGeometryCollection::Invalid)
+			{
+				Transform[TransformGroupIndex] = Transform[TransformGroupIndex].GetRelativeTransform(LastTransform);
+			}
+		}
+	}
+
+	NewGeometryCollection->Modify();
+	return NewGeometryCollection;
+}
+
+#undef LOCTEXT_NAMESPACE
+
+
+
diff --git a/Engine/Plugins/Experimental/GeometryCollectionPlugin/Source/GeometryCollectionEditor/Public/GeometryCollection/GeometryCollectionCacheFactory.h b/Engine/Plugins/Experimental/GeometryCollectionPlugin/Source/GeometryCollectionEditor/Public/GeometryCollection/GeometryCollectionCacheFactory.h
new file mode 100644
--- /dev/null
+++ b/Engine/Plugins/Experimental/GeometryCollectionPlugin/Source/GeometryCollectionEditor/Public/GeometryCollection/GeometryCollectionCacheFactory.h
@@ -0,0 +1,44 @@
+// Copyright 1998-2019 Epic Games, Inc. All Rights Reserved.
+
+#pragma once
+
+#include "Factories/Factory.h"
+#include "GeometryCollectionCacheFactory.generated.h"
+
+class UGeometryCollection;
+class UGeometryCollectionCache;
+class SWindow;
+struct FAssetData;
+
+UCLASS(Experimental)
+class UGeometryCollectionCacheFactory : public UFactory
+{
+	GENERATED_BODY()
+
+public:
+
+	/** Config properties required for CreateNew */
+
+	UPROPERTY()
+	UGeometryCollection* TargetCollection;
+
+	/** End required properties */
+
+	UGeometryCollectionCacheFactory();
+
+	/** UFactory Interface */
+	virtual bool CanCreateNew() const override;
+	virtual bool FactoryCanImport(const FString& Filename) override;
+	virtual bool ShouldShowInNewMenu() const override;
+	virtual UObject* FactoryCreateNew(UClass* InClass, UObject* InParent, FName InName, EObjectFlags Flags, UObject* Context, FFeedbackContext* Warn) override;
+	virtual bool ConfigureProperties() override;
+	/** End UFactory Interface */
+
+private:
+
+	/** Callback from window created in ConfigureProperties once a user picks a target Geometry Collection */
+	void OnConfigSelection(const FAssetData& InSelectedAssetData);
+
+	/** Window created in ConfigureProperties for target selection */
+	TSharedPtr<SWindow> PickerWindow;
+};
\ No newline at end of file
diff --git a/Engine/Plugins/Experimental/GeometryCollectionPlugin/Source/GeometryCollectionEditor/Public/GeometryCollection/GeometryCollectionConversion.h b/Engine/Plugins/Experimental/GeometryCollectionPlugin/Source/GeometryCollectionEditor/Public/GeometryCollection/GeometryCollectionConversion.h
--- a/Engine/Plugins/Experimental/GeometryCollectionPlugin/Source/GeometryCollectionEditor/Public/GeometryCollection/GeometryCollectionConversion.h
+++ b/Engine/Plugins/Experimental/GeometryCollectionPlugin/Source/GeometryCollectionEditor/Public/GeometryCollection/GeometryCollectionConversion.h
@@ -33,4 +33,21 @@ public:
 	*/
 	static void AppendSkeletalMesh(const USkeletalMesh* SkeletalMesh, const USkeletalMeshComponent *SkeletalMeshComponent, const FTransform & SkeletalMeshTransform, UGeometryCollection * GeometryCollection, bool ReindexMaterials = true);
 
+	/**
+	*  Command invoked from "GeometryCollection.CreatGeometryCollection", uses the selected Actors to create a GeometryCollection Asset
+	*  @param World
+	*/
+	static void CreateGeometryCollectionCommand(UWorld * World);
+
+	/**
+	*  Command invoked from "GeometryCollection.CreateFromSelectedActors", uses the selected Actors to create a GeometryCollection Asset
+	*  @param World
+	*/
+	static void CreateFromSelectedActorsCommand(UWorld * World);
+
+	/**
+	*  Command invoked from "GeometryCollection.CreateFromSelectedAssets", uses the selected Assets to create a GeometryCollection Asset
+	*  @param World
+	*/
+	static void CreateFromSelectedAssetsCommand(UWorld * World);
 };
\ No newline at end of file
diff --git a/Engine/Plugins/Experimental/GeometryCollectionPlugin/Source/GeometryCollectionEditor/Public/GeometryCollection/GeometryCollectionFactory.h b/Engine/Plugins/Experimental/GeometryCollectionPlugin/Source/GeometryCollectionEditor/Public/GeometryCollection/GeometryCollectionFactory.h
new file mode 100644
--- /dev/null
+++ b/Engine/Plugins/Experimental/GeometryCollectionPlugin/Source/GeometryCollectionEditor/Public/GeometryCollection/GeometryCollectionFactory.h
@@ -0,0 +1,36 @@
+// Copyright 1998-2019 Epic Games, Inc. All Rights Reserved.
+
+/** Factory which allows import of an GeometryCollectionAsset */
+
+#pragma once
+
+#include "CoreMinimal.h"
+#include "UObject/ObjectMacros.h"
+#include "Factories/Factory.h"
+#include "GeometryCollection/GeometryCollectionObject.h"
+
+#include "GeometryCollectionFactory.generated.h"
+
+class UGeometryCollection;
+class UGeometryCollectionComponent;
+
+typedef TTuple<const UGeometryCollection *, const UGeometryCollectionComponent *, FTransform> GeometryCollectionTuple;
+
+/**
+* Factory for Simple Cube
+*/
+
+UCLASS()
+class GEOMETRYCOLLECTIONEDITOR_API UGeometryCollectionFactory : public UFactory
+{
+    GENERATED_UCLASS_BODY()
+
+	//~ Begin UFactory Interface
+	virtual UObject* FactoryCreateNew(UClass* Class, UObject* InParent, FName Name, EObjectFlags Flags, UObject* Context, FFeedbackContext* Warn) override;
+	//~ Begin UFactory Interface	
+
+	static UGeometryCollection* StaticFactoryCreateNew(UClass* Class, UObject* InParent, FName Name, EObjectFlags Flags, UObject* Context, FFeedbackContext* Warn);
+
+};
+
+
