Chaos - minor fix to Geom Collection Plugin

diff --git a/Engine/Plugins/Experimental/GeometryCollectionPlugin/Source/GeometryCollectionExampleCore/Public/GeometryCollection/GeometryCollectionExampleSimulationObject.h b/Engine/Plugins/Experimental/GeometryCollectionPlugin/Source/GeometryCollectionExampleCore/Public/GeometryCollection/GeometryCollectionExampleSimulationObject.h
--- a/Engine/Plugins/Experimental/GeometryCollectionPlugin/Source/GeometryCollectionExampleCore/Public/GeometryCollection/GeometryCollectionExampleSimulationObject.h
+++ b/Engine/Plugins/Experimental/GeometryCollectionPlugin/Source/GeometryCollectionExampleCore/Public/GeometryCollection/GeometryCollectionExampleSimulationObject.h
@@ -67,7 +67,7 @@ namespace GeometryCollectionExample {
 			: Parameters(InParameters)
 			, RestCollection(InRestCollection)
 			, DynamicCollection(InDynamicCollection)
-			, PhysicalMaterial(InPhysicalMaterial)
+			, PhysicalMaterial(MoveTemp(InPhysicalMaterial))
 			, PhysicsProxy(new FGeometryCollectionPhysicsProxy(nullptr, DynamicCollection.Get(), [&](FSimulationParameters& P) {Init(P); }, nullptr, nullptr))
 		{
 		}
