Sequencer: Fix split regression. The newly created split section should be the right side.

diff --git a/Engine/Plugins/Experimental/GeometryCache/Source/GeometryCacheTracks/Private/MovieSceneGeometryCacheSection.cpp b/Engine/Plugins/Experimental/GeometryCache/Source/GeometryCacheTracks/Private/MovieSceneGeometryCacheSection.cpp
--- a/Engine/Plugins/Experimental/GeometryCache/Source/GeometryCacheTracks/Private/MovieSceneGeometryCacheSection.cpp
+++ b/Engine/Plugins/Experimental/GeometryCache/Source/GeometryCacheTracks/Private/MovieSceneGeometryCacheSection.cpp
@@ -131,6 +131,8 @@ void UMovieSceneGeometryCacheSection::TrimSection(FQualifiedFrameTime TrimTime,
 
 UMovieSceneSection* UMovieSceneGeometryCacheSection::SplitSection(FQualifiedFrameTime SplitTime, bool bDeleteKeys)
 {
+	const FFrameNumber InitialFirstLoopStartFrameOffset = Params.FirstLoopStartFrameOffset;
+
 	FFrameRate FrameRate = GetTypedOuter<UMovieScene>()->GetTickResolution();
 
 	const FFrameNumber NewOffset = HasStartFrame() ? GetFirstLoopStartOffsetAtTrimTime(SplitTime, Params, GetInclusiveStartFrame(), FrameRate) : 0;
@@ -141,6 +143,10 @@ UMovieSceneSection* UMovieSceneGeometryCacheSection::SplitSection(FQualifiedFram
 		UMovieSceneGeometryCacheSection* NewGeometrySection = Cast<UMovieSceneGeometryCacheSection>(NewSection);
 		NewGeometrySection->Params.FirstLoopStartFrameOffset = NewOffset;
 	}
+
+	// Restore original offset modified by splitting
+	Params.FirstLoopStartFrameOffset = InitialFirstLoopStartFrameOffset;
+
 	return NewSection;
 }
 
diff --git a/Engine/Plugins/Experimental/GeometryCollectionPlugin/Source/GeometryCollectionTracks/Private/MovieSceneGeometryCollectionSection.cpp b/Engine/Plugins/Experimental/GeometryCollectionPlugin/Source/GeometryCollectionTracks/Private/MovieSceneGeometryCollectionSection.cpp
--- a/Engine/Plugins/Experimental/GeometryCollectionPlugin/Source/GeometryCollectionTracks/Private/MovieSceneGeometryCollectionSection.cpp
+++ b/Engine/Plugins/Experimental/GeometryCollectionPlugin/Source/GeometryCollectionTracks/Private/MovieSceneGeometryCollectionSection.cpp
@@ -84,6 +84,8 @@ void UMovieSceneGeometryCollectionSection::TrimSection(FQualifiedFrameTime TrimT
 
 UMovieSceneSection* UMovieSceneGeometryCollectionSection::SplitSection(FQualifiedFrameTime SplitTime, bool bDeleteKeys)
 {
+	const FFrameNumber InitialStartFrameOffset = Params.StartFrameOffset;
+
 	FFrameRate FrameRate = GetTypedOuter<UMovieScene>()->GetTickResolution();
 
 	const FFrameNumber NewOffset = HasStartFrame() ? GetStartOffsetAtTrimTime(SplitTime, Params, GetInclusiveStartFrame(), FrameRate) : 0;
@@ -94,6 +96,10 @@ UMovieSceneSection* UMovieSceneGeometryCollectionSection::SplitSection(FQualifie
 		UMovieSceneGeometryCollectionSection* NewGeometrySection = Cast<UMovieSceneGeometryCollectionSection>(NewSection);
 		NewGeometrySection->Params.StartFrameOffset = NewOffset;
 	}
+
+	// Restore original offset modified by splitting
+	Params.StartFrameOffset = InitialStartFrameOffset;
+
 	return NewSection;
 }
 
diff --git a/Engine/Source/Runtime/MovieScene/Private/MovieSceneSection.cpp b/Engine/Source/Runtime/MovieScene/Private/MovieSceneSection.cpp
--- a/Engine/Source/Runtime/MovieScene/Private/MovieSceneSection.cpp
+++ b/Engine/Source/Runtime/MovieScene/Private/MovieSceneSection.cpp
@@ -446,8 +446,8 @@ UMovieSceneSection* UMovieSceneSection::SplitSection(FQualifiedFrameTime SplitTi
 
 		Track->AddSection(*NewSection);
 
-		NewSection->TrimSection(SplitTime, false, bDeleteKeys);
-		TrimSection(SplitTime, true, bDeleteKeys);
+		TrimSection(SplitTime, false, bDeleteKeys);
+		NewSection->TrimSection(SplitTime, true, bDeleteKeys);
 
 		return NewSection;
 	}
diff --git a/Engine/Source/Runtime/MovieScene/Private/Sections/MovieSceneSubSection.cpp b/Engine/Source/Runtime/MovieScene/Private/Sections/MovieSceneSubSection.cpp
--- a/Engine/Source/Runtime/MovieScene/Private/Sections/MovieSceneSubSection.cpp
+++ b/Engine/Source/Runtime/MovieScene/Private/Sections/MovieSceneSubSection.cpp
@@ -254,6 +254,9 @@ UMovieSceneSection* UMovieSceneSubSection::SplitSection( FQualifiedFrameTime Spl
 		return NewSection;
 	}
 
+	// Restore original offset modified by splitting
+	Parameters.StartFrameOffset = InitialStartOffset;
+
 	return nullptr;
 }
 
diff --git a/Engine/Source/Runtime/MovieSceneTracks/Private/Sections/MovieSceneAudioSection.cpp b/Engine/Source/Runtime/MovieSceneTracks/Private/Sections/MovieSceneAudioSection.cpp
--- a/Engine/Source/Runtime/MovieSceneTracks/Private/Sections/MovieSceneAudioSection.cpp
+++ b/Engine/Source/Runtime/MovieSceneTracks/Private/Sections/MovieSceneAudioSection.cpp
@@ -173,6 +173,8 @@ void UMovieSceneAudioSection::TrimSection(FQualifiedFrameTime TrimTime, bool bTr
 
 UMovieSceneSection* UMovieSceneAudioSection::SplitSection(FQualifiedFrameTime SplitTime, bool bDeleteKeys)
 {
+	const FFrameNumber InitialStartFrameOffset = StartFrameOffset;
+
 	const FFrameNumber NewOffset = HasStartFrame() ? GetStartOffsetAtTrimTime(SplitTime, StartFrameOffset, GetInclusiveStartFrame()) : 0;
 
 	UMovieSceneSection* NewSection = Super::SplitSection(SplitTime, bDeleteKeys);
@@ -181,5 +183,9 @@ UMovieSceneSection* UMovieSceneAudioSection::SplitSection(FQualifiedFrameTime Sp
 		UMovieSceneAudioSection* NewAudioSection = Cast<UMovieSceneAudioSection>(NewSection);
 		NewAudioSection->StartFrameOffset = NewOffset;
 	}
+
+	// Restore original offset modified by splitting
+	StartFrameOffset = InitialStartFrameOffset;
+
 	return NewSection;
 }
diff --git a/Engine/Source/Runtime/MovieSceneTracks/Private/Sections/MovieSceneSkeletalAnimationSection.cpp b/Engine/Source/Runtime/MovieSceneTracks/Private/Sections/MovieSceneSkeletalAnimationSection.cpp
--- a/Engine/Source/Runtime/MovieSceneTracks/Private/Sections/MovieSceneSkeletalAnimationSection.cpp
+++ b/Engine/Source/Runtime/MovieSceneTracks/Private/Sections/MovieSceneSkeletalAnimationSection.cpp
@@ -220,6 +220,8 @@ void UMovieSceneSkeletalAnimationSection::TrimSection(FQualifiedFrameTime TrimTi
 
 UMovieSceneSection* UMovieSceneSkeletalAnimationSection::SplitSection(FQualifiedFrameTime SplitTime, bool bDeleteKeys)
 {
+	const FFrameNumber InitialFirstLoopStartFrameOffset = Params.FirstLoopStartFrameOffset;
+
 	FFrameRate FrameRate = GetTypedOuter<UMovieScene>()->GetTickResolution();
 
 	const FFrameNumber NewOffset = HasStartFrame() ? GetFirstLoopStartOffsetAtTrimTime(SplitTime, Params, GetInclusiveStartFrame(), FrameRate) : 0;
@@ -230,6 +232,10 @@ UMovieSceneSection* UMovieSceneSkeletalAnimationSection::SplitSection(FQualified
 		UMovieSceneSkeletalAnimationSection* NewSkeletalSection = Cast<UMovieSceneSkeletalAnimationSection>(NewSection);
 		NewSkeletalSection->Params.FirstLoopStartFrameOffset = NewOffset;
 	}
+
+	// Restore original offset modified by splitting
+	Params.FirstLoopStartFrameOffset = InitialFirstLoopStartFrameOffset;
+
 	return NewSection;
 }
 
