Added parameter to thumbnail renderers Draw function so that we can instruct them to clear the render target or not. This is for fixing thumbnails that render multiple sub-thumbnails to the same render target.

diff --git a/Engine/Plugins/2D/Paper2D/Source/Paper2DEditor/Classes/PaperFlipbookThumbnailRenderer.h b/Engine/Plugins/2D/Paper2D/Source/Paper2DEditor/Classes/PaperFlipbookThumbnailRenderer.h
--- a/Engine/Plugins/2D/Paper2D/Source/Paper2DEditor/Classes/PaperFlipbookThumbnailRenderer.h
+++ b/Engine/Plugins/2D/Paper2D/Source/Paper2DEditor/Classes/PaperFlipbookThumbnailRenderer.h
@@ -16,6 +16,6 @@ class UPaperFlipbookThumbnailRenderer : public UPaperSpriteThumbnailRenderer
 	GENERATED_UCLASS_BODY()
 
 	// UThumbnailRenderer interface
-	virtual void Draw(UObject* Object, int32 X, int32 Y, uint32 Width, uint32 Height, FRenderTarget*, FCanvas* Canvas) override;
+	virtual void Draw(UObject* Object, int32 X, int32 Y, uint32 Width, uint32 Height, FRenderTarget*, FCanvas* Canvas, bool bAdditionalViewFamily) override;
 	// End of UThumbnailRenderer interface
 };
diff --git a/Engine/Plugins/2D/Paper2D/Source/Paper2DEditor/Classes/PaperSpriteThumbnailRenderer.h b/Engine/Plugins/2D/Paper2D/Source/Paper2DEditor/Classes/PaperSpriteThumbnailRenderer.h
--- a/Engine/Plugins/2D/Paper2D/Source/Paper2DEditor/Classes/PaperSpriteThumbnailRenderer.h
+++ b/Engine/Plugins/2D/Paper2D/Source/Paper2DEditor/Classes/PaperSpriteThumbnailRenderer.h
@@ -16,7 +16,7 @@ class UPaperSpriteThumbnailRenderer : public UDefaultSizedThumbnailRenderer
 	GENERATED_UCLASS_BODY()
 
 	// UThumbnailRenderer interface
-	virtual void Draw(UObject* Object, int32 X, int32 Y, uint32 Width, uint32 Height, FRenderTarget*, FCanvas* Canvas) override;
+	virtual void Draw(UObject* Object, int32 X, int32 Y, uint32 Width, uint32 Height, FRenderTarget*, FCanvas* Canvas, bool bAdditionalViewFamily) override;
 	// End of UThumbnailRenderer interface
 
 protected:
diff --git a/Engine/Plugins/2D/Paper2D/Source/Paper2DEditor/Classes/PaperTileSetThumbnailRenderer.h b/Engine/Plugins/2D/Paper2D/Source/Paper2DEditor/Classes/PaperTileSetThumbnailRenderer.h
--- a/Engine/Plugins/2D/Paper2D/Source/Paper2DEditor/Classes/PaperTileSetThumbnailRenderer.h
+++ b/Engine/Plugins/2D/Paper2D/Source/Paper2DEditor/Classes/PaperTileSetThumbnailRenderer.h
@@ -16,6 +16,6 @@ class UPaperTileSetThumbnailRenderer : public UDefaultSizedThumbnailRenderer
 	GENERATED_UCLASS_BODY()
 
 	// UThumbnailRenderer interface
-	virtual void Draw(UObject* Object, int32 X, int32 Y, uint32 Width, uint32 Height, FRenderTarget*, FCanvas* Canvas) override;
+	virtual void Draw(UObject* Object, int32 X, int32 Y, uint32 Width, uint32 Height, FRenderTarget*, FCanvas* Canvas, bool bAdditionalViewFamily) override;
 	// End of UThumbnailRenderer interface
 };
diff --git a/Engine/Plugins/2D/Paper2D/Source/Paper2DEditor/Private/PaperFlipbookThumbnailRenderer.cpp b/Engine/Plugins/2D/Paper2D/Source/Paper2DEditor/Private/PaperFlipbookThumbnailRenderer.cpp
--- a/Engine/Plugins/2D/Paper2D/Source/Paper2DEditor/Private/PaperFlipbookThumbnailRenderer.cpp
+++ b/Engine/Plugins/2D/Paper2D/Source/Paper2DEditor/Private/PaperFlipbookThumbnailRenderer.cpp
@@ -15,7 +15,7 @@ UPaperFlipbookThumbnailRenderer::UPaperFlipbookThumbnailRenderer(const FObjectIn
 {
 }
 
-void UPaperFlipbookThumbnailRenderer::Draw(UObject* Object, int32 X, int32 Y, uint32 Width, uint32 Height, FRenderTarget* RenderTarget, FCanvas* Canvas)
+void UPaperFlipbookThumbnailRenderer::Draw(UObject* Object, int32 X, int32 Y, uint32 Width, uint32 Height, FRenderTarget* RenderTarget, FCanvas* Canvas, bool bAdditionalViewFamily)
 {
 	if (UPaperFlipbook* Flipbook = Cast<UPaperFlipbook>(Object))
 	{
diff --git a/Engine/Plugins/2D/Paper2D/Source/Paper2DEditor/Private/PaperSpriteThumbnailRenderer.cpp b/Engine/Plugins/2D/Paper2D/Source/Paper2DEditor/Private/PaperSpriteThumbnailRenderer.cpp
--- a/Engine/Plugins/2D/Paper2D/Source/Paper2DEditor/Private/PaperSpriteThumbnailRenderer.cpp
+++ b/Engine/Plugins/2D/Paper2D/Source/Paper2DEditor/Private/PaperSpriteThumbnailRenderer.cpp
@@ -16,7 +16,7 @@ UPaperSpriteThumbnailRenderer::UPaperSpriteThumbnailRenderer(const FObjectInitia
 {
 }
 
-void UPaperSpriteThumbnailRenderer::Draw(UObject* Object, int32 X, int32 Y, uint32 Width, uint32 Height, FRenderTarget* RenderTarget, FCanvas* Canvas)
+void UPaperSpriteThumbnailRenderer::Draw(UObject* Object, int32 X, int32 Y, uint32 Width, uint32 Height, FRenderTarget* RenderTarget, FCanvas* Canvas, bool bAdditionalViewFamily)
 {
 	UPaperSprite* Sprite = Cast<UPaperSprite>(Object);
 	DrawFrame(Sprite, X, Y, Width, Height, RenderTarget, Canvas, nullptr);
diff --git a/Engine/Plugins/2D/Paper2D/Source/Paper2DEditor/Private/PaperTileSetThumbnailRenderer.cpp b/Engine/Plugins/2D/Paper2D/Source/Paper2DEditor/Private/PaperTileSetThumbnailRenderer.cpp
--- a/Engine/Plugins/2D/Paper2D/Source/Paper2DEditor/Private/PaperTileSetThumbnailRenderer.cpp
+++ b/Engine/Plugins/2D/Paper2D/Source/Paper2DEditor/Private/PaperTileSetThumbnailRenderer.cpp
@@ -14,7 +14,7 @@ UPaperTileSetThumbnailRenderer::UPaperTileSetThumbnailRenderer(const FObjectInit
 {
 }
 
-void UPaperTileSetThumbnailRenderer::Draw(UObject* Object, int32 X, int32 Y, uint32 Width, uint32 Height, FRenderTarget*, FCanvas* Canvas)
+void UPaperTileSetThumbnailRenderer::Draw(UObject* Object, int32 X, int32 Y, uint32 Width, uint32 Height, FRenderTarget*, FCanvas* Canvas, bool bAdditionalViewFamily)
 {
 	UPaperTileSet* TileSet = Cast<UPaperTileSet>(Object);
 	if ((TileSet != nullptr) && (TileSet->GetTileSheetTexture() != nullptr))
diff --git a/Engine/Plugins/Experimental/ControlRig/Source/ControlRigEditor/Private/ControlRigThumbnailRenderer.cpp b/Engine/Plugins/Experimental/ControlRig/Source/ControlRigEditor/Private/ControlRigThumbnailRenderer.cpp
--- a/Engine/Plugins/Experimental/ControlRig/Source/ControlRigEditor/Private/ControlRigThumbnailRenderer.cpp
+++ b/Engine/Plugins/Experimental/ControlRig/Source/ControlRigEditor/Private/ControlRigThumbnailRenderer.cpp
@@ -21,7 +21,7 @@ bool UControlRigThumbnailRenderer::CanVisualizeAsset(UObject* Object)
 	return Cast<UControlRigBlueprint>(Object) != nullptr;
 }
 
-void UControlRigThumbnailRenderer::Draw(UObject* Object, int32 X, int32 Y, uint32 Width, uint32 Height, FRenderTarget* RenderTarget, FCanvas* Canvas)
+void UControlRigThumbnailRenderer::Draw(UObject* Object, int32 X, int32 Y, uint32 Width, uint32 Height, FRenderTarget* RenderTarget, FCanvas* Canvas, bool bAdditionalViewFamily)
 {
 	RigBlueprint = nullptr;
 
@@ -31,7 +31,7 @@ void UControlRigThumbnailRenderer::Draw(UObject* Object, int32 X, int32 Y, uint3
 		if (SkeletalMesh != nullptr)
 		{
 			RigBlueprint = InRigBlueprint;
-			Super::Draw(SkeletalMesh, X, Y, Width, Height, RenderTarget, Canvas);
+			Super::Draw(SkeletalMesh, X, Y, Width, Height, RenderTarget, Canvas, bAdditionalViewFamily);
 
 			for (auto Pair : GizmoActors)
 			{
diff --git a/Engine/Plugins/Experimental/ControlRig/Source/ControlRigEditor/Public/ControlRigThumbnailRenderer.h b/Engine/Plugins/Experimental/ControlRig/Source/ControlRigEditor/Public/ControlRigThumbnailRenderer.h
--- a/Engine/Plugins/Experimental/ControlRig/Source/ControlRigEditor/Public/ControlRigThumbnailRenderer.h
+++ b/Engine/Plugins/Experimental/ControlRig/Source/ControlRigEditor/Public/ControlRigThumbnailRenderer.h
@@ -24,7 +24,7 @@ class UControlRigThumbnailRenderer : public USkeletalMeshThumbnailRenderer
 
 	// Begin UThumbnailRenderer Object
 	CONTROLRIGEDITOR_API virtual bool CanVisualizeAsset(UObject* Object) override;
-	CONTROLRIGEDITOR_API virtual void Draw(UObject* Object, int32 X, int32 Y, uint32 Width, uint32 Height, FRenderTarget*, FCanvas* Canvas) override;
+	CONTROLRIGEDITOR_API virtual void Draw(UObject* Object, int32 X, int32 Y, uint32 Width, uint32 Height, FRenderTarget*, FCanvas* Canvas, bool bAdditionalViewFamily) override;
 	// End UThumbnailRenderer Object
 
 	// USkeletalMeshThumbnailRenderer implementation
diff --git a/Engine/Plugins/Experimental/GeometryCache/Source/GeometryCacheEd/Classes/GeometryCacheThumbnailRenderer.h b/Engine/Plugins/Experimental/GeometryCache/Source/GeometryCacheEd/Classes/GeometryCacheThumbnailRenderer.h
--- a/Engine/Plugins/Experimental/GeometryCache/Source/GeometryCacheEd/Classes/GeometryCacheThumbnailRenderer.h
+++ b/Engine/Plugins/Experimental/GeometryCache/Source/GeometryCacheEd/Classes/GeometryCacheThumbnailRenderer.h
@@ -22,7 +22,7 @@ class UGeometryCacheThumbnailRenderer : public UDefaultSizedThumbnailRenderer
 	GENERATED_UCLASS_BODY()
 
 	// Begin UThumbnailRenderer Object
-	virtual void Draw(UObject* Object, int32 X, int32 Y, uint32 Width, uint32 Height, FRenderTarget* RenderTarget, FCanvas* Canvas) override;
+	virtual void Draw(UObject* Object, int32 X, int32 Y, uint32 Width, uint32 Height, FRenderTarget* RenderTarget, FCanvas* Canvas, bool bAdditionalViewFamily) override;
 	// End UThumbnailRenderer Object
 
 	// Begin UObject implementation
diff --git a/Engine/Plugins/Experimental/GeometryCache/Source/GeometryCacheEd/Private/GeometryCacheThumbnailRenderer.cpp b/Engine/Plugins/Experimental/GeometryCache/Source/GeometryCacheEd/Private/GeometryCacheThumbnailRenderer.cpp
--- a/Engine/Plugins/Experimental/GeometryCache/Source/GeometryCacheEd/Private/GeometryCacheThumbnailRenderer.cpp
+++ b/Engine/Plugins/Experimental/GeometryCache/Source/GeometryCacheEd/Private/GeometryCacheThumbnailRenderer.cpp
@@ -13,7 +13,7 @@ UGeometryCacheThumbnailRenderer::UGeometryCacheThumbnailRenderer(const FObjectIn
 	ThumbnailScene = nullptr;
 }
 
-void UGeometryCacheThumbnailRenderer::Draw(UObject* Object, int32 X, int32 Y, uint32 Width, uint32 Height, FRenderTarget* RenderTarget, FCanvas* Canvas)
+void UGeometryCacheThumbnailRenderer::Draw(UObject* Object, int32 X, int32 Y, uint32 Width, uint32 Height, FRenderTarget* RenderTarget, FCanvas* Canvas, bool bAdditionalViewFamily)
 {
 	UGeometryCache* GeometryCache = Cast<UGeometryCache>(Object);
 	if (GeometryCache != nullptr && !GeometryCache->IsPendingKill())
@@ -27,7 +27,8 @@ void UGeometryCacheThumbnailRenderer::Draw(UObject* Object, int32 X, int32 Y, ui
 		ThumbnailScene->GetScene()->UpdateSpeedTreeWind(0.0);
 
 		FSceneViewFamilyContext ViewFamily(FSceneViewFamily::ConstructionValues(RenderTarget, ThumbnailScene->GetScene(), FEngineShowFlags(ESFIM_Game))
-			.SetWorldTimes(FApp::GetCurrentTime() - GStartTime, FApp::GetDeltaTime(), FApp::GetCurrentTime() - GStartTime));
+			.SetWorldTimes(FApp::GetCurrentTime() - GStartTime, FApp::GetDeltaTime(), FApp::GetCurrentTime() - GStartTime)
+			.SetAdditionalViewFamily(bAdditionalViewFamily));
 
 		ViewFamily.EngineShowFlags.DisableAdvancedFeatures();
 		ViewFamily.EngineShowFlags.MotionBlur = 0;
diff --git a/Engine/Plugins/Experimental/GeometryCollectionPlugin/Source/GeometryCollectionEditor/Private/GeometryCollection/GeometryCollectionThumbnailRenderer.cpp b/Engine/Plugins/Experimental/GeometryCollectionPlugin/Source/GeometryCollectionEditor/Private/GeometryCollection/GeometryCollectionThumbnailRenderer.cpp
--- a/Engine/Plugins/Experimental/GeometryCollectionPlugin/Source/GeometryCollectionEditor/Private/GeometryCollection/GeometryCollectionThumbnailRenderer.cpp
+++ b/Engine/Plugins/Experimental/GeometryCollectionPlugin/Source/GeometryCollectionEditor/Private/GeometryCollection/GeometryCollectionThumbnailRenderer.cpp
@@ -13,7 +13,7 @@ UGeometryCollectionThumbnailRenderer::UGeometryCollectionThumbnailRenderer(const
 	ThumbnailScene = nullptr;
 }
 
-void UGeometryCollectionThumbnailRenderer::Draw(UObject* Object, int32 X, int32 Y, uint32 Width, uint32 Height, FRenderTarget* RenderTarget, FCanvas* Canvas)
+void UGeometryCollectionThumbnailRenderer::Draw(UObject* Object, int32 X, int32 Y, uint32 Width, uint32 Height, FRenderTarget* RenderTarget, FCanvas* Canvas, bool bAdditionalViewFamily)
 {
 	UGeometryCollection* GeometryCollection = Cast<UGeometryCollection>(Object);
 	if (GeometryCollection != nullptr && !GeometryCollection->IsPendingKill())
@@ -27,7 +27,8 @@ void UGeometryCollectionThumbnailRenderer::Draw(UObject* Object, int32 X, int32
 		ThumbnailScene->GetScene()->UpdateSpeedTreeWind(0.0);
 
 		FSceneViewFamilyContext ViewFamily(FSceneViewFamily::ConstructionValues(RenderTarget, ThumbnailScene->GetScene(), FEngineShowFlags(ESFIM_Game))
-			.SetWorldTimes(FApp::GetCurrentTime() - GStartTime, FApp::GetDeltaTime(), FApp::GetCurrentTime() - GStartTime));
+			.SetWorldTimes(FApp::GetCurrentTime() - GStartTime, FApp::GetDeltaTime(), FApp::GetCurrentTime() - GStartTime)
+			.SetAdditionalViewFamily(bAdditionalViewFamily));
 
 		ViewFamily.EngineShowFlags.DisableAdvancedFeatures();
 		ViewFamily.EngineShowFlags.MotionBlur = 0;
diff --git a/Engine/Plugins/Experimental/GeometryCollectionPlugin/Source/GeometryCollectionEditor/Public/GeometryCollection/GeometryCollectionThumbnailRenderer.h b/Engine/Plugins/Experimental/GeometryCollectionPlugin/Source/GeometryCollectionEditor/Public/GeometryCollection/GeometryCollectionThumbnailRenderer.h
--- a/Engine/Plugins/Experimental/GeometryCollectionPlugin/Source/GeometryCollectionEditor/Public/GeometryCollection/GeometryCollectionThumbnailRenderer.h
+++ b/Engine/Plugins/Experimental/GeometryCollectionPlugin/Source/GeometryCollectionEditor/Public/GeometryCollection/GeometryCollectionThumbnailRenderer.h
@@ -22,7 +22,7 @@ class UGeometryCollectionThumbnailRenderer : public UDefaultSizedThumbnailRender
 	GENERATED_UCLASS_BODY()
 
 	// Begin UThumbnailRenderer Object
-	virtual void Draw(UObject* Object, int32 X, int32 Y, uint32 Width, uint32 Height, FRenderTarget* RenderTarget, FCanvas* Canvas) override;
+	virtual void Draw(UObject* Object, int32 X, int32 Y, uint32 Width, uint32 Height, FRenderTarget* RenderTarget, FCanvas* Canvas, bool bAdditionalViewFamily) override;
 	// End UThumbnailRenderer Object
 
 	// Begin UObject implementation
diff --git a/Engine/Plugins/FX/Niagara/Source/NiagaraEditor/Private/NiagaraThumbnailRenderer.cpp b/Engine/Plugins/FX/Niagara/Source/NiagaraEditor/Private/NiagaraThumbnailRenderer.cpp
--- a/Engine/Plugins/FX/Niagara/Source/NiagaraEditor/Private/NiagaraThumbnailRenderer.cpp
+++ b/Engine/Plugins/FX/Niagara/Source/NiagaraEditor/Private/NiagaraThumbnailRenderer.cpp
@@ -25,7 +25,7 @@ void UNiagaraThumbnailRendererBase::GetThumbnailSize(UObject* Object, float Zoom
 	}
 }
 
-void UNiagaraThumbnailRendererBase::Draw(UObject* Object, int32 X, int32 Y, uint32 Width, uint32 Height, FRenderTarget*, FCanvas* Canvas)
+void UNiagaraThumbnailRendererBase::Draw(UObject* Object, int32 X, int32 Y, uint32 Width, uint32 Height, FRenderTarget*, FCanvas* Canvas, bool bAdditionalViewFamily)
 {
 	UTexture2D* ObjectTexture = GetThumbnailTextureFromObject(Object);
 	if (ObjectTexture != nullptr)
diff --git a/Engine/Plugins/FX/Niagara/Source/NiagaraEditor/Private/NiagaraThumbnailRenderer.h b/Engine/Plugins/FX/Niagara/Source/NiagaraEditor/Private/NiagaraThumbnailRenderer.h
--- a/Engine/Plugins/FX/Niagara/Source/NiagaraEditor/Private/NiagaraThumbnailRenderer.h
+++ b/Engine/Plugins/FX/Niagara/Source/NiagaraEditor/Private/NiagaraThumbnailRenderer.h
@@ -16,7 +16,7 @@ public:
 	// UThumbnailRenderer interface.
 	virtual bool CanVisualizeAsset(UObject* Object) override;
 	virtual void GetThumbnailSize(UObject* Object, float Zoom, uint32& OutWidth, uint32& OutHeight) const override;
-	virtual void Draw(UObject* Object, int32 X, int32 Y, uint32 Width, uint32 Height, FRenderTarget*, FCanvas* Canvas) override;
+	virtual void Draw(UObject* Object, int32 X, int32 Y, uint32 Width, uint32 Height, FRenderTarget*, FCanvas* Canvas, bool bAdditionalViewFamily) override;
 
 protected:
 	virtual UTexture2D* GetThumbnailTextureFromObject(UObject* Object) const { return nullptr; }
diff --git a/Engine/Source/Editor/Cascade/Private/CascadeEmitterCanvasClient.cpp b/Engine/Source/Editor/Cascade/Private/CascadeEmitterCanvasClient.cpp
--- a/Engine/Source/Editor/Cascade/Private/CascadeEmitterCanvasClient.cpp
+++ b/Engine/Source/Editor/Cascade/Private/CascadeEmitterCanvasClient.cpp
@@ -1160,7 +1160,7 @@ void FCascadeEmitterCanvasClient::DrawHeaderBlock(int32 Index, int32 XPos, UPart
 							Thumbnail.Material = MaterialInterface;
 							// Re-render
 							FCanvas ThumbnailCanvas(Thumbnail.Texture->GameThread_GetRenderTargetResource(), nullptr, GetWorld(), GetWorld()->FeatureLevel, FCanvas::CDM_DeferDrawing, ShouldDPIScaleSceneCanvas() ? GetDPIScale() : 1.0f);
-							RenderInfo->Renderer->Draw(MaterialInterface, 0, 0, ScaledSize, ScaledSize, Thumbnail.Texture->GameThread_GetRenderTargetResource(), &ThumbnailCanvas);
+							RenderInfo->Renderer->Draw(MaterialInterface, 0, 0, ScaledSize, ScaledSize, Thumbnail.Texture->GameThread_GetRenderTargetResource(), &ThumbnailCanvas, false);
 						}
 					
 						Canvas->DrawTile(ThumbPos.X - Origin2D.X, ThumbPos.Y - Origin2D.Y, ThumbSize, ThumbSize, 0.f, 0.f, 1.f, 1.f, FLinearColor::White, Thumbnail.Texture->GameThread_GetRenderTargetResource(), false);
diff --git a/Engine/Source/Editor/FoliageEdit/Private/FoliageType_ActorThumbnailRenderer.cpp b/Engine/Source/Editor/FoliageEdit/Private/FoliageType_ActorThumbnailRenderer.cpp
--- a/Engine/Source/Editor/FoliageEdit/Private/FoliageType_ActorThumbnailRenderer.cpp
+++ b/Engine/Source/Editor/FoliageEdit/Private/FoliageType_ActorThumbnailRenderer.cpp
@@ -17,7 +17,7 @@ bool UFoliageType_ActorThumbnailRenderer::CanVisualizeAsset(UObject* Object)
 	return Super::CanVisualizeAsset(FoliageType->ActorClass->ClassGeneratedBy);
 }
 
-void UFoliageType_ActorThumbnailRenderer::Draw(UObject* Object, int32 X, int32 Y, uint32 Width, uint32 Height, FRenderTarget* RenderTarget, FCanvas* Canvas)
+void UFoliageType_ActorThumbnailRenderer::Draw(UObject* Object, int32 X, int32 Y, uint32 Width, uint32 Height, FRenderTarget* RenderTarget, FCanvas* Canvas, bool bAdditionalViewFamily)
 {
 	UFoliageType_Actor* FoliageType = Cast<UFoliageType_Actor>(Object);
 	if (FoliageType == nullptr || FoliageType->ActorClass == nullptr)
@@ -25,7 +25,7 @@ void UFoliageType_ActorThumbnailRenderer::Draw(UObject* Object, int32 X, int32 Y
 		return;
 	}
 
-	Super::Draw(FoliageType->ActorClass->ClassGeneratedBy, X, Y, Width, Height, RenderTarget, Canvas);
+	Super::Draw(FoliageType->ActorClass->ClassGeneratedBy, X, Y, Width, Height, RenderTarget, Canvas, bAdditionalViewFamily);
 }
 
 
diff --git a/Engine/Source/Editor/FoliageEdit/Private/FoliageType_ISMThumbnailRenderer.cpp b/Engine/Source/Editor/FoliageEdit/Private/FoliageType_ISMThumbnailRenderer.cpp
--- a/Engine/Source/Editor/FoliageEdit/Private/FoliageType_ISMThumbnailRenderer.cpp
+++ b/Engine/Source/Editor/FoliageEdit/Private/FoliageType_ISMThumbnailRenderer.cpp
@@ -14,7 +14,7 @@ bool UFoliageType_ISMThumbnailRenderer::CanVisualizeAsset(UObject* Object)
 	return FoliageType && FoliageType->GetStaticMesh();
 }
 
-void UFoliageType_ISMThumbnailRenderer::Draw(UObject* Object, int32 X, int32 Y, uint32 Width, uint32 Height, FRenderTarget* RenderTarget, FCanvas* Canvas)
+void UFoliageType_ISMThumbnailRenderer::Draw(UObject* Object, int32 X, int32 Y, uint32 Width, uint32 Height, FRenderTarget* RenderTarget, FCanvas* Canvas, bool bAdditionalViewFamily)
 {
 	UFoliageType_InstancedStaticMesh* FoliageType = Cast<UFoliageType_InstancedStaticMesh>(Object);
 	if (FoliageType && FoliageType->GetStaticMesh())
@@ -29,7 +29,8 @@ void UFoliageType_ISMThumbnailRenderer::Draw(UObject* Object, int32 X, int32 Y,
 		ThumbnailScene->GetScene()->UpdateSpeedTreeWind(0.0);
 
 		FSceneViewFamilyContext ViewFamily(FSceneViewFamily::ConstructionValues(RenderTarget, ThumbnailScene->GetScene(), FEngineShowFlags(ESFIM_Game))
-			.SetWorldTimes(FApp::GetCurrentTime() - GStartTime, FApp::GetDeltaTime(), FApp::GetCurrentTime() - GStartTime));
+			.SetWorldTimes(FApp::GetCurrentTime() - GStartTime, FApp::GetDeltaTime(), FApp::GetCurrentTime() - GStartTime)
+			.SetAdditionalViewFamily(bAdditionalViewFamily));
 
 		ViewFamily.EngineShowFlags.DisableAdvancedFeatures();
 		ViewFamily.EngineShowFlags.MotionBlur = 0;
diff --git a/Engine/Source/Editor/FoliageEdit/Public/FoliageType_ActorThumbnailRenderer.h b/Engine/Source/Editor/FoliageEdit/Public/FoliageType_ActorThumbnailRenderer.h
--- a/Engine/Source/Editor/FoliageEdit/Public/FoliageType_ActorThumbnailRenderer.h
+++ b/Engine/Source/Editor/FoliageEdit/Public/FoliageType_ActorThumbnailRenderer.h
@@ -25,6 +25,6 @@ class UFoliageType_ActorThumbnailRenderer : public UBlueprintThumbnailRenderer
 	{}
 
 	// UThumbnailRenderer implementation
-	virtual void Draw(UObject* Object, int32 X, int32 Y, uint32 Width, uint32 Height, FRenderTarget*, FCanvas* Canvas) override;
+	virtual void Draw(UObject* Object, int32 X, int32 Y, uint32 Width, uint32 Height, FRenderTarget*, FCanvas* Canvas, bool bAdditionalViewFamily) override;
 	virtual bool CanVisualizeAsset(UObject* Object) override;
 };
diff --git a/Engine/Source/Editor/FoliageEdit/Public/FoliageType_ISMThumbnailRenderer.h b/Engine/Source/Editor/FoliageEdit/Public/FoliageType_ISMThumbnailRenderer.h
--- a/Engine/Source/Editor/FoliageEdit/Public/FoliageType_ISMThumbnailRenderer.h
+++ b/Engine/Source/Editor/FoliageEdit/Public/FoliageType_ISMThumbnailRenderer.h
@@ -25,7 +25,7 @@ class UFoliageType_ISMThumbnailRenderer : public UDefaultSizedThumbnailRenderer
 	{}
 
 	// UThumbnailRenderer implementation
-	virtual void Draw(UObject* Object, int32 X, int32 Y, uint32 Width, uint32 Height, FRenderTarget*, FCanvas* Canvas) override;
+	virtual void Draw(UObject* Object, int32 X, int32 Y, uint32 Width, uint32 Height, FRenderTarget*, FCanvas* Canvas, bool bAdditionalViewFamily) override;
 	virtual bool CanVisualizeAsset(UObject* Object) override;
 	// UObject implementation
 	virtual void BeginDestroy() override;
diff --git a/Engine/Source/Editor/UnrealEd/Classes/Factories/Texture2dArrayThumbnailRenderer.h b/Engine/Source/Editor/UnrealEd/Classes/Factories/Texture2dArrayThumbnailRenderer.h
--- a/Engine/Source/Editor/UnrealEd/Classes/Factories/Texture2dArrayThumbnailRenderer.h
+++ b/Engine/Source/Editor/UnrealEd/Classes/Factories/Texture2dArrayThumbnailRenderer.h
@@ -19,6 +19,6 @@ class UTexture2DArrayThumbnailRenderer : public UTextureThumbnailRenderer
 
 	// Begin UThumbnailRenderer Interface
 	virtual void GetThumbnailSize(UObject* Object, float Zoom, uint32& OutWidth, uint32& OutHeight) const override;
-	virtual void Draw(UObject* Object, int32 X, int32 Y, uint32 Width, uint32 Height, FRenderTarget*, FCanvas* Canvas) override;
+	virtual void Draw(UObject* Object, int32 X, int32 Y, uint32 Width, uint32 Height, FRenderTarget*, FCanvas* Canvas, bool bAdditionalViewFamily) override;
 	// End UThumbnailRenderer Interface
 };
diff --git a/Engine/Source/Editor/UnrealEd/Classes/Factories/TextureCubeThumbnailRenderer.h b/Engine/Source/Editor/UnrealEd/Classes/Factories/TextureCubeThumbnailRenderer.h
--- a/Engine/Source/Editor/UnrealEd/Classes/Factories/TextureCubeThumbnailRenderer.h
+++ b/Engine/Source/Editor/UnrealEd/Classes/Factories/TextureCubeThumbnailRenderer.h
@@ -23,7 +23,7 @@ class UTextureCubeThumbnailRenderer : public UTextureThumbnailRenderer
 
 	//~ Begin UThumbnailRenderer Interface
 	virtual void GetThumbnailSize(UObject* Object, float Zoom, uint32& OutWidth, uint32& OutHeight) const override;
-	virtual void Draw(UObject* Object, int32 X, int32 Y, uint32 Width, uint32 Height, FRenderTarget*, FCanvas* Canvas) override;
+	virtual void Draw(UObject* Object, int32 X, int32 Y, uint32 Width, uint32 Height, FRenderTarget*, FCanvas* Canvas, bool bAdditionalViewFamily) override;
 	//~ End UThumbnailRenderer Interface
 };
 
diff --git a/Engine/Source/Editor/UnrealEd/Classes/ThumbnailRendering/AnimBlueprintThumbnailRenderer.h b/Engine/Source/Editor/UnrealEd/Classes/ThumbnailRendering/AnimBlueprintThumbnailRenderer.h
--- a/Engine/Source/Editor/UnrealEd/Classes/ThumbnailRendering/AnimBlueprintThumbnailRenderer.h
+++ b/Engine/Source/Editor/UnrealEd/Classes/ThumbnailRendering/AnimBlueprintThumbnailRenderer.h
@@ -23,7 +23,7 @@ class UAnimBlueprintThumbnailRenderer : public UDefaultSizedThumbnailRenderer
 
 	// Begin UThumbnailRenderer Object
 	UNREALED_API virtual bool CanVisualizeAsset(UObject* Object) override;
-	UNREALED_API virtual void Draw(UObject* Object, int32 X, int32 Y, uint32 Width, uint32 Height, FRenderTarget*, FCanvas* Canvas) override;
+	UNREALED_API virtual void Draw(UObject* Object, int32 X, int32 Y, uint32 Width, uint32 Height, FRenderTarget*, FCanvas* Canvas, bool bAdditionalViewFamily) override;
 	// End UThumbnailRenderer Object
 
 	// UObject Implementation
diff --git a/Engine/Source/Editor/UnrealEd/Classes/ThumbnailRendering/AnimSequenceThumbnailRenderer.h b/Engine/Source/Editor/UnrealEd/Classes/ThumbnailRendering/AnimSequenceThumbnailRenderer.h
--- a/Engine/Source/Editor/UnrealEd/Classes/ThumbnailRendering/AnimSequenceThumbnailRenderer.h
+++ b/Engine/Source/Editor/UnrealEd/Classes/ThumbnailRendering/AnimSequenceThumbnailRenderer.h
@@ -22,7 +22,7 @@ class UAnimSequenceThumbnailRenderer : public UDefaultSizedThumbnailRenderer
 
 
 	// Begin UThumbnailRenderer Object
-	UNREALED_API virtual void Draw(UObject* Object, int32 X, int32 Y, uint32 Width, uint32 Height, FRenderTarget*, FCanvas* Canvas) override;
+	UNREALED_API virtual void Draw(UObject* Object, int32 X, int32 Y, uint32 Width, uint32 Height, FRenderTarget*, FCanvas* Canvas, bool bAdditionalViewFamily) override;
 	// End UThumbnailRenderer Object
 
 	// UObject implementation
diff --git a/Engine/Source/Editor/UnrealEd/Classes/ThumbnailRendering/BlendSpaceThumbnailRenderer.h b/Engine/Source/Editor/UnrealEd/Classes/ThumbnailRendering/BlendSpaceThumbnailRenderer.h
--- a/Engine/Source/Editor/UnrealEd/Classes/ThumbnailRendering/BlendSpaceThumbnailRenderer.h
+++ b/Engine/Source/Editor/UnrealEd/Classes/ThumbnailRendering/BlendSpaceThumbnailRenderer.h
@@ -22,7 +22,7 @@ class UBlendSpaceThumbnailRenderer : public UDefaultSizedThumbnailRenderer
 
 
 	// Begin UThumbnailRenderer Object
-	UNREALED_API virtual void Draw(UObject* Object, int32 X, int32 Y, uint32 Width, uint32 Height, FRenderTarget*, FCanvas* Canvas) override;
+	UNREALED_API virtual void Draw(UObject* Object, int32 X, int32 Y, uint32 Width, uint32 Height, FRenderTarget*, FCanvas* Canvas, bool bAdditionalViewFamily) override;
 	// End UThumbnailRenderer Object
 
 	// UObject implementation
diff --git a/Engine/Source/Editor/UnrealEd/Classes/ThumbnailRendering/BlueprintThumbnailRenderer.h b/Engine/Source/Editor/UnrealEd/Classes/ThumbnailRendering/BlueprintThumbnailRenderer.h
--- a/Engine/Source/Editor/UnrealEd/Classes/ThumbnailRendering/BlueprintThumbnailRenderer.h
+++ b/Engine/Source/Editor/UnrealEd/Classes/ThumbnailRendering/BlueprintThumbnailRenderer.h
@@ -23,7 +23,7 @@ class UNREALED_VTABLE UBlueprintThumbnailRenderer : public UDefaultSizedThumbnai
 
 	// Begin UThumbnailRenderer Object
 	UNREALED_API virtual bool CanVisualizeAsset(UObject* Object) override;
-	UNREALED_API virtual void Draw(UObject* Object, int32 X, int32 Y, uint32 Width, uint32 Height, FRenderTarget*, FCanvas* Canvas) override;
+	UNREALED_API virtual void Draw(UObject* Object, int32 X, int32 Y, uint32 Width, uint32 Height, FRenderTarget* RenderTarget, FCanvas* Canvas, bool bAdditionalViewFamily) override;
 	// End UThumbnailRenderer Object
 
 	// UObject implementation
diff --git a/Engine/Source/Editor/UnrealEd/Classes/ThumbnailRendering/ClassThumbnailRenderer.h b/Engine/Source/Editor/UnrealEd/Classes/ThumbnailRendering/ClassThumbnailRenderer.h
--- a/Engine/Source/Editor/UnrealEd/Classes/ThumbnailRendering/ClassThumbnailRenderer.h
+++ b/Engine/Source/Editor/UnrealEd/Classes/ThumbnailRendering/ClassThumbnailRenderer.h
@@ -18,7 +18,7 @@ class UClassThumbnailRenderer : public UDefaultSizedThumbnailRenderer
 
 	// Begin UThumbnailRenderer Object
 	UNREALED_API virtual bool CanVisualizeAsset(UObject* Object) override;
-	UNREALED_API virtual void Draw(UObject* Object, int32 X, int32 Y, uint32 Width, uint32 Height, FRenderTarget*, FCanvas* Canvas) override;
+	UNREALED_API virtual void Draw(UObject* Object, int32 X, int32 Y, uint32 Width, uint32 Height, FRenderTarget*, FCanvas* Canvas, bool bAdditionalViewFamily) override;
 	// End UThumbnailRenderer Object
 
 	// UObject Implementation
diff --git a/Engine/Source/Editor/UnrealEd/Classes/ThumbnailRendering/CurveLinearColorThumbnailRenderer.h b/Engine/Source/Editor/UnrealEd/Classes/ThumbnailRendering/CurveLinearColorThumbnailRenderer.h
--- a/Engine/Source/Editor/UnrealEd/Classes/ThumbnailRendering/CurveLinearColorThumbnailRenderer.h
+++ b/Engine/Source/Editor/UnrealEd/Classes/ThumbnailRendering/CurveLinearColorThumbnailRenderer.h
@@ -23,7 +23,7 @@ public:
 
 	// Begin UThumbnailRenderer Object
 	virtual void GetThumbnailSize(UObject* Object, float Zoom, uint32& OutWidth, uint32& OutHeight) const override;
-	virtual void Draw(UObject* Object, int32 X, int32 Y, uint32 Width, uint32 Height, FRenderTarget*, FCanvas* Canvas) override;
+	virtual void Draw(UObject* Object, int32 X, int32 Y, uint32 Width, uint32 Height, FRenderTarget*, FCanvas* Canvas, bool bAdditionalViewFamily) override;
 	// End UThumbnailRenderer Object
 };
 
diff --git a/Engine/Source/Editor/UnrealEd/Classes/ThumbnailRendering/FontThumbnailRenderer.h b/Engine/Source/Editor/UnrealEd/Classes/ThumbnailRendering/FontThumbnailRenderer.h
--- a/Engine/Source/Editor/UnrealEd/Classes/ThumbnailRendering/FontThumbnailRenderer.h
+++ b/Engine/Source/Editor/UnrealEd/Classes/ThumbnailRendering/FontThumbnailRenderer.h
@@ -23,7 +23,7 @@ class UFontThumbnailRenderer : public UTextureThumbnailRenderer
 
 	// Begin UThumbnailRenderer Object
 	virtual void GetThumbnailSize(UObject* Object, float Zoom, uint32& OutWidth, uint32& OutHeight) const override;
-	virtual void Draw(UObject* Object, int32 X, int32 Y, uint32 Width, uint32 Height, FRenderTarget*, FCanvas* Canvas) override;
+	virtual void Draw(UObject* Object, int32 X, int32 Y, uint32 Width, uint32 Height, FRenderTarget*, FCanvas* Canvas, bool bAdditionalViewFamily) override;
 	// End UThumbnailRenderer Object
 };
 
diff --git a/Engine/Source/Editor/UnrealEd/Classes/ThumbnailRendering/LevelThumbnailRenderer.h b/Engine/Source/Editor/UnrealEd/Classes/ThumbnailRendering/LevelThumbnailRenderer.h
--- a/Engine/Source/Editor/UnrealEd/Classes/ThumbnailRendering/LevelThumbnailRenderer.h
+++ b/Engine/Source/Editor/UnrealEd/Classes/ThumbnailRendering/LevelThumbnailRenderer.h
@@ -24,7 +24,7 @@ class ULevelThumbnailRenderer : public UDefaultSizedThumbnailRenderer
 
 
 	// Begin UThumbnailRenderer Object
-	virtual void Draw(UObject* Object, int32 X, int32 Y, uint32 Width, uint32 Height, FRenderTarget*, FCanvas* Canvas) override;
+	virtual void Draw(UObject* Object, int32 X, int32 Y, uint32 Width, uint32 Height, FRenderTarget*, FCanvas* Canvas, bool bAdditionalViewFamily) override;
 	// End UThumbnailRenderer Object
 
 private:
diff --git a/Engine/Source/Editor/UnrealEd/Classes/ThumbnailRendering/MaterialFunctionThumbnailRenderer.h b/Engine/Source/Editor/UnrealEd/Classes/ThumbnailRendering/MaterialFunctionThumbnailRenderer.h
--- a/Engine/Source/Editor/UnrealEd/Classes/ThumbnailRendering/MaterialFunctionThumbnailRenderer.h
+++ b/Engine/Source/Editor/UnrealEd/Classes/ThumbnailRendering/MaterialFunctionThumbnailRenderer.h
@@ -22,7 +22,7 @@ class UMaterialFunctionThumbnailRenderer : public UDefaultSizedThumbnailRenderer
 
 
 	// Begin UThumbnailRenderer Object
-	virtual void Draw(UObject* Object, int32 X, int32 Y, uint32 Width, uint32 Height, FRenderTarget*, FCanvas* Canvas) override;
+	virtual void Draw(UObject* Object, int32 X, int32 Y, uint32 Width, uint32 Height, FRenderTarget*, FCanvas* Canvas, bool bAdditionalViewFamily) override;
 	// End UThumbnailRenderer Object
 
 	// UObject implementation
diff --git a/Engine/Source/Editor/UnrealEd/Classes/ThumbnailRendering/MaterialInstanceThumbnailRenderer.h b/Engine/Source/Editor/UnrealEd/Classes/ThumbnailRendering/MaterialInstanceThumbnailRenderer.h
--- a/Engine/Source/Editor/UnrealEd/Classes/ThumbnailRendering/MaterialInstanceThumbnailRenderer.h
+++ b/Engine/Source/Editor/UnrealEd/Classes/ThumbnailRendering/MaterialInstanceThumbnailRenderer.h
@@ -22,7 +22,7 @@ class UMaterialInstanceThumbnailRenderer : public UDefaultSizedThumbnailRenderer
 
 
 	// Begin UThumbnailRenderer Object
-	virtual void Draw(UObject* Object, int32 X, int32 Y, uint32 Width, uint32 Height, FRenderTarget*, FCanvas* Canvas) override;
+	virtual void Draw(UObject* Object, int32 X, int32 Y, uint32 Width, uint32 Height, FRenderTarget*, FCanvas* Canvas, bool bAdditionalViewFamily) override;
 	// End UThumbnailRenderer Object
 
 	// UObject implementation
diff --git a/Engine/Source/Editor/UnrealEd/Classes/ThumbnailRendering/ParticleSystemThumbnailRenderer.h b/Engine/Source/Editor/UnrealEd/Classes/ThumbnailRendering/ParticleSystemThumbnailRenderer.h
--- a/Engine/Source/Editor/UnrealEd/Classes/ThumbnailRendering/ParticleSystemThumbnailRenderer.h
+++ b/Engine/Source/Editor/UnrealEd/Classes/ThumbnailRendering/ParticleSystemThumbnailRenderer.h
@@ -29,7 +29,7 @@ class UParticleSystemThumbnailRenderer : public UTextureThumbnailRenderer
 
 	// Begin UThumbnailRenderer Object
 	virtual void GetThumbnailSize(UObject* Object, float Zoom, uint32& OutWidth, uint32& OutHeight) const override;
-	virtual void Draw(UObject* Object, int32 X, int32 Y, uint32 Width, uint32 Height, FRenderTarget*, FCanvas* Canvas) override;
+	virtual void Draw(UObject* Object, int32 X, int32 Y, uint32 Width, uint32 Height, FRenderTarget*, FCanvas* Canvas, bool bAdditionalViewFamily) override;
 	// End UThumbnailRenderer Object
 
 	// UObject implementation
diff --git a/Engine/Source/Editor/UnrealEd/Classes/ThumbnailRendering/PhysicalMaterialMaskThumbnailRenderer.h b/Engine/Source/Editor/UnrealEd/Classes/ThumbnailRendering/PhysicalMaterialMaskThumbnailRenderer.h
--- a/Engine/Source/Editor/UnrealEd/Classes/ThumbnailRendering/PhysicalMaterialMaskThumbnailRenderer.h
+++ b/Engine/Source/Editor/UnrealEd/Classes/ThumbnailRendering/PhysicalMaterialMaskThumbnailRenderer.h
@@ -22,7 +22,7 @@ class UNREALED_VTABLE UPhysicalMaterialMaskThumbnailRenderer : public UTextureTh
 
 	// Begin UThumbnailRenderer Object
 	UNREALED_API virtual void GetThumbnailSize(UObject* Object, float Zoom, uint32& OutWidth, uint32& OutHeight) const override;
-	UNREALED_API virtual void Draw(UObject* Object, int32 X, int32 Y, uint32 Width, uint32 Height, FRenderTarget*, FCanvas* Canvas) override;
+	UNREALED_API virtual void Draw(UObject* Object, int32 X, int32 Y, uint32 Width, uint32 Height, FRenderTarget*, FCanvas* Canvas, bool bAdditionalViewFamily) override;
 	// End UThumbnailRenderer Object
 };
 
diff --git a/Engine/Source/Editor/UnrealEd/Classes/ThumbnailRendering/PhysicsAssetThumbnailRenderer.h b/Engine/Source/Editor/UnrealEd/Classes/ThumbnailRendering/PhysicsAssetThumbnailRenderer.h
--- a/Engine/Source/Editor/UnrealEd/Classes/ThumbnailRendering/PhysicsAssetThumbnailRenderer.h
+++ b/Engine/Source/Editor/UnrealEd/Classes/ThumbnailRendering/PhysicsAssetThumbnailRenderer.h
@@ -17,7 +17,7 @@ class UPhysicsAssetThumbnailRenderer : public UDefaultSizedThumbnailRenderer
 
 
 	// Begin UThumbnailRenderer Object
-	UNREALED_API virtual void Draw(UObject* Object, int32 X, int32 Y, uint32 Width, uint32 Height, FRenderTarget*, FCanvas* Canvas) override;
+	UNREALED_API virtual void Draw(UObject* Object, int32 X, int32 Y, uint32 Width, uint32 Height, FRenderTarget*, FCanvas* Canvas, bool bAdditionalViewFamily) override;
 	// End UThumbnailRenderer Object
 
 	// UObject implementation
diff --git a/Engine/Source/Editor/UnrealEd/Classes/ThumbnailRendering/SkeletalMeshThumbnailRenderer.h b/Engine/Source/Editor/UnrealEd/Classes/ThumbnailRendering/SkeletalMeshThumbnailRenderer.h
--- a/Engine/Source/Editor/UnrealEd/Classes/ThumbnailRendering/SkeletalMeshThumbnailRenderer.h
+++ b/Engine/Source/Editor/UnrealEd/Classes/ThumbnailRendering/SkeletalMeshThumbnailRenderer.h
@@ -22,7 +22,7 @@ class UNREALED_API USkeletalMeshThumbnailRenderer : public UDefaultSizedThumbnai
 
 
 	// Begin UThumbnailRenderer Object
-	virtual void Draw(UObject* Object, int32 X, int32 Y, uint32 Width, uint32 Height, FRenderTarget*, FCanvas* Canvas) override;
+	virtual void Draw(UObject* Object, int32 X, int32 Y, uint32 Width, uint32 Height, FRenderTarget*, FCanvas* Canvas, bool bAdditionalViewFamily) override;
 	// End UThumbnailRenderer Object
 
 	// UObject implementation
diff --git a/Engine/Source/Editor/UnrealEd/Classes/ThumbnailRendering/SlateBrushThumbnailRenderer.h b/Engine/Source/Editor/UnrealEd/Classes/ThumbnailRendering/SlateBrushThumbnailRenderer.h
--- a/Engine/Source/Editor/UnrealEd/Classes/ThumbnailRendering/SlateBrushThumbnailRenderer.h
+++ b/Engine/Source/Editor/UnrealEd/Classes/ThumbnailRendering/SlateBrushThumbnailRenderer.h
@@ -23,7 +23,7 @@ class USlateBrushThumbnailRenderer : public UDefaultSizedThumbnailRenderer
 public:
 
 	// Begin UThumbnailRenderer Object
-	UNREALED_API virtual void Draw(UObject* Object, int32 X, int32 Y, uint32 Width, uint32 Height, FRenderTarget*, FCanvas* Canvas) override;
+	UNREALED_API virtual void Draw(UObject* Object, int32 X, int32 Y, uint32 Width, uint32 Height, FRenderTarget*, FCanvas* Canvas, bool bAdditionalViewFamily) override;
 	// End UThumbnailRenderer Object
 
 };
diff --git a/Engine/Source/Editor/UnrealEd/Classes/ThumbnailRendering/SoundWaveThumbnailRenderer.h b/Engine/Source/Editor/UnrealEd/Classes/ThumbnailRendering/SoundWaveThumbnailRenderer.h
--- a/Engine/Source/Editor/UnrealEd/Classes/ThumbnailRendering/SoundWaveThumbnailRenderer.h
+++ b/Engine/Source/Editor/UnrealEd/Classes/ThumbnailRendering/SoundWaveThumbnailRenderer.h
@@ -22,7 +22,7 @@ class USoundWaveThumbnailRenderer : public UThumbnailRenderer
 
 	// Begin UThumbnailRenderer Object
 	virtual bool CanVisualizeAsset(UObject* Object);
-	virtual void Draw(UObject* Object, int32 X, int32 Y, uint32 Width, uint32 Height, FRenderTarget*, FCanvas* Canvas) override;
+	virtual void Draw(UObject* Object, int32 X, int32 Y, uint32 Width, uint32 Height, FRenderTarget*, FCanvas* Canvas, bool bAdditionalViewFamily) override;
 	virtual bool AllowsRealtimeThumbnails(UObject* Object) const override;
 	// End UThumbnailRenderer Object
 };
diff --git a/Engine/Source/Editor/UnrealEd/Classes/ThumbnailRendering/StaticMeshThumbnailRenderer.h b/Engine/Source/Editor/UnrealEd/Classes/ThumbnailRendering/StaticMeshThumbnailRenderer.h
--- a/Engine/Source/Editor/UnrealEd/Classes/ThumbnailRendering/StaticMeshThumbnailRenderer.h
+++ b/Engine/Source/Editor/UnrealEd/Classes/ThumbnailRendering/StaticMeshThumbnailRenderer.h
@@ -22,7 +22,7 @@ class UStaticMeshThumbnailRenderer : public UDefaultSizedThumbnailRenderer
 
 
 	// Begin UThumbnailRenderer Object
-	UNREALED_API virtual void Draw(UObject* Object, int32 X, int32 Y, uint32 Width, uint32 Height, FRenderTarget* RenderTarget, FCanvas* Canvas) override;
+	UNREALED_API virtual void Draw(UObject* Object, int32 X, int32 Y, uint32 Width, uint32 Height, FRenderTarget* RenderTarget, FCanvas* Canvas, bool bAdditionalViewFamily) override;
 	// End UThumbnailRenderer Object
 
 	// UObject implementation
diff --git a/Engine/Source/Editor/UnrealEd/Classes/ThumbnailRendering/SubsurfaceProfileRenderer.h b/Engine/Source/Editor/UnrealEd/Classes/ThumbnailRendering/SubsurfaceProfileRenderer.h
--- a/Engine/Source/Editor/UnrealEd/Classes/ThumbnailRendering/SubsurfaceProfileRenderer.h
+++ b/Engine/Source/Editor/UnrealEd/Classes/ThumbnailRendering/SubsurfaceProfileRenderer.h
@@ -23,7 +23,7 @@ class USubsurfaceProfileRenderer : public UTextureThumbnailRenderer
 
 	// Begin UThumbnailRenderer Object
 	virtual void GetThumbnailSize(UObject* Object, float Zoom, uint32& OutWidth, uint32& OutHeight) const override;
-	virtual void Draw(UObject* Object, int32 X, int32 Y, uint32 Width, uint32 Height, FRenderTarget*, FCanvas* Canvas) override;
+	virtual void Draw(UObject* Object, int32 X, int32 Y, uint32 Width, uint32 Height, FRenderTarget*, FCanvas* Canvas, bool bAdditionalViewFamily) override;
 	// End UThumbnailRenderer Object
 };
 
diff --git a/Engine/Source/Editor/UnrealEd/Classes/ThumbnailRendering/TextureThumbnailRenderer.h b/Engine/Source/Editor/UnrealEd/Classes/ThumbnailRendering/TextureThumbnailRenderer.h
--- a/Engine/Source/Editor/UnrealEd/Classes/ThumbnailRendering/TextureThumbnailRenderer.h
+++ b/Engine/Source/Editor/UnrealEd/Classes/ThumbnailRendering/TextureThumbnailRenderer.h
@@ -23,7 +23,7 @@ class UNREALED_VTABLE UTextureThumbnailRenderer : public UThumbnailRenderer
 
 	// Begin UThumbnailRenderer Object
 	UNREALED_API virtual void GetThumbnailSize(UObject* Object, float Zoom, uint32& OutWidth, uint32& OutHeight) const override;
-	UNREALED_API virtual void Draw(UObject* Object, int32 X, int32 Y, uint32 Width, uint32 Height, FRenderTarget*, FCanvas* Canvas) override;
+	UNREALED_API virtual void Draw(UObject* Object, int32 X, int32 Y, uint32 Width, uint32 Height, FRenderTarget*, FCanvas* Canvas, bool bAdditionalViewFamily) override;
 	// End UThumbnailRenderer Object
 };
 
diff --git a/Engine/Source/Editor/UnrealEd/Classes/ThumbnailRendering/ThumbnailRenderer.h b/Engine/Source/Editor/UnrealEd/Classes/ThumbnailRendering/ThumbnailRenderer.h
--- a/Engine/Source/Editor/UnrealEd/Classes/ThumbnailRendering/ThumbnailRenderer.h
+++ b/Engine/Source/Editor/UnrealEd/Classes/ThumbnailRendering/ThumbnailRenderer.h
@@ -52,7 +52,22 @@ public:
 	 * @param Viewport the viewport being drawn in
 	 * @param Canvas the render interface to draw with
 	 */
-	virtual void Draw(UObject* Object, int32 X, int32 Y, uint32 Width, uint32 Height, FRenderTarget* Viewport, FCanvas* Canvas) PURE_VIRTUAL(UThumbnailRenderer::Draw,);
+	UE_DEPRECATED(4.24, "Please override the other prototype of the Draw function.")
+	virtual void Draw(UObject* Object, int32 X, int32 Y, uint32 Width, uint32 Height, FRenderTarget* Viewport, FCanvas* Canvas) { Draw(Object, X, Y, Width, Height, Viewport, Canvas, false); }
+
+	/**
+	 * Draws a thumbnail for the object that was specified.
+	 *
+	 * @param Object the object to draw the thumbnail for
+	 * @param X the X coordinate to start drawing at
+	 * @param Y the Y coordinate to start drawing at
+	 * @param Width the width of the thumbnail to draw
+	 * @param Height the height of the thumbnail to draw
+	 * @param Viewport the viewport being drawn in
+	 * @param Canvas the render interface to draw with
+	 * @param bAdditionalViewFamily whether this draw should write over the render target (true) or clear it before (false)
+	 */
+	virtual void Draw(UObject* Object, int32 X, int32 Y, uint32 Width, uint32 Height, FRenderTarget* Viewport, FCanvas* Canvas, bool bAdditionalViewFamily) PURE_VIRTUAL(UThumbnailRenderer::Draw, );
 
 	/**
 	 * Checks to see if the specified asset supports realtime thumbnails, which will cause them to always be rerendered to reflect any changes
diff --git a/Engine/Source/Editor/UnrealEd/Classes/ThumbnailRendering/VolumeTextureThumbnailRenderer.h b/Engine/Source/Editor/UnrealEd/Classes/ThumbnailRendering/VolumeTextureThumbnailRenderer.h
--- a/Engine/Source/Editor/UnrealEd/Classes/ThumbnailRendering/VolumeTextureThumbnailRenderer.h
+++ b/Engine/Source/Editor/UnrealEd/Classes/ThumbnailRendering/VolumeTextureThumbnailRenderer.h
@@ -23,7 +23,7 @@ class UVolumeTextureThumbnailRenderer : public UDefaultSizedThumbnailRenderer
 
 
 	// Begin UThumbnailRenderer Object
-	virtual void Draw(UObject* Object, int32 X, int32 Y, uint32 Width, uint32 Height, FRenderTarget*, FCanvas* Canvas) override;
+	virtual void Draw(UObject* Object, int32 X, int32 Y, uint32 Width, uint32 Height, FRenderTarget*, FCanvas* Canvas, bool bAdditionalViewFamily) override;
 	// End UThumbnailRenderer Object
 
 	// UObject implementation
diff --git a/Engine/Source/Editor/UnrealEd/Classes/ThumbnailRendering/WorldThumbnailRenderer.h b/Engine/Source/Editor/UnrealEd/Classes/ThumbnailRendering/WorldThumbnailRenderer.h
--- a/Engine/Source/Editor/UnrealEd/Classes/ThumbnailRendering/WorldThumbnailRenderer.h
+++ b/Engine/Source/Editor/UnrealEd/Classes/ThumbnailRendering/WorldThumbnailRenderer.h
@@ -23,7 +23,7 @@ class UWorldThumbnailRenderer : public UDefaultSizedThumbnailRenderer
 
 	// Begin UThumbnailRenderer Object
 	virtual bool CanVisualizeAsset(UObject* Object) override;
-	virtual void Draw(UObject* Object, int32 X, int32 Y, uint32 Width, uint32 Height, FRenderTarget*, FCanvas* Canvas) override;
+	virtual void Draw(UObject* Object, int32 X, int32 Y, uint32 Width, uint32 Height, FRenderTarget*, FCanvas* Canvas, bool bAdditionalViewFamily) override;
 	// End UThumbnailRenderer Object
 
 private:
diff --git a/Engine/Source/Editor/UnrealEd/Private/ObjectTools.cpp b/Engine/Source/Editor/UnrealEd/Private/ObjectTools.cpp
--- a/Engine/Source/Editor/UnrealEd/Private/ObjectTools.cpp
+++ b/Engine/Source/Editor/UnrealEd/Private/ObjectTools.cpp
@@ -4176,6 +4176,7 @@ namespace ThumbnailTools
 			// Draw the thumbnail
 			const int32 XPos = 0;
 			const int32 YPos = 0;
+			const bool bAdditionalViewFamily = false;
 			RenderInfo->Renderer->Draw(
 				InObject,
 				XPos,
@@ -4183,7 +4184,8 @@ namespace ThumbnailTools
 				DrawWidth,
 				DrawHeight,
 				RenderTargetResource,
-				&Canvas
+				&Canvas, 
+				bAdditionalViewFamily
 				);
 		}
 
diff --git a/Engine/Source/Editor/UnrealEd/Private/ThumbnailRendering/AnimBlueprintThumbnailRenderer.cpp b/Engine/Source/Editor/UnrealEd/Private/ThumbnailRendering/AnimBlueprintThumbnailRenderer.cpp
--- a/Engine/Source/Editor/UnrealEd/Private/ThumbnailRendering/AnimBlueprintThumbnailRenderer.cpp
+++ b/Engine/Source/Editor/UnrealEd/Private/ThumbnailRendering/AnimBlueprintThumbnailRenderer.cpp
@@ -17,7 +17,7 @@ bool UAnimBlueprintThumbnailRenderer::CanVisualizeAsset(UObject* Object)
 	return AnimBlueprint->BlueprintType != BPTYPE_Interface;
 }
 
-void UAnimBlueprintThumbnailRenderer::Draw(UObject* Object, int32 X, int32 Y, uint32 Width, uint32 Height, FRenderTarget* RenderTarget, FCanvas* Canvas)
+void UAnimBlueprintThumbnailRenderer::Draw(UObject* Object, int32 X, int32 Y, uint32 Width, uint32 Height, FRenderTarget* RenderTarget, FCanvas* Canvas, bool bAdditionalViewFamily)
 {
 	UAnimBlueprint* AnimBlueprint = Cast<UAnimBlueprint>(Object);
 	if (AnimBlueprint && AnimBlueprint->GeneratedClass)
@@ -27,7 +27,8 @@ void UAnimBlueprintThumbnailRenderer::Draw(UObject* Object, int32 X, int32 Y, ui
 		if(ThumbnailScene->SetAnimBlueprint(AnimBlueprint))
 		{
 			FSceneViewFamilyContext ViewFamily(FSceneViewFamily::ConstructionValues(RenderTarget, ThumbnailScene->GetScene(), FEngineShowFlags(ESFIM_Game))
-				.SetWorldTimes(FApp::GetCurrentTime() - GStartTime, FApp::GetDeltaTime(), FApp::GetCurrentTime() - GStartTime));
+				.SetWorldTimes(FApp::GetCurrentTime() - GStartTime, FApp::GetDeltaTime(), FApp::GetCurrentTime() - GStartTime)
+				.SetAdditionalViewFamily(bAdditionalViewFamily));
 
 			ViewFamily.EngineShowFlags.DisableAdvancedFeatures();
 			ViewFamily.EngineShowFlags.MotionBlur = 0;
diff --git a/Engine/Source/Editor/UnrealEd/Private/ThumbnailRendering/AnimSequenceThumbnailRenderer.cpp b/Engine/Source/Editor/UnrealEd/Private/ThumbnailRendering/AnimSequenceThumbnailRenderer.cpp
--- a/Engine/Source/Editor/UnrealEd/Private/ThumbnailRendering/AnimSequenceThumbnailRenderer.cpp
+++ b/Engine/Source/Editor/UnrealEd/Private/ThumbnailRendering/AnimSequenceThumbnailRenderer.cpp
@@ -13,7 +13,7 @@ UAnimSequenceThumbnailRenderer::UAnimSequenceThumbnailRenderer(const FObjectInit
 	ThumbnailScene = nullptr;
 }
 
-void UAnimSequenceThumbnailRenderer::Draw(UObject* Object, int32 X, int32 Y, uint32 Width, uint32 Height, FRenderTarget* RenderTarget, FCanvas* Canvas)
+void UAnimSequenceThumbnailRenderer::Draw(UObject* Object, int32 X, int32 Y, uint32 Width, uint32 Height, FRenderTarget* RenderTarget, FCanvas* Canvas, bool bAdditionalViewFamily)
 {
 	UAnimSequenceBase* Anim = Cast<UAnimSequenceBase>(Object);
 	if (Anim != nullptr)
@@ -26,7 +26,8 @@ void UAnimSequenceThumbnailRenderer::Draw(UObject* Object, int32 X, int32 Y, uin
 		if(ThumbnailScene->SetAnimation(Anim))
 		{
 			FSceneViewFamilyContext ViewFamily(FSceneViewFamily::ConstructionValues(RenderTarget, ThumbnailScene->GetScene(), FEngineShowFlags(ESFIM_Game))
-				.SetWorldTimes(FApp::GetCurrentTime() - GStartTime, FApp::GetDeltaTime(), FApp::GetCurrentTime() - GStartTime));
+				.SetWorldTimes(FApp::GetCurrentTime() - GStartTime, FApp::GetDeltaTime(), FApp::GetCurrentTime() - GStartTime)
+				.SetAdditionalViewFamily(bAdditionalViewFamily));
 
 			ViewFamily.EngineShowFlags.DisableAdvancedFeatures();
 			ViewFamily.EngineShowFlags.MotionBlur = 0;
diff --git a/Engine/Source/Editor/UnrealEd/Private/ThumbnailRendering/BlendSpaceThumbnailRenderer.cpp b/Engine/Source/Editor/UnrealEd/Private/ThumbnailRendering/BlendSpaceThumbnailRenderer.cpp
--- a/Engine/Source/Editor/UnrealEd/Private/ThumbnailRendering/BlendSpaceThumbnailRenderer.cpp
+++ b/Engine/Source/Editor/UnrealEd/Private/ThumbnailRendering/BlendSpaceThumbnailRenderer.cpp
@@ -14,7 +14,7 @@ UBlendSpaceThumbnailRenderer::UBlendSpaceThumbnailRenderer(const FObjectInitiali
 	ThumbnailScene = nullptr;
 }
 
-void UBlendSpaceThumbnailRenderer::Draw(UObject* Object, int32 X, int32 Y, uint32 Width, uint32 Height, FRenderTarget* RenderTarget, FCanvas* Canvas)
+void UBlendSpaceThumbnailRenderer::Draw(UObject* Object, int32 X, int32 Y, uint32 Width, uint32 Height, FRenderTarget* RenderTarget, FCanvas* Canvas, bool bAdditionalViewFamily)
 {
 	UBlendSpaceBase* BlendSpace = Cast<UBlendSpaceBase>(Object);
 	if (BlendSpace != nullptr)
@@ -27,7 +27,8 @@ void UBlendSpaceThumbnailRenderer::Draw(UObject* Object, int32 X, int32 Y, uint3
 		if (ThumbnailScene->SetBlendSpace(BlendSpace))
 		{
 			FSceneViewFamilyContext ViewFamily(FSceneViewFamily::ConstructionValues(RenderTarget, ThumbnailScene->GetScene(), FEngineShowFlags(ESFIM_Game))
-				.SetWorldTimes(FApp::GetCurrentTime() - GStartTime, FApp::GetDeltaTime(), FApp::GetCurrentTime() - GStartTime));
+				.SetWorldTimes(FApp::GetCurrentTime() - GStartTime, FApp::GetDeltaTime(), FApp::GetCurrentTime() - GStartTime)
+				.SetAdditionalViewFamily(bAdditionalViewFamily));
 
 			ViewFamily.EngineShowFlags.DisableAdvancedFeatures();
 			ViewFamily.EngineShowFlags.MotionBlur = 0;
diff --git a/Engine/Source/Editor/UnrealEd/Private/ThumbnailRendering/BlueprintThumbnailRenderer.cpp b/Engine/Source/Editor/UnrealEd/Private/ThumbnailRendering/BlueprintThumbnailRenderer.cpp
--- a/Engine/Source/Editor/UnrealEd/Private/ThumbnailRendering/BlueprintThumbnailRenderer.cpp
+++ b/Engine/Source/Editor/UnrealEd/Private/ThumbnailRendering/BlueprintThumbnailRenderer.cpp
@@ -69,7 +69,7 @@ bool UBlueprintThumbnailRenderer::CanVisualizeAsset(UObject* Object)
 	return false;
 }
 
-void UBlueprintThumbnailRenderer::Draw(UObject* Object, int32 X, int32 Y, uint32 Width, uint32 Height, FRenderTarget* RenderTarget, FCanvas* Canvas)
+void UBlueprintThumbnailRenderer::Draw(UObject* Object, int32 X, int32 Y, uint32 Width, uint32 Height, FRenderTarget* RenderTarget, FCanvas* Canvas, bool bAdditionalViewFamily)
 {
 	UBlueprint* Blueprint = Cast<UBlueprint>(Object);
 
@@ -86,7 +86,8 @@ void UBlueprintThumbnailRenderer::Draw(UObject* Object, int32 X, int32 Y, uint32
 
 		ThumbnailScene->SetBlueprint(Blueprint);
 		FSceneViewFamilyContext ViewFamily( FSceneViewFamily::ConstructionValues( RenderTarget, ThumbnailScene->GetScene(), FEngineShowFlags(ESFIM_Game) )
-			.SetWorldTimes(FApp::GetCurrentTime() - GStartTime, FApp::GetDeltaTime(), FApp::GetCurrentTime() - GStartTime));
+			.SetWorldTimes(FApp::GetCurrentTime() - GStartTime, FApp::GetDeltaTime(), FApp::GetCurrentTime() - GStartTime)
+			.SetAdditionalViewFamily(bAdditionalViewFamily));
 
 		ViewFamily.EngineShowFlags.DisableAdvancedFeatures();
 		ViewFamily.EngineShowFlags.MotionBlur = 0;
diff --git a/Engine/Source/Editor/UnrealEd/Private/ThumbnailRendering/ClassThumbnailRenderer.cpp b/Engine/Source/Editor/UnrealEd/Private/ThumbnailRendering/ClassThumbnailRenderer.cpp
--- a/Engine/Source/Editor/UnrealEd/Private/ThumbnailRendering/ClassThumbnailRenderer.cpp
+++ b/Engine/Source/Editor/UnrealEd/Private/ThumbnailRendering/ClassThumbnailRenderer.cpp
@@ -32,7 +32,7 @@ bool UClassThumbnailRenderer::CanVisualizeAsset(UObject* Object)
 	return false;
 }
 
-void UClassThumbnailRenderer::Draw(UObject* Object, int32 X, int32 Y, uint32 Width, uint32 Height, FRenderTarget* RenderTarget, FCanvas* Canvas)
+void UClassThumbnailRenderer::Draw(UObject* Object, int32 X, int32 Y, uint32 Width, uint32 Height, FRenderTarget* RenderTarget, FCanvas* Canvas, bool bAdditionalViewFamily)
 {
 	UClass* Class = Cast<UClass>(Object);
 	if (Class != nullptr)
@@ -41,7 +41,8 @@ void UClassThumbnailRenderer::Draw(UObject* Object, int32 X, int32 Y, uint32 Wid
 
 		ThumbnailScene->SetClass(Class);
 		FSceneViewFamilyContext ViewFamily( FSceneViewFamily::ConstructionValues( RenderTarget, ThumbnailScene->GetScene(), FEngineShowFlags(ESFIM_Game) )
-			.SetWorldTimes(FApp::GetCurrentTime() - GStartTime, FApp::GetDeltaTime(), FApp::GetCurrentTime() - GStartTime));
+			.SetWorldTimes(FApp::GetCurrentTime() - GStartTime, FApp::GetDeltaTime(), FApp::GetCurrentTime() - GStartTime)
+			.SetAdditionalViewFamily(bAdditionalViewFamily));
 
 		ViewFamily.EngineShowFlags.DisableAdvancedFeatures();
 		ViewFamily.EngineShowFlags.MotionBlur = 0;
diff --git a/Engine/Source/Editor/UnrealEd/Private/ThumbnailRendering/CurveLinearColorThumbnailRenderer.cpp b/Engine/Source/Editor/UnrealEd/Private/ThumbnailRendering/CurveLinearColorThumbnailRenderer.cpp
--- a/Engine/Source/Editor/UnrealEd/Private/ThumbnailRendering/CurveLinearColorThumbnailRenderer.cpp
+++ b/Engine/Source/Editor/UnrealEd/Private/ThumbnailRendering/CurveLinearColorThumbnailRenderer.cpp
@@ -29,7 +29,7 @@ void UCurveLinearColorThumbnailRenderer::GetThumbnailSize(UObject* Object, float
 	}
 }
 
-void UCurveLinearColorThumbnailRenderer::Draw(UObject* Object, int32 X, int32 Y, uint32 Width, uint32 Height, FRenderTarget* Viewport, FCanvas* Canvas)
+void UCurveLinearColorThumbnailRenderer::Draw(UObject* Object, int32 X, int32 Y, uint32 Width, uint32 Height, FRenderTarget* Viewport, FCanvas* Canvas, bool bAdditionalViewFamily)
 {
 	UCurveLinearColor* GradientCurve = Cast<UCurveLinearColor>(Object);
 	if (GradientCurve)
diff --git a/Engine/Source/Editor/UnrealEd/Private/ThumbnailRendering/FontThumbnailRenderer.cpp b/Engine/Source/Editor/UnrealEd/Private/ThumbnailRendering/FontThumbnailRenderer.cpp
--- a/Engine/Source/Editor/UnrealEd/Private/ThumbnailRendering/FontThumbnailRenderer.cpp
+++ b/Engine/Source/Editor/UnrealEd/Private/ThumbnailRendering/FontThumbnailRenderer.cpp
@@ -48,7 +48,7 @@ void UFontThumbnailRenderer::GetThumbnailSize(UObject* Object, float Zoom, uint3
 	OutWidth = OutHeight = 0;
 }
 
-void UFontThumbnailRenderer::Draw(UObject* Object, int32 X, int32 Y, uint32 Width, uint32 Height, FRenderTarget*, FCanvas* Canvas)
+void UFontThumbnailRenderer::Draw(UObject* Object, int32 X, int32 Y, uint32 Width, uint32 Height, FRenderTarget*, FCanvas* Canvas, bool bAdditionalViewFamily)
 {
 	UFont* const Font = Cast<UFont>(Object);
 
diff --git a/Engine/Source/Editor/UnrealEd/Private/ThumbnailRendering/LevelThumbnailRenderer.cpp b/Engine/Source/Editor/UnrealEd/Private/ThumbnailRendering/LevelThumbnailRenderer.cpp
--- a/Engine/Source/Editor/UnrealEd/Private/ThumbnailRendering/LevelThumbnailRenderer.cpp
+++ b/Engine/Source/Editor/UnrealEd/Private/ThumbnailRendering/LevelThumbnailRenderer.cpp
@@ -13,13 +13,14 @@ ULevelThumbnailRenderer::ULevelThumbnailRenderer(const FObjectInitializer& Objec
 {
 }
 
-void ULevelThumbnailRenderer::Draw(UObject* Object, int32 X, int32 Y, uint32 Width, uint32 Height, FRenderTarget* RenderTarget, FCanvas* Canvas)
+void ULevelThumbnailRenderer::Draw(UObject* Object, int32 X, int32 Y, uint32 Width, uint32 Height, FRenderTarget* RenderTarget, FCanvas* Canvas, bool bAdditionalViewFamily)
 {
 	ULevel* Level = Cast<ULevel>(Object);
 	if (Level != nullptr)
 	{
 		FSceneViewFamilyContext ViewFamily( FSceneViewFamily::ConstructionValues( RenderTarget, Level->OwningWorld->Scene, FEngineShowFlags(ESFIM_Game) )
-			.SetWorldTimes(FApp::GetCurrentTime() - GStartTime, FApp::GetDeltaTime(), FApp::GetCurrentTime() - GStartTime));
+			.SetWorldTimes(FApp::GetCurrentTime() - GStartTime, FApp::GetDeltaTime(), FApp::GetCurrentTime() - GStartTime)
+			.SetAdditionalViewFamily(bAdditionalViewFamily));
 
 		ViewFamily.EngineShowFlags.DisableAdvancedFeatures();
 		ViewFamily.EngineShowFlags.MotionBlur = 0;
diff --git a/Engine/Source/Editor/UnrealEd/Private/ThumbnailRendering/MaterialFunctionThumbnailRenderer.cpp b/Engine/Source/Editor/UnrealEd/Private/ThumbnailRendering/MaterialFunctionThumbnailRenderer.cpp
--- a/Engine/Source/Editor/UnrealEd/Private/ThumbnailRendering/MaterialFunctionThumbnailRenderer.cpp
+++ b/Engine/Source/Editor/UnrealEd/Private/ThumbnailRendering/MaterialFunctionThumbnailRenderer.cpp
@@ -17,7 +17,7 @@ UMaterialFunctionThumbnailRenderer::UMaterialFunctionThumbnailRenderer(const FOb
 	ThumbnailScene = nullptr;
 }
 
-void UMaterialFunctionThumbnailRenderer::Draw(UObject* Object, int32 X, int32 Y, uint32 Width, uint32 Height, FRenderTarget* RenderTarget, FCanvas* Canvas)
+void UMaterialFunctionThumbnailRenderer::Draw(UObject* Object, int32 X, int32 Y, uint32 Width, uint32 Height, FRenderTarget* RenderTarget, FCanvas* Canvas, bool bAdditionalViewFamily)
 {
 	UMaterialFunctionInterface* MatFunc = Cast<UMaterialFunctionInterface>(Object);
 	UMaterialFunctionInstance* MatFuncInst = Cast<UMaterialFunctionInstance>(Object);
@@ -50,7 +50,8 @@ void UMaterialFunctionThumbnailRenderer::Draw(UObject* Object, int32 X, int32 Y,
 				ThumbnailScene->SetMaterialInterface(PreviewMaterial);
 	
 			FSceneViewFamilyContext ViewFamily( FSceneViewFamily::ConstructionValues( RenderTarget, ThumbnailScene->GetScene(), FEngineShowFlags(ESFIM_Game) )
-				.SetWorldTimes(FApp::GetCurrentTime() - GStartTime, FApp::GetDeltaTime(), FApp::GetCurrentTime() - GStartTime));
+				.SetWorldTimes(FApp::GetCurrentTime() - GStartTime, FApp::GetDeltaTime(), FApp::GetCurrentTime() - GStartTime)
+				.SetAdditionalViewFamily(bAdditionalViewFamily));
 
 			ViewFamily.EngineShowFlags.DisableAdvancedFeatures();
 			ViewFamily.EngineShowFlags.MotionBlur = 0;
diff --git a/Engine/Source/Editor/UnrealEd/Private/ThumbnailRendering/MaterialInstanceThumbnailRenderer.cpp b/Engine/Source/Editor/UnrealEd/Private/ThumbnailRendering/MaterialInstanceThumbnailRenderer.cpp
--- a/Engine/Source/Editor/UnrealEd/Private/ThumbnailRendering/MaterialInstanceThumbnailRenderer.cpp
+++ b/Engine/Source/Editor/UnrealEd/Private/ThumbnailRendering/MaterialInstanceThumbnailRenderer.cpp
@@ -21,7 +21,7 @@ UMaterialInstanceThumbnailRenderer::UMaterialInstanceThumbnailRenderer(const FOb
 {
 }
 
-void UMaterialInstanceThumbnailRenderer::Draw(UObject* Object, int32 X, int32 Y, uint32 Width, uint32 Height, FRenderTarget* RenderTarget, FCanvas* Canvas)
+void UMaterialInstanceThumbnailRenderer::Draw(UObject* Object, int32 X, int32 Y, uint32 Width, uint32 Height, FRenderTarget* RenderTarget, FCanvas* Canvas, bool bAdditionalViewFamily)
 {
 	UMaterialInterface* MatInst = Cast<UMaterialInterface>(Object);
 	if (MatInst != nullptr)
@@ -87,7 +87,8 @@ void UMaterialInstanceThumbnailRenderer::Draw(UObject* Object, int32 X, int32 Y,
 
 			ThumbnailScene->SetMaterialInterface(MatInst);
 			FSceneViewFamilyContext ViewFamily( FSceneViewFamily::ConstructionValues( RenderTarget, ThumbnailScene->GetScene(), FEngineShowFlags(ESFIM_Game) )
-				.SetWorldTimes(FApp::GetCurrentTime() - GStartTime, FApp::GetDeltaTime(), FApp::GetCurrentTime() - GStartTime));
+				.SetWorldTimes(FApp::GetCurrentTime() - GStartTime, FApp::GetDeltaTime(), FApp::GetCurrentTime() - GStartTime)
+				.SetAdditionalViewFamily(bAdditionalViewFamily));
 
 
 			ViewFamily.EngineShowFlags.DisableAdvancedFeatures();
diff --git a/Engine/Source/Editor/UnrealEd/Private/ThumbnailRendering/ParticleSystemThumbnailRenderer.cpp b/Engine/Source/Editor/UnrealEd/Private/ThumbnailRendering/ParticleSystemThumbnailRenderer.cpp
--- a/Engine/Source/Editor/UnrealEd/Private/ThumbnailRendering/ParticleSystemThumbnailRenderer.cpp
+++ b/Engine/Source/Editor/UnrealEd/Private/ThumbnailRendering/ParticleSystemThumbnailRenderer.cpp
@@ -61,7 +61,7 @@ void UParticleSystemThumbnailRenderer::GetThumbnailSize(UObject* Object, float Z
 	}
 }
 
-void UParticleSystemThumbnailRenderer::Draw(UObject* Object, int32 X, int32 Y, uint32 Width, uint32 Height, FRenderTarget* RenderTarget,FCanvas* Canvas)
+void UParticleSystemThumbnailRenderer::Draw(UObject* Object, int32 X, int32 Y, uint32 Width, uint32 Height, FRenderTarget* RenderTarget, FCanvas* Canvas, bool bAdditionalViewFamily)
 {
 	if (GUnrealEd->GetThumbnailManager())
 	{
@@ -77,7 +77,8 @@ void UParticleSystemThumbnailRenderer::Draw(UObject* Object, int32 X, int32 Y, u
 
 				ThumbnailScene->SetParticleSystem(ParticleSystem);
 				FSceneViewFamilyContext ViewFamily( FSceneViewFamily::ConstructionValues( RenderTarget, ThumbnailScene->GetScene(), FEngineShowFlags(ESFIM_Game))
-					.SetWorldTimes(FApp::GetCurrentTime() - GStartTime, FApp::GetDeltaTime(), FApp::GetCurrentTime() - GStartTime));
+					.SetWorldTimes(FApp::GetCurrentTime() - GStartTime, FApp::GetDeltaTime(), FApp::GetCurrentTime() - GStartTime)
+					.SetAdditionalViewFamily(bAdditionalViewFamily));
 			
 				ViewFamily.EngineShowFlags.DisableAdvancedFeatures();
 				ViewFamily.EngineShowFlags.MotionBlur = 0;
diff --git a/Engine/Source/Editor/UnrealEd/Private/ThumbnailRendering/PhysicalMaterialMaskThumbnailRenderer.cpp b/Engine/Source/Editor/UnrealEd/Private/ThumbnailRendering/PhysicalMaterialMaskThumbnailRenderer.cpp
--- a/Engine/Source/Editor/UnrealEd/Private/ThumbnailRendering/PhysicalMaterialMaskThumbnailRenderer.cpp
+++ b/Engine/Source/Editor/UnrealEd/Private/ThumbnailRendering/PhysicalMaterialMaskThumbnailRenderer.cpp
@@ -21,11 +21,11 @@ void UPhysicalMaterialMaskThumbnailRenderer::GetThumbnailSize(UObject* Object, f
 	}
 }
 
-void UPhysicalMaterialMaskThumbnailRenderer::Draw(UObject* Object, int32 X, int32 Y, uint32 Width, uint32 Height, FRenderTarget* RenderTarget, FCanvas* Canvas)
+void UPhysicalMaterialMaskThumbnailRenderer::Draw(UObject* Object, int32 X, int32 Y, uint32 Width, uint32 Height, FRenderTarget* RenderTarget, FCanvas* Canvas, bool bAdditionalViewFamily)
 {
 	UPhysicalMaterialMask* PhysMatMask = Cast<UPhysicalMaterialMask>(Object);
 	if (PhysMatMask->MaskTexture != nullptr)
 	{
-		Super::Draw(PhysMatMask->MaskTexture, X, Y, Width, Height, RenderTarget, Canvas);
+		Super::Draw(PhysMatMask->MaskTexture, X, Y, Width, Height, RenderTarget, Canvas, bAdditionalViewFamily);
 	}
 }
diff --git a/Engine/Source/Editor/UnrealEd/Private/ThumbnailRendering/PhysicsAssetThumbnailRenderer.cpp b/Engine/Source/Editor/UnrealEd/Private/ThumbnailRendering/PhysicsAssetThumbnailRenderer.cpp
--- a/Engine/Source/Editor/UnrealEd/Private/ThumbnailRendering/PhysicsAssetThumbnailRenderer.cpp
+++ b/Engine/Source/Editor/UnrealEd/Private/ThumbnailRendering/PhysicsAssetThumbnailRenderer.cpp
@@ -13,7 +13,7 @@ UPhysicsAssetThumbnailRenderer::UPhysicsAssetThumbnailRenderer(const FObjectInit
 	ThumbnailScene = nullptr;
 }
 
-void UPhysicsAssetThumbnailRenderer::Draw(UObject* Object, int32 X, int32 Y, uint32 Width, uint32 Height, FRenderTarget* RenderTarget, FCanvas* Canvas)
+void UPhysicsAssetThumbnailRenderer::Draw(UObject* Object, int32 X, int32 Y, uint32 Width, uint32 Height, FRenderTarget* RenderTarget, FCanvas* Canvas, bool bAdditionalViewFamily)
 {
 	UPhysicsAsset* PhysicsAsset = Cast<UPhysicsAsset>(Object);
 	if (PhysicsAsset != nullptr)
@@ -25,7 +25,8 @@ void UPhysicsAssetThumbnailRenderer::Draw(UObject* Object, int32 X, int32 Y, uin
 
 		ThumbnailScene->SetPhysicsAsset(PhysicsAsset);
 		FSceneViewFamilyContext ViewFamily( FSceneViewFamily::ConstructionValues( RenderTarget, ThumbnailScene->GetScene(), FEngineShowFlags(ESFIM_Game) )
-			.SetWorldTimes(FApp::GetCurrentTime() - GStartTime, FApp::GetDeltaTime(), FApp::GetCurrentTime() - GStartTime));
+			.SetWorldTimes(FApp::GetCurrentTime() - GStartTime, FApp::GetDeltaTime(), FApp::GetCurrentTime() - GStartTime)
+			.SetAdditionalViewFamily(bAdditionalViewFamily));
 
 		ViewFamily.EngineShowFlags.DisableAdvancedFeatures();
 		ViewFamily.EngineShowFlags.MotionBlur = 0;
diff --git a/Engine/Source/Editor/UnrealEd/Private/ThumbnailRendering/SkeletalMeshThumbnailRenderer.cpp b/Engine/Source/Editor/UnrealEd/Private/ThumbnailRendering/SkeletalMeshThumbnailRenderer.cpp
--- a/Engine/Source/Editor/UnrealEd/Private/ThumbnailRendering/SkeletalMeshThumbnailRenderer.cpp
+++ b/Engine/Source/Editor/UnrealEd/Private/ThumbnailRendering/SkeletalMeshThumbnailRenderer.cpp
@@ -13,7 +13,7 @@ USkeletalMeshThumbnailRenderer::USkeletalMeshThumbnailRenderer(const FObjectInit
 	ThumbnailScene = nullptr;
 }
 
-void USkeletalMeshThumbnailRenderer::Draw(UObject* Object, int32 X, int32 Y, uint32 Width, uint32 Height, FRenderTarget* RenderTarget, FCanvas* Canvas)
+void USkeletalMeshThumbnailRenderer::Draw(UObject* Object, int32 X, int32 Y, uint32 Width, uint32 Height, FRenderTarget* RenderTarget, FCanvas* Canvas, bool bAdditionalViewFamily)
 {
 	USkeletalMesh* SkeletalMesh = Cast<USkeletalMesh>(Object);
 	if (SkeletalMesh != nullptr)
@@ -27,7 +27,8 @@ void USkeletalMeshThumbnailRenderer::Draw(UObject* Object, int32 X, int32 Y, uin
 		AddAdditionalPreviewSceneContent(Object, ThumbnailScene->GetWorld());
 
 		FSceneViewFamilyContext ViewFamily( FSceneViewFamily::ConstructionValues( RenderTarget, ThumbnailScene->GetScene(), FEngineShowFlags(ESFIM_Game) )
-			.SetWorldTimes(FApp::GetCurrentTime() - GStartTime, FApp::GetDeltaTime(), FApp::GetCurrentTime() - GStartTime));
+			.SetWorldTimes(FApp::GetCurrentTime() - GStartTime, FApp::GetDeltaTime(), FApp::GetCurrentTime() - GStartTime)
+			.SetAdditionalViewFamily(bAdditionalViewFamily));
 
 		ViewFamily.EngineShowFlags.DisableAdvancedFeatures();
 		ViewFamily.EngineShowFlags.MotionBlur = 0;
diff --git a/Engine/Source/Editor/UnrealEd/Private/ThumbnailRendering/SlateBrushThumbnailRenderer.cpp b/Engine/Source/Editor/UnrealEd/Private/ThumbnailRendering/SlateBrushThumbnailRenderer.cpp
--- a/Engine/Source/Editor/UnrealEd/Private/ThumbnailRendering/SlateBrushThumbnailRenderer.cpp
+++ b/Engine/Source/Editor/UnrealEd/Private/ThumbnailRendering/SlateBrushThumbnailRenderer.cpp
@@ -14,7 +14,7 @@ USlateBrushThumbnailRenderer::USlateBrushThumbnailRenderer(const FObjectInitiali
 {
 }
 
-void USlateBrushThumbnailRenderer::Draw(UObject* Object, int32 X, int32 Y, uint32 Width, uint32 Height, FRenderTarget* RenderTarget, FCanvas* Canvas)
+void USlateBrushThumbnailRenderer::Draw(UObject* Object, int32 X, int32 Y, uint32 Width, uint32 Height, FRenderTarget* RenderTarget, FCanvas* Canvas, bool bAdditionalViewFamily)
 {
 	USlateBrushAsset* SlateBrushAsset = Cast<USlateBrushAsset>(Object);
 	if (SlateBrushAsset)
diff --git a/Engine/Source/Editor/UnrealEd/Private/ThumbnailRendering/SoundWaveThumbnailRenderer.cpp b/Engine/Source/Editor/UnrealEd/Private/ThumbnailRendering/SoundWaveThumbnailRenderer.cpp
--- a/Engine/Source/Editor/UnrealEd/Private/ThumbnailRendering/SoundWaveThumbnailRenderer.cpp
+++ b/Engine/Source/Editor/UnrealEd/Private/ThumbnailRendering/SoundWaveThumbnailRenderer.cpp
@@ -21,7 +21,7 @@ bool USoundWaveThumbnailRenderer::CanVisualizeAsset(UObject* Object)
 	return true;
 }
 
-void USoundWaveThumbnailRenderer::Draw(UObject* Object, int32 X, int32 Y, uint32 Width, uint32 Height, FRenderTarget*, FCanvas* Canvas)
+void USoundWaveThumbnailRenderer::Draw(UObject* Object, int32 X, int32 Y, uint32 Width, uint32 Height, FRenderTarget*, FCanvas* Canvas, bool bAdditionalViewFamily)
 {
 	USoundWave* SoundWave = Cast<USoundWave>(Object);
 	if (SoundWave != nullptr && SoundWave->NumChannels > 0)
diff --git a/Engine/Source/Editor/UnrealEd/Private/ThumbnailRendering/StaticMeshThumbnailRenderer.cpp b/Engine/Source/Editor/UnrealEd/Private/ThumbnailRendering/StaticMeshThumbnailRenderer.cpp
--- a/Engine/Source/Editor/UnrealEd/Private/ThumbnailRendering/StaticMeshThumbnailRenderer.cpp
+++ b/Engine/Source/Editor/UnrealEd/Private/ThumbnailRendering/StaticMeshThumbnailRenderer.cpp
@@ -13,7 +13,7 @@ UStaticMeshThumbnailRenderer::UStaticMeshThumbnailRenderer(const FObjectInitiali
 	ThumbnailScene = nullptr;
 }
 
-void UStaticMeshThumbnailRenderer::Draw(UObject* Object, int32 X, int32 Y, uint32 Width, uint32 Height, FRenderTarget* RenderTarget, FCanvas* Canvas)
+void UStaticMeshThumbnailRenderer::Draw(UObject* Object, int32 X, int32 Y, uint32 Width, uint32 Height, FRenderTarget* RenderTarget, FCanvas* Canvas, bool bAdditionalViewFamily)
 {
 	UStaticMesh* StaticMesh = Cast<UStaticMesh>(Object);
 	if (StaticMesh != nullptr && !StaticMesh->IsPendingKill())
@@ -32,7 +32,8 @@ void UStaticMeshThumbnailRenderer::Draw(UObject* Object, int32 X, int32 Y, uint3
 		ThumbnailScene->GetScene()->UpdateSpeedTreeWind(0.0);
 
 		FSceneViewFamilyContext ViewFamily( FSceneViewFamily::ConstructionValues( RenderTarget, ThumbnailScene->GetScene(), FEngineShowFlags(ESFIM_Game) )
-			.SetWorldTimes(FApp::GetCurrentTime() - GStartTime, FApp::GetDeltaTime(), FApp::GetCurrentTime() - GStartTime));
+			.SetWorldTimes(FApp::GetCurrentTime() - GStartTime, FApp::GetDeltaTime(), FApp::GetCurrentTime() - GStartTime)
+			.SetAdditionalViewFamily(bAdditionalViewFamily));
 
 		ViewFamily.EngineShowFlags.DisableAdvancedFeatures();
 		ViewFamily.EngineShowFlags.MotionBlur = 0;
diff --git a/Engine/Source/Editor/UnrealEd/Private/ThumbnailRendering/SubsurfaceProfileRenderer.cpp b/Engine/Source/Editor/UnrealEd/Private/ThumbnailRendering/SubsurfaceProfileRenderer.cpp
--- a/Engine/Source/Editor/UnrealEd/Private/ThumbnailRendering/SubsurfaceProfileRenderer.cpp
+++ b/Engine/Source/Editor/UnrealEd/Private/ThumbnailRendering/SubsurfaceProfileRenderer.cpp
@@ -18,7 +18,7 @@ void USubsurfaceProfileRenderer::GetThumbnailSize(UObject* Object, float Zoom, u
 	OutHeight = 128;
 }
 
-void USubsurfaceProfileRenderer::Draw(UObject* Object, int32 X, int32 Y, uint32 Width, uint32 Height, FRenderTarget*, FCanvas* Canvas)
+void USubsurfaceProfileRenderer::Draw(UObject* Object, int32 X, int32 Y, uint32 Width, uint32 Height, FRenderTarget*, FCanvas* Canvas, bool bAdditionalViewFamily)
 {
 	USubsurfaceProfile* LocalSubsurfaceProfile = Cast<USubsurfaceProfile>(Object);
 	if (LocalSubsurfaceProfile)
diff --git a/Engine/Source/Editor/UnrealEd/Private/ThumbnailRendering/Texture2DArrayThumbnailRenderer.cpp b/Engine/Source/Editor/UnrealEd/Private/ThumbnailRendering/Texture2DArrayThumbnailRenderer.cpp
--- a/Engine/Source/Editor/UnrealEd/Private/ThumbnailRendering/Texture2DArrayThumbnailRenderer.cpp
+++ b/Engine/Source/Editor/UnrealEd/Private/ThumbnailRendering/Texture2DArrayThumbnailRenderer.cpp
@@ -23,11 +23,11 @@ void UTexture2DArrayThumbnailRenderer::GetThumbnailSize(UObject* Object, float Z
 
 }
 
-void UTexture2DArrayThumbnailRenderer::Draw(UObject* Object, int32 X, int32 Y, uint32 Width, uint32 Height, FRenderTarget*, FCanvas* Canvas)
+void UTexture2DArrayThumbnailRenderer::Draw(UObject* Object, int32 X, int32 Y, uint32 Width, uint32 Height, FRenderTarget*, FCanvas* Canvas, bool bAdditionalViewFamily)
 {
 	UTexture2DArray* TextureArray = Cast<UTexture2DArray>(Object);
 	if (TextureArray != nullptr)
 	{
-		Super::Draw(TextureArray, X, Y, Width, Height, nullptr, Canvas);
+		Super::Draw(TextureArray, X, Y, Width, Height, nullptr, Canvas, bAdditionalViewFamily);
 	}
 }
diff --git a/Engine/Source/Editor/UnrealEd/Private/ThumbnailRendering/TextureCubeThumbnailRenderer.cpp b/Engine/Source/Editor/UnrealEd/Private/ThumbnailRendering/TextureCubeThumbnailRenderer.cpp
--- a/Engine/Source/Editor/UnrealEd/Private/ThumbnailRendering/TextureCubeThumbnailRenderer.cpp
+++ b/Engine/Source/Editor/UnrealEd/Private/ThumbnailRendering/TextureCubeThumbnailRenderer.cpp
@@ -21,11 +21,11 @@ void UTextureCubeThumbnailRenderer::GetThumbnailSize(UObject* Object, float Zoom
 	}
 }
 
-void UTextureCubeThumbnailRenderer::Draw(UObject* Object, int32 X, int32 Y, uint32 Width, uint32 Height, FRenderTarget*, FCanvas* Canvas)
+void UTextureCubeThumbnailRenderer::Draw(UObject* Object, int32 X, int32 Y, uint32 Width, uint32 Height, FRenderTarget*, FCanvas* Canvas, bool bAdditionalViewFamily)
 {
 	UTextureCube* CubeMap = Cast<UTextureCube>(Object);
 	if (CubeMap != nullptr)
 	{
-		Super::Draw(CubeMap, X, Y, Width, Height, nullptr,	Canvas);
+		Super::Draw(CubeMap, X, Y, Width, Height, nullptr, Canvas, bAdditionalViewFamily);
 	}
 }
diff --git a/Engine/Source/Editor/UnrealEd/Private/ThumbnailRendering/TextureThumbnailRenderer.cpp b/Engine/Source/Editor/UnrealEd/Private/ThumbnailRendering/TextureThumbnailRenderer.cpp
--- a/Engine/Source/Editor/UnrealEd/Private/ThumbnailRendering/TextureThumbnailRenderer.cpp
+++ b/Engine/Source/Editor/UnrealEd/Private/ThumbnailRendering/TextureThumbnailRenderer.cpp
@@ -46,7 +46,7 @@ void UTextureThumbnailRenderer::GetThumbnailSize(UObject* Object, float Zoom, ui
 	}
 }
 
-void UTextureThumbnailRenderer::Draw(UObject* Object, int32 X, int32 Y, uint32 Width, uint32 Height, FRenderTarget*, FCanvas* Canvas)
+void UTextureThumbnailRenderer::Draw(UObject* Object, int32 X, int32 Y, uint32 Width, uint32 Height, FRenderTarget*, FCanvas* Canvas, bool bAdditionalViewFamily)
 {
 	UTexture* Texture = Cast<UTexture>(Object);
 	if (Texture != nullptr && Texture->Resource != nullptr) 
diff --git a/Engine/Source/Editor/UnrealEd/Private/ThumbnailRendering/VolumeTextureThumbnailRenderer.cpp b/Engine/Source/Editor/UnrealEd/Private/ThumbnailRendering/VolumeTextureThumbnailRenderer.cpp
--- a/Engine/Source/Editor/UnrealEd/Private/ThumbnailRendering/VolumeTextureThumbnailRenderer.cpp
+++ b/Engine/Source/Editor/UnrealEd/Private/ThumbnailRendering/VolumeTextureThumbnailRenderer.cpp
@@ -134,7 +134,7 @@ UVolumeTextureThumbnailRenderer::UVolumeTextureThumbnailRenderer(const FObjectIn
 {
 }
 
-void UVolumeTextureThumbnailRenderer::Draw(UObject* Object, int32 X, int32 Y, uint32 Width, uint32 Height, FRenderTarget* RenderTarget, FCanvas* Canvas)
+void UVolumeTextureThumbnailRenderer::Draw(UObject* Object, int32 X, int32 Y, uint32 Width, uint32 Height, FRenderTarget* RenderTarget, FCanvas* Canvas, bool bAdditionalViewFamily)
 {
 	UVolumeTexture* VolumeTexture = Cast<UVolumeTexture>(Object);
 	if (VolumeTexture != nullptr)
@@ -161,7 +161,8 @@ void UVolumeTextureThumbnailRenderer::Draw(UObject* Object, int32 X, int32 Y, ui
 
 			ThumbnailScene->SetMaterialInterface(MaterialInstance);
 			FSceneViewFamilyContext ViewFamily( FSceneViewFamily::ConstructionValues( RenderTarget, ThumbnailScene->GetScene(), FEngineShowFlags(ESFIM_Game) )
-				.SetWorldTimes(FApp::GetCurrentTime() - GStartTime, FApp::GetDeltaTime(), FApp::GetCurrentTime() - GStartTime));
+				.SetWorldTimes(FApp::GetCurrentTime() - GStartTime, FApp::GetDeltaTime(), FApp::GetCurrentTime() - GStartTime)
+				.SetAdditionalViewFamily(bAdditionalViewFamily));
 
 			ViewFamily.EngineShowFlags.DisableAdvancedFeatures();
 			ViewFamily.EngineShowFlags.SetSeparateTranslucency(true);
diff --git a/Engine/Source/Editor/UnrealEd/Private/ThumbnailRendering/WorldThumbnailRenderer.cpp b/Engine/Source/Editor/UnrealEd/Private/ThumbnailRendering/WorldThumbnailRenderer.cpp
--- a/Engine/Source/Editor/UnrealEd/Private/ThumbnailRendering/WorldThumbnailRenderer.cpp
+++ b/Engine/Source/Editor/UnrealEd/Private/ThumbnailRendering/WorldThumbnailRenderer.cpp
@@ -34,13 +34,14 @@ bool UWorldThumbnailRenderer::CanVisualizeAsset(UObject* Object)
 	return false;
 }
 
-void UWorldThumbnailRenderer::Draw(UObject* Object, int32 X, int32 Y, uint32 Width, uint32 Height, FRenderTarget* RenderTarget, FCanvas* Canvas)
+void UWorldThumbnailRenderer::Draw(UObject* Object, int32 X, int32 Y, uint32 Width, uint32 Height, FRenderTarget* RenderTarget, FCanvas* Canvas, bool bAdditionalViewFamily)
 {
 	UWorld* World = Cast<UWorld>(Object);
 	if (World != nullptr && World->Scene)
 	{
 		FSceneViewFamilyContext ViewFamily( FSceneViewFamily::ConstructionValues( RenderTarget, World->Scene, FEngineShowFlags(ESFIM_All0) )
-			.SetWorldTimes(FApp::GetCurrentTime() - GStartTime, FApp::GetDeltaTime(), FApp::GetCurrentTime() - GStartTime));
+			.SetWorldTimes(FApp::GetCurrentTime() - GStartTime, FApp::GetDeltaTime(), FApp::GetCurrentTime() - GStartTime)
+			.SetAdditionalViewFamily(bAdditionalViewFamily));
 
 		ViewFamily.EngineShowFlags.SetDiffuse(true);
 		ViewFamily.EngineShowFlags.SetSkeletalMeshes(true);
diff --git a/Engine/Source/Editor/VirtualTexturingEditor/Private/RuntimeVirtualTextureThumbnailRenderer.cpp b/Engine/Source/Editor/VirtualTexturingEditor/Private/RuntimeVirtualTextureThumbnailRenderer.cpp
--- a/Engine/Source/Editor/VirtualTexturingEditor/Private/RuntimeVirtualTextureThumbnailRenderer.cpp
+++ b/Engine/Source/Editor/VirtualTexturingEditor/Private/RuntimeVirtualTextureThumbnailRenderer.cpp
@@ -54,7 +54,7 @@ bool URuntimeVirtualTextureThumbnailRenderer::CanVisualizeAsset(UObject* Object)
 	return false;
 }
 
-void URuntimeVirtualTextureThumbnailRenderer::Draw(UObject* Object, int32 X, int32 Y, uint32 Width, uint32 Height, FRenderTarget* RenderTarget, FCanvas* Canvas)
+void URuntimeVirtualTextureThumbnailRenderer::Draw(UObject* Object, int32 X, int32 Y, uint32 Width, uint32 Height, FRenderTarget* RenderTarget, FCanvas* Canvas, bool bAdditionalViewFamily)
 {
 	URuntimeVirtualTexture* RuntimeVirtualTexture = Cast<URuntimeVirtualTexture>(Object);
 	URuntimeVirtualTextureComponent* RuntimeVirtualTextureComponent = FindComponent(RuntimeVirtualTexture);
diff --git a/Engine/Source/Editor/VirtualTexturingEditor/Private/RuntimeVirtualTextureThumbnailRenderer.h b/Engine/Source/Editor/VirtualTexturingEditor/Private/RuntimeVirtualTextureThumbnailRenderer.h
--- a/Engine/Source/Editor/VirtualTexturingEditor/Private/RuntimeVirtualTextureThumbnailRenderer.h
+++ b/Engine/Source/Editor/VirtualTexturingEditor/Private/RuntimeVirtualTextureThumbnailRenderer.h
@@ -13,6 +13,6 @@ class URuntimeVirtualTextureThumbnailRenderer : public UThumbnailRenderer
 
 	//~ Begin UThumbnailRenderer Interface.
 	virtual bool CanVisualizeAsset(UObject* Object);
-	virtual void Draw(UObject* Object, int32 X, int32 Y, uint32 Width, uint32 Height, FRenderTarget*, FCanvas* Canvas) override;
+	virtual void Draw(UObject* Object, int32 X, int32 Y, uint32 Width, uint32 Height, FRenderTarget*, FCanvas* Canvas, bool bAdditionalViewFamily) override;
 	//~ EndUThumbnailRenderer Interface.
 };
