GeometryCollectionTest :

Added back a few tests.

diff --git a/Engine/Plugins/Experimental/GeometryCollectionPlugin/Source/GeometryCollectionExampleCore/Private/GeometryCollection/GeometryCollectionExampleSimulationSolver.cpp b/Engine/Plugins/Experimental/GeometryCollectionPlugin/Source/GeometryCollectionExampleCore/Private/GeometryCollection/GeometryCollectionExampleSimulationSolver.cpp
--- a/Engine/Plugins/Experimental/GeometryCollectionPlugin/Source/GeometryCollectionExampleCore/Private/GeometryCollection/GeometryCollectionExampleSimulationSolver.cpp
+++ b/Engine/Plugins/Experimental/GeometryCollectionPlugin/Source/GeometryCollectionExampleCore/Private/GeometryCollection/GeometryCollectionExampleSimulationSolver.cpp
@@ -27,10 +27,6 @@
 #define SMALL_THRESHOLD 1e-4
 
 
-// #TODO Lots of duplication in here, anyone making solver or object changes
-// has to go and fix up so many callsites here and they're all pretty much
-// Identical. The similar code should be pulled out
-
 namespace GeometryCollectionExample
 {
 
@@ -51,7 +47,11 @@ namespace GeometryCollectionExample
 	template<class T>
 	void Solver_AdvanceDisabledObjects()
 	{
-		TSharedPtr<FGeometryCollection> RestCollection = GeometryCollection::MakeCubeElement(FTransform::Identity, FVector(1.0));
+		FChaosSolversModule* Module = FChaosSolversModule::GetModule();
+		Module->ChangeThreadingMode(EChaosThreadingMode::SingleThread);
+
+		FVector Scale(1);
+		TSharedPtr<FGeometryCollection> RestCollection = GeometryCollection::MakeCubeElement(FTransform::Identity, Scale);
 		TSharedPtr<FGeometryDynamicCollection> DynamicCollection = GeometryCollectionToGeometryDynamicCollection(RestCollection.Get());
 
 		auto InitFunc = [&RestCollection, &DynamicCollection](FSimulationParameters& InParams)
@@ -68,9 +68,8 @@ namespace GeometryCollectionExample
 		PhysObject->Initialize();
 
 		Chaos::FPBDRigidsSolver* Solver = FChaosSolversModule::GetModule()->CreateSolver(true);
-#if CHAOS_PARTICLEHANDLE_TODO
 		Solver->RegisterObject(PhysObject);
-#endif
+
 		Solver->SetHasFloor(false);
 		Solver->SetEnabled(true);
 		PhysObject->ActivateBodies();
diff --git a/Engine/Source/Runtime/Experimental/ChaosSolvers/Private/PhysicsProxy/GeometryCollectionPhysicsProxy.cpp b/Engine/Source/Runtime/Experimental/ChaosSolvers/Private/PhysicsProxy/GeometryCollectionPhysicsProxy.cpp
--- a/Engine/Source/Runtime/Experimental/ChaosSolvers/Private/PhysicsProxy/GeometryCollectionPhysicsProxy.cpp
+++ b/Engine/Source/Runtime/Experimental/ChaosSolvers/Private/PhysicsProxy/GeometryCollectionPhysicsProxy.cpp
@@ -2081,11 +2081,7 @@ void FGeometryCollectionPhysicsProxy::PushToPhysicsState(const Chaos::FParticleD
 	if (!GState)
 	{
 		// This will happen if GameToPhysInterchange hasn't been updated since
-		// the last time GetConsumerBuffer() was called.  At this point, that's 
-		// probably a bug, as the interface code is assuming it needs to lock
-		// to prevent stomping on a double buffer.  We should look into doing 
-		// away with that lock, but for the time being, we'll assert:
-		ensure(false);
+		// the last time GetConsumerBuffer() was called. 
 		return;
 	}
 
