HeadlessChaos :
- Enabled more tests.
- Minor restructuring.

diff --git a/Engine/Plugins/Experimental/GeometryCollectionPlugin/Source/GeometryCollectionExampleCore/Private/GeometryCollection/GeometryCollectionExampleDecimation.cpp b/Engine/Plugins/Experimental/GeometryCollectionPlugin/Source/GeometryCollectionExampleCore/Private/GeometryCollection/GeometryCollectionExampleDecimation.cpp
--- a/Engine/Plugins/Experimental/GeometryCollectionPlugin/Source/GeometryCollectionExampleCore/Private/GeometryCollection/GeometryCollectionExampleDecimation.cpp
+++ b/Engine/Plugins/Experimental/GeometryCollectionPlugin/Source/GeometryCollectionExampleCore/Private/GeometryCollection/GeometryCollectionExampleDecimation.cpp
@@ -168,8 +168,8 @@ namespace GeometryCollectionExample
 				<< " - expected importance ordering hash: " << ExpectedHash
 				<< " got: " << hash << ".  Failing." << std::endl;
 		}
-		//return hash == ExpectedHash;
-		return true;	//todo: update hash
+		return hash == ExpectedHash;
+		//return true;	//todo: update hash
 	}
 
 	template <class T, class TGeom>
