Chaos : Drop Collection
- Always add proxies to the handles map.
- Move FilpBuffers out of BufferPhysicsResults
- Update Tests to use the Solver for Finalization of the simulation results.
- Enabled rigid bodies falling under gravity.

diff --git a/Engine/Plugins/Experimental/GeometryCollectionPlugin/Source/GeometryCollectionExampleCore/Private/GeometryCollection/GeometryCollectionExampleSimulation.cpp b/Engine/Plugins/Experimental/GeometryCollectionPlugin/Source/GeometryCollectionExampleCore/Private/GeometryCollection/GeometryCollectionExampleSimulation.cpp
--- a/Engine/Plugins/Experimental/GeometryCollectionPlugin/Source/GeometryCollectionExampleCore/Private/GeometryCollection/GeometryCollectionExampleSimulation.cpp
+++ b/Engine/Plugins/Experimental/GeometryCollectionPlugin/Source/GeometryCollectionExampleCore/Private/GeometryCollection/GeometryCollectionExampleSimulation.cpp
@@ -29,6 +29,7 @@ namespace GeometryCollectionExample
 	template<class T>
 	void RigidBodiesFallingUnderGravity()
 	{
+		float Dt = 1 / 24.;
 		FChaosSolversModule* Module = FChaosSolversModule::GetModule();
 		Module->ChangeThreadingMode(EChaosThreadingMode::SingleThread);
 
@@ -41,7 +42,7 @@ namespace GeometryCollectionExample
 			FTransform::Identity,	// RestCenter
 			FVector(1.0),			// RestScale
 			nullptr,				// RestInitFunc
-			(int32)EObjectStateTypeEnum::Chaos_Object_Kinematic // DynamicStateDefault
+			(int32)EObjectStateTypeEnum::Chaos_Object_Dynamic // DynamicStateDefault
 		};
 
 		TUniquePtr<Chaos::FChaosPhysicsMaterial> PhysicalMaterial = nullptr; // Allocated and zero'ed
@@ -64,27 +65,27 @@ namespace GeometryCollectionExample
 		Solver->SetEnabled(true);
 		PhysObject->ActivateBodies();
 
-		Solver->AddDirtyProxy(PhysObject); // why?
+		Solver->AddDirtyProxy(PhysObject);
 		Solver->PushPhysicsState(Module->GetDispatcher());
 
-		Solver->AdvanceSolverBy(1 / 24.);
+		Solver->AdvanceSolverBy(Dt);
 
 		// Calls BufferPhysicsResults(), FlipBuffer(), and PullFromPhysicsState() on each proxy.
-		//FinalizeSolver(*Solver);
-
-		Solver->BufferPhysicsResults();
-		Solver->FlipBuffers();
-		Solver->UpdateGameThreadStructures();
+		FinalizeSolver(*Solver);
 
 		// never touched
 		TManagedArray<FTransform>& RestTransform = RestCollection->Transform;
 		EXPECT_LT(FMath::Abs(RestTransform[0].GetTranslation().Z), SMALL_THRESHOLD);
 
+
 		// simulated
 		TManagedArray<FTransform>& Transform = GTDynamicCollection->Transform;
 		EXPECT_EQ(Transform.Num(), 1);
-		EXPECT_LT(Transform[0].GetTranslation().Z, 0);
-		
+		EXPECT_LT(Transform[0].GetTranslation().Z, 0.f);
+
+		float TwiceGravityDisplacement = /* 0.5* */-980.f * Dt * Dt; // we seem to be twice gravity
+		EXPECT_NEAR(Transform[0].GetTranslation().Z, TwiceGravityDisplacement, 1e-2);
+
 		FChaosSolversModule::GetModule()->DestroySolver(Solver);
 
 		delete PhysObject;
diff --git a/Engine/Plugins/Experimental/GeometryCollectionPlugin/Source/GeometryCollectionExampleCore/Private/GeometryCollection/GeometryCollectionExampleUtility.cpp b/Engine/Plugins/Experimental/GeometryCollectionPlugin/Source/GeometryCollectionExampleCore/Private/GeometryCollection/GeometryCollectionExampleUtility.cpp
--- a/Engine/Plugins/Experimental/GeometryCollectionPlugin/Source/GeometryCollectionExampleCore/Private/GeometryCollection/GeometryCollectionExampleUtility.cpp
+++ b/Engine/Plugins/Experimental/GeometryCollectionPlugin/Source/GeometryCollectionExampleCore/Private/GeometryCollection/GeometryCollectionExampleUtility.cpp
@@ -40,12 +40,9 @@ namespace GeometryCollectionExample {
 
 	void FinalizeSolver(Chaos::FPBDRigidsSolver& InSolver)
 	{
-		InSolver.ForEachPhysicsProxy([](auto* Object)
-		{
-			Object->BufferPhysicsResults();
-			Object->FlipBuffer();
-			Object->PullFromPhysicsState();
-		});
+		InSolver.BufferPhysicsResults();
+		InSolver.FlipBuffers();
+		InSolver.UpdateGameThreadStructures();
 	}
 
 
diff --git a/Engine/Source/Runtime/Experimental/ChaosSolvers/Private/PhysicsProxy/GeometryCollectionPhysicsProxy.cpp b/Engine/Source/Runtime/Experimental/ChaosSolvers/Private/PhysicsProxy/GeometryCollectionPhysicsProxy.cpp
--- a/Engine/Source/Runtime/Experimental/ChaosSolvers/Private/PhysicsProxy/GeometryCollectionPhysicsProxy.cpp
+++ b/Engine/Source/Runtime/Experimental/ChaosSolvers/Private/PhysicsProxy/GeometryCollectionPhysicsProxy.cpp
@@ -611,11 +611,8 @@ void FGeometryCollectionPhysicsProxy::InitializeBodiesPT(
 
 				Chaos::TPBDGeometryCollectionParticleHandle<float, 3>* Handle = Handles[NextIdx++];
 
-				if (DynamicState[ Idx ] == (int32)EObjectStateTypeEnum::Chaos_Object_Dynamic)
-				{
-					Handle->GTGeometryParticle() = nullptr; // GeometryCollections do not support Game Thread TGeometryParticle
-					RigidsSolver->AddParticleToProxy(Handle, this);
-				}
+				Handle->GTGeometryParticle() = nullptr; // GeometryCollections do not support Game Thread TGeometryParticle
+				RigidsSolver->AddParticleToProxy(Handle, this);
 
 				SolverParticleHandles[Idx] = Handle;
 				HandleToTransformGroupIndex.Add(Handle, Idx);
@@ -2312,9 +2309,6 @@ void FGeometryCollectionPhysicsProxy::BufferPhysicsResults()
 	// so that it has an updated view.
 	UpdateGeometryCollection(TargetResults);
 
-	// We've updated a buffer.  There's no reason to wait for FGeometryCollectionPhysicsProxy::FlipBuffer()
-	// to be called.  Just flip.
-	PhysToGameInterchange->FlipProducer();
 //#endif // TODO_REIMPLEMENT_GET_RIGID_PARTICLES
 }
 
@@ -2326,6 +2320,8 @@ void FGeometryCollectionPhysicsProxy::FlipBuffer()
 	 * The physics thread has pre-locked an RW lock for this operation so the game thread won't be reading
 	 * the data
 	 */
+
+	PhysToGameInterchange->FlipProducer();
 }
 
 // Called from FPhysScene_ChaosInterface::SyncBodies(), NOT the solver.
