Chaos : Floor

Removed the defunct floor api off rigid solver.

diff --git a/Engine/Plugins/Experimental/GeometryCollectionPlugin/Source/GeometryCollectionExampleCore/Private/GeometryCollection/GeometryCollectionExampleClustering.cpp b/Engine/Plugins/Experimental/GeometryCollectionPlugin/Source/GeometryCollectionExampleCore/Private/GeometryCollection/GeometryCollectionExampleClustering.cpp
--- a/Engine/Plugins/Experimental/GeometryCollectionPlugin/Source/GeometryCollectionExampleCore/Private/GeometryCollection/GeometryCollectionExampleClustering.cpp
+++ b/Engine/Plugins/Experimental/GeometryCollectionPlugin/Source/GeometryCollectionExampleCore/Private/GeometryCollection/GeometryCollectionExampleClustering.cpp
@@ -87,7 +87,7 @@ namespace GeometryCollectionExample
 #if CHAOS_PARTICLEHANDLE_TODO
 		Solver->RegisterObject(PhysObject);
 #endif
-		Solver->SetHasFloor(true);
+		////Solver->SetHasFloor(true);
 		Solver->SetEnabled(true);
 		PhysObject->ActivateBodies();
 
@@ -179,7 +179,7 @@ namespace GeometryCollectionExample
 #if CHAOS_PARTICLEHANDLE_TODO
 		Solver->RegisterObject(PhysObject);
 #endif
-		Solver->SetHasFloor(true);
+		////Solver->SetHasFloor(true);
 		Solver->SetEnabled(true);
 		PhysObject->ActivateBodies();
 
@@ -324,7 +324,7 @@ namespace GeometryCollectionExample
 #if CHAOS_PARTICLEHANDLE_TODO
 		Solver->RegisterObject(PhysObject);
 #endif
-		Solver->SetHasFloor(true);
+		//Solver->SetHasFloor(true);
 		Solver->SetEnabled(true);
 		PhysObject->ActivateBodies();
 
@@ -462,7 +462,7 @@ namespace GeometryCollectionExample
 #if CHAOS_PARTICLEHANDLE_TODO
 		Solver->RegisterObject(PhysObject);
 #endif
-		Solver->SetHasFloor(true);
+		//Solver->SetHasFloor(true);
 		Solver->SetEnabled(true);
 		PhysObject->ActivateBodies();
 
@@ -592,7 +592,7 @@ namespace GeometryCollectionExample
 #if CHAOS_PARTICLEHANDLE_TODO
 		Solver->RegisterObject(PhysObject);
 #endif
-		Solver->SetHasFloor(true);
+		//Solver->SetHasFloor(true);
 		Solver->SetEnabled(true);
 		PhysObject->ActivateBodies();
 
@@ -747,7 +747,7 @@ namespace GeometryCollectionExample
 #if CHAOS_PARTICLEHANDLE_TODO
 		Solver->RegisterObject(PhysObject);
 #endif
-		Solver->SetHasFloor(true);
+		//Solver->SetHasFloor(true);
 		Solver->SetEnabled(true);
 		PhysObject->ActivateBodies();
 
@@ -882,7 +882,7 @@ namespace GeometryCollectionExample
 #if CHAOS_PARTICLEHANDLE_TODO
 		Solver->RegisterObject(PhysObject);
 #endif
-		Solver->SetHasFloor(true);
+		//Solver->SetHasFloor(true);
 		Solver->SetEnabled(true);
 		PhysObject->ActivateBodies();
 
@@ -1119,7 +1119,7 @@ namespace GeometryCollectionExample
 #if CHAOS_PARTICLEHANDLE_TODO
 		Solver->RegisterObject(PhysObject);
 #endif
-		Solver->SetHasFloor(true);
+		//Solver->SetHasFloor(true);
 		Solver->SetEnabled(true);
 		PhysObject->ActivateBodies();
 
@@ -1336,7 +1336,7 @@ namespace GeometryCollectionExample
 		Solver->RegisterObject(PhysObject);
 		Solver->RegisterObject(PhysObject2);
 #endif
-		Solver->SetHasFloor(true);
+		//Solver->SetHasFloor(true);
 		Solver->SetEnabled(true);
 		PhysObject->ActivateBodies();
 		PhysObject2->ActivateBodies();
@@ -1472,7 +1472,7 @@ namespace GeometryCollectionExample
 		Solver->RegisterObject(PhysObject);
 		Solver->RegisterObject(PhysObject2);
 #endif
-		Solver->SetHasFloor(true);
+		//Solver->SetHasFloor(true);
 		Solver->SetEnabled(true);
 		PhysObject->ActivateBodies();
 		PhysObject2->ActivateBodies();
@@ -1609,7 +1609,7 @@ namespace GeometryCollectionExample
 		Solver->RegisterObject(PhysObject);
 		Solver->RegisterObject(PhysObject2);
 #endif
-		Solver->SetHasFloor(true);
+		//Solver->SetHasFloor(true);
 		Solver->SetEnabled(true);
 		PhysObject->ActivateBodies();
 		PhysObject2->ActivateBodies();
@@ -1714,7 +1714,7 @@ namespace GeometryCollectionExample
 #if CHAOS_PARTICLEHANDLE_TODO
 		Solver->RegisterObject(PhysObject);
 #endif
-		Solver->SetHasFloor(false);
+
 		Solver->SetEnabled(true);
 		PhysObject->ActivateBodies();
 
@@ -1808,7 +1808,6 @@ namespace GeometryCollectionExample
 #if CHAOS_PARTICLEHANDLE_TODO
 		Solver->RegisterObject(PhysObject);
 #endif
-		Solver->SetHasFloor(false);
 		Solver->SetEnabled(true);
 		PhysObject->ActivateBodies();
 
@@ -1918,7 +1917,6 @@ namespace GeometryCollectionExample
 		Solver->RegisterObject(PhysObject);
 		Solver->RegisterObject(FieldObject);
 #endif
-		Solver->SetHasFloor(false);
 		Solver->SetEnabled(true);
 		PhysObject->ActivateBodies();
 
@@ -1986,7 +1984,6 @@ namespace GeometryCollectionExample
 
 
 		Chaos::FPBDRigidsSolver* Solver = FChaosSolversModule::GetModule()->CreateSolver(nullptr, ESolverFlags::Standalone);
-		Solver->SetHasFloor(false);
 		Solver->SetEnabled(true);
 #if CHAOS_PARTICLEHANDLE_TODO
 		Solver->RegisterObject(Object->PhysicsProxy.Get());
diff --git a/Engine/Plugins/Experimental/GeometryCollectionPlugin/Source/GeometryCollectionExampleCore/Private/GeometryCollection/GeometryCollectionExampleSimulation.cpp b/Engine/Plugins/Experimental/GeometryCollectionPlugin/Source/GeometryCollectionExampleCore/Private/GeometryCollection/GeometryCollectionExampleSimulation.cpp
--- a/Engine/Plugins/Experimental/GeometryCollectionPlugin/Source/GeometryCollectionExampleCore/Private/GeometryCollection/GeometryCollectionExampleSimulation.cpp
+++ b/Engine/Plugins/Experimental/GeometryCollectionPlugin/Source/GeometryCollectionExampleCore/Private/GeometryCollection/GeometryCollectionExampleSimulation.cpp
@@ -61,7 +61,7 @@ namespace GeometryCollectionExample
 
 		Chaos::FPBDRigidsSolver* Solver = FChaosSolversModule::GetModule()->CreateSolver(nullptr, ESolverFlags::Standalone);
 		Solver->RegisterObject(PhysObject);
-		Solver->SetHasFloor(false);
+		// no floor
 		Solver->SetEnabled(true);
 		PhysObject->ActivateBodies();
 
@@ -87,7 +87,7 @@ namespace GeometryCollectionExample
 		EXPECT_NEAR(Transform[0].GetTranslation().Z, TwiceGravityDisplacement, 1e-2);
 
 		FChaosSolversModule::GetModule()->DestroySolver(Solver);
-
+		
 		delete PhysObject;
 	}
 	template void RigidBodiesFallingUnderGravity<float>();
@@ -118,8 +118,8 @@ namespace GeometryCollectionExample
 #if CHAOS_PARTICLEHANDLE_TODO
 		Solver->RegisterObject(PhysObject);
 #endif
-		Solver->SetHasFloor(true);
-		Solver->SetIsFloorAnalytic(true);
+		////Solver->SetHasFloor(true);
+		////Solver->SetIsFloorAnalytic(true);
 		Solver->SetEnabled(true);
 		PhysObject->ActivateBodies();
 
@@ -173,8 +173,8 @@ namespace GeometryCollectionExample
 #if CHAOS_PARTICLEHANDLE_TODO
 		Solver->RegisterObject(PhysObject);
 #endif
-		Solver->SetHasFloor(true);
-		Solver->SetIsFloorAnalytic(true);
+		//Solver->SetHasFloor(true);
+		//Solver->SetIsFloorAnalytic(true);
 		Solver->SetEnabled(true);
 		PhysObject->ActivateBodies();
 
@@ -228,8 +228,8 @@ namespace GeometryCollectionExample
 #if CHAOS_PARTICLEHANDLE_TODO
 		Solver->RegisterObject(PhysObject);
 #endif
-		Solver->SetHasFloor(true);
-		Solver->SetIsFloorAnalytic(true);
+		//Solver->SetHasFloor(true);
+		//Solver->SetIsFloorAnalytic(true);
 		Solver->SetEnabled(true);
 		PhysObject->ActivateBodies();
 
@@ -276,8 +276,8 @@ namespace GeometryCollectionExample
 #if CHAOS_PARTICLEHANDLE_TODO
 		Solver->RegisterObject(PhysObject);
 #endif
-		Solver->SetHasFloor(false);
-		Solver->SetIsFloorAnalytic(true);
+		// no floor
+		//Solver->SetIsFloorAnalytic(true);
 		Solver->SetEnabled(true);
 		PhysObject->ActivateBodies();
 
@@ -337,8 +337,8 @@ namespace GeometryCollectionExample
 #if CHAOS_PARTICLEHANDLE_TODO
 		Solver->RegisterObject(PhysObject);
 #endif
-		Solver->SetHasFloor(false);
-		Solver->SetIsFloorAnalytic(true);
+		// no floor
+		//Solver->SetIsFloorAnalytic(true);
 		Solver->SetEnabled(true);
 		PhysObject->ActivateBodies();
 
@@ -395,7 +395,7 @@ namespace GeometryCollectionExample
 #if CHAOS_PARTICLEHANDLE_TODO
 		Solver->RegisterObject(PhysObject);
 #endif
-		Solver->SetHasFloor(true);
+		//Solver->SetHasFloor(true);
 		Solver->SetEnabled(true);
 		PhysObject->ActivateBodies();
 
diff --git a/Engine/Plugins/Experimental/GeometryCollectionPlugin/Source/GeometryCollectionExampleCore/Private/GeometryCollection/GeometryCollectionExampleSimulationField.cpp b/Engine/Plugins/Experimental/GeometryCollectionPlugin/Source/GeometryCollectionExampleCore/Private/GeometryCollection/GeometryCollectionExampleSimulationField.cpp
--- a/Engine/Plugins/Experimental/GeometryCollectionPlugin/Source/GeometryCollectionExampleCore/Private/GeometryCollection/GeometryCollectionExampleSimulationField.cpp
+++ b/Engine/Plugins/Experimental/GeometryCollectionPlugin/Source/GeometryCollectionExampleCore/Private/GeometryCollection/GeometryCollectionExampleSimulationField.cpp
@@ -54,8 +54,8 @@ namespace GeometryCollectionExample
 
 		Chaos::FPBDRigidsSolver* Solver = Module->CreateSolver(nullptr, ESolverFlags::Standalone);
 		Solver->RegisterObject(PhysObject);
-		Solver->SetHasFloor(false);
-		Solver->SetIsFloorAnalytic(true);
+		// no floor
+		//Solver->SetIsFloorAnalytic(true);
 		Solver->SetEnabled(true);
 		PhysObject->ActivateBodies();
 
@@ -143,8 +143,8 @@ namespace GeometryCollectionExample
 		Chaos::FPBDRigidsSolver* Solver = FChaosSolversModule::GetModule()->CreateSolver(nullptr, ESolverFlags::Standalone);
 		Solver->RegisterObject(PhysObject);
 		Solver->RegisterObject(FieldObject);
-		Solver->SetHasFloor(false);
-		Solver->SetIsFloorAnalytic(true);
+		// no floor
+		//Solver->SetIsFloorAnalytic(true);
 		Solver->SetEnabled(true);
 		PhysObject->ActivateBodies();
 
@@ -229,7 +229,7 @@ namespace GeometryCollectionExample
 		Chaos::FPBDRigidsSolver* Solver = FChaosSolversModule::GetModule()->CreateSolver(nullptr, ESolverFlags::Standalone);
 		Solver->RegisterObject(PhysObject);
 		Solver->RegisterObject(FieldObject);
-		Solver->SetHasFloor(false);
+		// no floor
 		Solver->SetEnabled(true);
 		PhysObject->ActivateBodies();
 
@@ -316,7 +316,7 @@ namespace GeometryCollectionExample
 		Chaos::FPBDRigidsSolver* Solver = FChaosSolversModule::GetModule()->CreateSolver(nullptr, ESolverFlags::Standalone);
 		Solver->RegisterObject(PhysObject);
 		Solver->RegisterObject(FieldObject);
-		Solver->SetHasFloor(false);
+		// no floor
 		Solver->SetEnabled(true);
 		PhysObject->ActivateBodies();
 
@@ -397,7 +397,7 @@ namespace GeometryCollectionExample
 		Chaos::FPBDRigidsSolver* Solver = FChaosSolversModule::GetModule()->CreateSolver(nullptr, ESolverFlags::Standalone);
 		Solver->RegisterObject(PhysObject);
 		Solver->RegisterObject(FieldObject);
-		Solver->SetHasFloor(false);
+		// no floor
 		Solver->SetEnabled(true);
 		PhysObject->ActivateBodies();
 
@@ -488,7 +488,7 @@ namespace GeometryCollectionExample
 		Chaos::FPBDRigidsSolver* Solver = FChaosSolversModule::GetModule()->CreateSolver(nullptr, ESolverFlags::Standalone);
 		Solver->RegisterObject(PhysObject);
 		Solver->RegisterObject(FieldObject);
-		Solver->SetHasFloor(false);
+		// no floor
 		Solver->SetEnabled(true);
 		PhysObject->ActivateBodies();
 
@@ -558,7 +558,7 @@ namespace GeometryCollectionExample
 		Chaos::FPBDRigidsSolver* Solver = FChaosSolversModule::GetModule()->CreateSolver(nullptr, ESolverFlags::Standalone);
 		Solver->RegisterObject(PhysObject);
 		Solver->RegisterObject(FieldObject);
-		Solver->SetHasFloor(false);
+		// no floor
 		Solver->SetEnabled(true);
 		PhysObject->ActivateBodies();
 
@@ -639,7 +639,7 @@ namespace GeometryCollectionExample
 
 		Chaos::FPBDRigidsSolver* Solver = FChaosSolversModule::GetModule()->CreateSolver(nullptr, ESolverFlags::Standalone);
 		Solver->RegisterObject(PhysObject);
-		Solver->SetHasFloor(true);
+		//Solver->SetHasFloor(true);
 		Solver->SetEnabled(true);
 		PhysObject->ActivateBodies();
 
@@ -727,7 +727,7 @@ namespace GeometryCollectionExample
 
 		Solver->RegisterObject(PhysObject);
 		Solver->RegisterObject(FieldObject);
-		Solver->SetHasFloor(false);
+		// no floor
 		Solver->SetEnabled(true);
 		PhysObject->ActivateBodies();
 
@@ -824,7 +824,7 @@ namespace GeometryCollectionExample
 
 		Solver->RegisterObject(PhysObject);
 		Solver->RegisterObject(FieldObject);
-		Solver->SetHasFloor(false);
+		// no floor
 		Solver->SetEnabled(true);
 		PhysObject->ActivateBodies();
 
@@ -930,7 +930,7 @@ namespace GeometryCollectionExample
 #endif
 		Solver->RegisterObject(PhysObject);
 		Solver->RegisterObject(FieldObject);
-		Solver->SetHasFloor(false);
+		// no floor
 		Solver->SetEnabled(true);
 		PhysObject->ActivateBodies();
 
@@ -1039,7 +1039,7 @@ namespace GeometryCollectionExample
 #endif
 		Solver->RegisterObject(PhysObject);
 		Solver->RegisterObject(FieldObject);
-		Solver->SetHasFloor(false);
+		// no floor
 		Solver->SetEnabled(true);
 		PhysObject->ActivateBodies();
 
diff --git a/Engine/Plugins/Experimental/GeometryCollectionPlugin/Source/GeometryCollectionExampleCore/Private/GeometryCollection/GeometryCollectionExampleSimulationSolver.cpp b/Engine/Plugins/Experimental/GeometryCollectionPlugin/Source/GeometryCollectionExampleCore/Private/GeometryCollection/GeometryCollectionExampleSimulationSolver.cpp
--- a/Engine/Plugins/Experimental/GeometryCollectionPlugin/Source/GeometryCollectionExampleCore/Private/GeometryCollection/GeometryCollectionExampleSimulationSolver.cpp
+++ b/Engine/Plugins/Experimental/GeometryCollectionPlugin/Source/GeometryCollectionExampleCore/Private/GeometryCollection/GeometryCollectionExampleSimulationSolver.cpp
@@ -34,7 +34,7 @@ namespace GeometryCollectionExample
 	void Solver_AdvanceNoObjects()
 	{
 		Chaos::FPBDRigidsSolver* Solver = FChaosSolversModule::GetModule()->CreateSolver(nullptr, ESolverFlags::Standalone);
-		Solver->SetHasFloor(false);
+		// no floor
 		Solver->SetEnabled(true);
 		Solver->AdvanceSolverBy(1 / 24.);
 		FinalizeSolver(*Solver);
@@ -71,7 +71,7 @@ namespace GeometryCollectionExample
 
 		Solver->RegisterObject(PhysObject);
 
-		Solver->SetHasFloor(false);
+		// no floor
 		Solver->SetEnabled(true);
 		PhysObject->ActivateBodies();
 
@@ -129,7 +129,7 @@ namespace GeometryCollectionExample
 #if CHAOS_PARTICLEHANDLE_TODO
 		Solver->RegisterObject(PhysObject);
 #endif
-		Solver->SetHasFloor(true);
+		//Solver->SetHasFloor(true);
 		Solver->SetEnabled(true);
 		PhysObject->ActivateBodies();
 
@@ -157,7 +157,7 @@ namespace GeometryCollectionExample
 		InitMaterialToZero(PhysicalMaterial);
 
 		Chaos::FPBDRigidsSolver* Solver = FChaosSolversModule::GetModule()->CreateSolver(nullptr, ESolverFlags::Standalone);
-		Solver->SetHasFloor(true);
+		//Solver->SetHasFloor(true);
 		Solver->SetEnabled(true);
 		TArray<TSharedPtr<FGeometryCollection> > RestArray;
 		TArray<TSharedPtr<FGeometryDynamicCollection> > DynamicArray;
@@ -219,7 +219,7 @@ namespace GeometryCollectionExample
 
 		*SolverInOut = FChaosSolversModule::GetModule()->CreateSolver(nullptr, ESolverFlags::Standalone);
 		Chaos::FPBDRigidsSolver* Solver = *SolverInOut;
-		Solver->SetHasFloor(true);
+		//Solver->SetHasFloor(true);
 		Solver->SetEnabled(true);
 
 		for (int32 i = 0; i < NumObjects; i++)
@@ -254,7 +254,7 @@ namespace GeometryCollectionExample
 #if CHAOS_PARTICLEHANDLE_TODO
 			Solver->RegisterObject(PhysObject);
 #endif
-			Solver->SetHasFloor(true);
+			//Solver->SetHasFloor(true);
 			Solver->SetEnabled(true);
 			PhysObject->ActivateBodies();
 
diff --git a/Engine/Plugins/Experimental/GeometryCollectionPlugin/Source/GeometryCollectionExampleCore/Private/GeometryCollection/GeometryCollectionExampleSimulationStreaming.cpp b/Engine/Plugins/Experimental/GeometryCollectionPlugin/Source/GeometryCollectionExampleCore/Private/GeometryCollection/GeometryCollectionExampleSimulationStreaming.cpp
--- a/Engine/Plugins/Experimental/GeometryCollectionPlugin/Source/GeometryCollectionExampleCore/Private/GeometryCollection/GeometryCollectionExampleSimulationStreaming.cpp
+++ b/Engine/Plugins/Experimental/GeometryCollectionPlugin/Source/GeometryCollectionExampleCore/Private/GeometryCollection/GeometryCollectionExampleSimulationStreaming.cpp
@@ -27,7 +27,7 @@ namespace GeometryCollectionExample
 	void RigidBodies_Streaming_StartSolverEmpty()
 	{
 		Chaos::FPBDRigidsSolver* Solver = FChaosSolversModule::GetModule()->CreateSolver(nullptr, ESolverFlags::Standalone);
-		Solver->SetHasFloor(false);
+		// no floor
 		Solver->SetEnabled(true);
 
 		Solver->AdvanceSolverBy(1 / 24.);
@@ -67,7 +67,7 @@ namespace GeometryCollectionExample
 		P.CollisionGroup = -1;
 
 		Chaos::FPBDRigidsSolver* Solver = FChaosSolversModule::GetModule()->CreateSolver(nullptr, ESolverFlags::Standalone);
-		Solver->SetHasFloor(false);
+		// no floor
 		Solver->SetEnabled(true);
 		Solver->AdvanceSolverBy(1 / 24.);
 #if TODO_REIMPLEMENT_GET_RIGID_PARTICLES
@@ -116,7 +116,7 @@ namespace GeometryCollectionExample
 		P.ClusterGroupIndex = 1;
 
 		Chaos::FPBDRigidsSolver* Solver = FChaosSolversModule::GetModule()->CreateSolver(nullptr, ESolverFlags::Standalone);
-		Solver->SetHasFloor(false);
+		// no floor
 		Solver->SetEnabled(true);
 		Solver->AdvanceSolverBy(1 / 24.);
 #if TODO_REIMPLEMENT_GET_RIGID_PARTICLES
diff --git a/Engine/Source/Runtime/Experimental/ChaosSolverEngine/Private/Chaos/ChaosSolverActor.cpp b/Engine/Source/Runtime/Experimental/ChaosSolverEngine/Private/Chaos/ChaosSolverActor.cpp
--- a/Engine/Source/Runtime/Experimental/ChaosSolverEngine/Private/Chaos/ChaosSolverActor.cpp
+++ b/Engine/Source/Runtime/Experimental/ChaosSolverEngine/Private/Chaos/ChaosSolverActor.cpp
@@ -306,8 +306,7 @@ void AChaosSolverActor::BeginPlay()
 			InSolver->SetCollisionFilterSettings(InCollisionFilterSettings);
 			InSolver->SetBreakingFilterSettings(InBreakingFilterSettings);
 			InSolver->SetTrailingFilterSettings(InTrailingFilterSettings);
-			InSolver->SetHasFloor(InHasFloor);
-			InSolver->SetFloorHeight(InFloorHeight);
+
 #if TODO_REIMPLEMENT_SOLVER_SETTINGS_ACCESSORS
 			InSolver->SetMassScale(InMassScale);
 #endif
@@ -504,7 +503,6 @@ void AChaosSolverActor::PostEditChangeProperty(struct FPropertyChangedEvent& Pro
 				PhysDispatcher->EnqueueCommandImmediate(Solver, [InHasFloor = bHasFloor]
 				(Chaos::FPhysicsSolver* InSolver)
 				{
-					InSolver->SetHasFloor(InHasFloor);
 				});
 			}
 			else if (PropertyChangedEvent.Property->GetFName() == GET_MEMBER_NAME_CHECKED(AChaosSolverActor, FloorHeight))
@@ -512,7 +510,6 @@ void AChaosSolverActor::PostEditChangeProperty(struct FPropertyChangedEvent& Pro
 				PhysDispatcher->EnqueueCommandImmediate(Solver, [InFloorHeight = FloorHeight]
 				(Chaos::FPhysicsSolver* InSolver)
 				{
-					InSolver->SetFloorHeight(InFloorHeight);
 				});
 			}
 #if TODO_REIMPLEMENT_TIMESTEP_MULTIPLIER
diff --git a/Engine/Source/Runtime/Experimental/ChaosSolvers/Public/PBDRigidsSolver.h b/Engine/Source/Runtime/Experimental/ChaosSolvers/Public/PBDRigidsSolver.h
--- a/Engine/Source/Runtime/Experimental/ChaosSolvers/Public/PBDRigidsSolver.h
+++ b/Engine/Source/Runtime/Experimental/ChaosSolvers/Public/PBDRigidsSolver.h
@@ -269,9 +269,6 @@ namespace Chaos
 		void SetCollisionFilterSettings(const FSolverCollisionFilterSettings& InCollisionFilterSettings) { GetEventFilters()->GetCollisionFilter()->UpdateFilterSettings(InCollisionFilterSettings); }
 		void SetBreakingFilterSettings(const FSolverBreakingFilterSettings& InBreakingFilterSettings) { GetEventFilters()->GetBreakingFilter()->UpdateFilterSettings(InBreakingFilterSettings); }
 		void SetTrailingFilterSettings(const FSolverTrailingFilterSettings& InTrailingFilterSettings) { GetEventFilters()->GetTrailingFilter()->UpdateFilterSettings(InTrailingFilterSettings); }
-		void SetHasFloor(bool bHasFloorIn) { bHasFloor = bHasFloorIn; }
-		void SetIsFloorAnalytic(bool bIsAnalytic) { bIsFloorAnalytic = bIsAnalytic; }
-		void SetFloorHeight(float Height) { FloorHeight = Height; }
 
 		/**/
 		FPBDRigidsEvolution* GetEvolution() { return MEvolution.Get(); }
