GeometryCollectionTest : Framework

- Added a framework wrapper to reduce repeated code in the tests.

diff --git a/Engine/Plugins/Experimental/GeometryCollectionPlugin/Source/GeometryCollectionExampleCore/Private/GeometryCollection/GeometryCollectionExampleClustering.cpp b/Engine/Plugins/Experimental/GeometryCollectionPlugin/Source/GeometryCollectionExampleCore/Private/GeometryCollection/GeometryCollectionExampleClustering.cpp
--- a/Engine/Plugins/Experimental/GeometryCollectionPlugin/Source/GeometryCollectionExampleCore/Private/GeometryCollection/GeometryCollectionExampleClustering.cpp
+++ b/Engine/Plugins/Experimental/GeometryCollectionPlugin/Source/GeometryCollectionExampleCore/Private/GeometryCollection/GeometryCollectionExampleClustering.cpp
@@ -1686,7 +1686,7 @@ namespace GeometryCollectionExample
 	void RigidBodies_ClusterTest_ReleaseClusterParticles_AllLeafNodes()
 	{
 		TUniquePtr<Chaos::FChaosPhysicsMaterial> PhysicalMaterial = MakeUnique<Chaos::FChaosPhysicsMaterial>();
-		InitMaterialToZero(PhysicalMaterial);
+		InitMaterialToZero(PhysicalMaterial.Get());
 
 		TSharedPtr<FGeometryCollection> RestCollection = CreateClusteredBody(FVector(0, 0, 100));
 		TSharedPtr<FGeometryDynamicCollection> DynamicCollection = GeometryCollectionToGeometryDynamicCollection(RestCollection.Get(), (uint8)EObjectStateTypeEnum::Chaos_Object_Dynamic);
@@ -1765,7 +1765,6 @@ namespace GeometryCollectionExample
 		FChaosSolversModule::GetModule()->DestroySolver(Solver);
 
 		delete PhysObject;
-
 		
 	}
 	template void RigidBodies_ClusterTest_ReleaseClusterParticles_AllLeafNodes<float>();
@@ -1774,7 +1773,7 @@ namespace GeometryCollectionExample
 	void RigidBodies_ClusterTest_ReleaseClusterParticles_ClusterNodeAndSubClusterNode()
 	{
 		TUniquePtr<Chaos::FChaosPhysicsMaterial> PhysicalMaterial = MakeUnique<Chaos::FChaosPhysicsMaterial>();
-		InitMaterialToZero(PhysicalMaterial);
+		InitMaterialToZero(PhysicalMaterial.Get());
 
 		TSharedPtr<FGeometryCollection> RestCollection = CreateClusteredBody_TwoParents_TwoBodies(FVector(0, 0, 100));
 		TSharedPtr<FGeometryDynamicCollection> DynamicCollection = GeometryCollectionToGeometryDynamicCollection(RestCollection.Get(), (uint8)EObjectStateTypeEnum::Chaos_Object_Dynamic);
@@ -1868,7 +1867,7 @@ namespace GeometryCollectionExample
 	void RigidBodies_ClusterTest_RemoveOnFracture()
 	{
 		TUniquePtr<Chaos::FChaosPhysicsMaterial> PhysicalMaterial = MakeUnique<Chaos::FChaosPhysicsMaterial>();
-		InitMaterialToZero(PhysicalMaterial);
+		InitMaterialToZero(PhysicalMaterial.Get());
 
 		TSharedPtr<FGeometryCollection> RestCollection = GeometryCollection::MakeCubeElement(FTransform(FQuat::MakeFromEuler(FVector(0, 0, 0.)), FVector(0, -10, 10)), FVector(1.0));
 		RestCollection->AppendGeometry(*GeometryCollection::MakeCubeElement(FTransform(FQuat::MakeFromEuler(FVector(0, 0, 0.)), FVector(0, 10, 10)), FVector(1.0)));
diff --git a/Engine/Plugins/Experimental/GeometryCollectionPlugin/Source/GeometryCollectionExampleCore/Private/GeometryCollection/GeometryCollectionExampleFramework.cpp b/Engine/Plugins/Experimental/GeometryCollectionPlugin/Source/GeometryCollectionExampleCore/Private/GeometryCollection/GeometryCollectionExampleFramework.cpp
new file mode 100644
--- /dev/null
+++ b/Engine/Plugins/Experimental/GeometryCollectionPlugin/Source/GeometryCollectionExampleCore/Private/GeometryCollection/GeometryCollectionExampleFramework.cpp
@@ -0,0 +1,118 @@
+// Copyright Epic Games, Inc. All Rights Reserved.
+
+#include "GeometryCollection/GeometryCollectionExampleFramework.h"
+#include "GeometryCollection/GeometryCollectionExampleUtility.h"
+#include "ChaosSolversModule.h"
+
+#include "Chaos/ErrorReporter.h"
+
+
+namespace GeometryCollectionExample
+{
+	template<>
+	WrapperBase* NewSimulationObject<GeometryType::GeometryCollectionWithSingleCube>(const CreationParameters Params)
+	{
+		TSharedPtr<Chaos::FChaosPhysicsMaterial> PhysicalMaterial;
+		TSharedPtr<FGeometryCollection> RestCollection;
+		TSharedPtr<FGeometryDynamicCollection> DynamicCollection;
+
+		InitCollectionsParameters InitParams = { Params.Position, Params.Scale, nullptr, Params.DynamicState };
+		InitCollections(PhysicalMaterial, RestCollection, DynamicCollection, InitParams);
+
+		auto InitFunc = [&RestCollection, &DynamicCollection, &PhysicalMaterial](FSimulationParameters& InParams)
+		{
+			InParams.RestCollection = RestCollection.Get();
+			InParams.DynamicCollection = DynamicCollection.Get();
+			InParams.PhysicalMaterial = MakeSerializable(PhysicalMaterial);
+			InParams.Shared.SizeSpecificData[0].CollisionType = ECollisionTypeEnum::Chaos_Volumetric;
+
+			InParams.Simulating = true;
+			Chaos::FErrorReporter ErrorReporter;
+			BuildSimulationData(ErrorReporter, *RestCollection, InParams.Shared);
+		};
+
+		FGeometryCollectionPhysicsProxy* PhysObject = new FGeometryCollectionPhysicsProxy(nullptr, DynamicCollection.Get(), InitFunc, nullptr, nullptr);;
+		PhysObject->Initialize();
+
+		return new GeometryCollectionWrapper(PhysicalMaterial,RestCollection,DynamicCollection,PhysObject);
+	}
+
+
+	template<>
+	WrapperBase* NewSimulationObject<GeometryType::RigidBodyAnalyticCube>(const CreationParameters Params)
+	{
+		return nullptr;
+	}
+
+
+	template<class T>
+	Framework<T>::Framework(FrameworkParameters Parameters)
+		: Dt(Parameters.Dt)
+		, Module(FChaosSolversModule::GetModule())
+		, Solver(Module->CreateSolver(nullptr, ESolverFlags::Standalone))
+	{
+		Module->ChangeThreadingMode(Parameters.ThreadingMode);
+	}
+
+	template<class T>
+	Framework<T>::~Framework()
+	{
+		for (WrapperBase* Object : PhysicsObjects)
+		{
+			if (GeometryCollectionWrapper* GCW = Object->As<GeometryCollectionWrapper>())
+			{
+				delete GCW->PhysObject;
+			}
+			//else if (RigidBodyWrapper* BCW = Object->As<WrapperType::RigidBody>())
+			//{
+			//	delete BCW->PhysObject;
+			//}
+			delete Object;
+		}
+
+		FChaosSolversModule::GetModule()->DestroySolver(Solver);
+	}
+
+	template<class T>
+	void Framework<T>::AddSimulationObject(WrapperBase * Object)
+	{
+		PhysicsObjects.Add(Object);
+	}
+
+	template<class T>
+	void Framework<T>::Initialize()
+	{
+		Solver->SetEnabled(true);
+
+		for (WrapperBase* Object : PhysicsObjects)
+		{
+			if (GeometryCollectionWrapper* GCW = Object->As<GeometryCollectionWrapper>())
+			{
+				Solver->RegisterObject(GCW->PhysObject);
+				GCW->PhysObject->ActivateBodies();
+				Solver->AddDirtyProxy(GCW->PhysObject);
+			}
+			//else if (RigidBodyWrapper* RBW = Object->As<RigidBodyWrapper>())
+			//{
+			//	Solver->RegisterObject(RBW->PhysObject);
+			//	RBW->PhysObject->ActivateBodies();
+			//	Solver->AddDirtyProxy(RBW->PhysObject);
+			//}
+		}
+
+		Solver->PushPhysicsState(Module->GetDispatcher());
+	}
+
+	template<class T>
+	void Framework<T>::Advance()
+	{
+		Solver->AdvanceSolverBy(Dt);
+
+		Solver->BufferPhysicsResults();
+		Solver->FlipBuffers();
+		Solver->UpdateGameThreadStructures();
+	}
+	
+
+} // end namespace GeometryCollectionExample
+
diff --git a/Engine/Plugins/Experimental/GeometryCollectionPlugin/Source/GeometryCollectionExampleCore/Private/GeometryCollection/GeometryCollectionExampleSimulation.cpp b/Engine/Plugins/Experimental/GeometryCollectionPlugin/Source/GeometryCollectionExampleCore/Private/GeometryCollection/GeometryCollectionExampleSimulation.cpp
--- a/Engine/Plugins/Experimental/GeometryCollectionPlugin/Source/GeometryCollectionExampleCore/Private/GeometryCollection/GeometryCollectionExampleSimulation.cpp
+++ b/Engine/Plugins/Experimental/GeometryCollectionPlugin/Source/GeometryCollectionExampleCore/Private/GeometryCollection/GeometryCollectionExampleSimulation.cpp
@@ -3,6 +3,7 @@
 #include "GeometryCollection/GeometryCollectionExampleSimulation.h"
 #include "GeometryCollection/GeometryCollectionExampleUtility.h"
 #include "GeometryCollection/GeometryCollectionExampleSimulationObject.h"
+#include "GeometryCollection/GeometryCollectionExampleFramework.h"
 
 #include "GeometryCollection/GeometryCollection.h"
 #include "GeometryCollection/GeometryCollectionUtility.h"
@@ -27,83 +28,40 @@ namespace GeometryCollectionExample
 {
 	
 	template<class T>
-	void RigidBodiesFallingUnderGravity()
+	void RigidBodies_SingleFallingUnderGravity()
 	{
-		float Dt = 1 / 24.;
-		FChaosSolversModule* Module = FChaosSolversModule::GetModule();
-		Module->ChangeThreadingMode(EChaosThreadingMode::SingleThread);
-
-		//
-		//  Rigid Body Setup
-		//
-
-		InitCollectionsParameters InitParams = 
-		{ 
-			FTransform::Identity,	// RestCenter
-			FVector(1.0),			// RestScale
-			nullptr,				// RestInitFunc
-			(int32)EObjectStateTypeEnum::Chaos_Object_Dynamic // DynamicStateDefault
-		};
-
-		TUniquePtr<Chaos::FChaosPhysicsMaterial> PhysicalMaterial = nullptr; // Allocated and zero'ed
-		TSharedPtr<FGeometryCollection> RestCollection = nullptr;				// GeometryCollection::MakeCubeElement(InitParams)
-		TSharedPtr<FGeometryDynamicCollection> GTDynamicCollection = nullptr;		// New'ed w/copy of RestCollection Transform, Parent, Children, SimulationType, and StatusFlags attrs.
-		InitCollections(PhysicalMaterial, RestCollection, GTDynamicCollection, InitParams);
-
-		//
-		// Sim Initialization
-		//
-		
-		// Creates FGeometryCollectionPhysicsProxy with GTDynamicCollection, and 
-		// an InitFunc that sets FSimulationParametes's RestCollection and 
-		// DynamicCollection pointers.  Calls FGeometryCollectionPhysicsProx::Initialize().
-		FGeometryCollectionPhysicsProxy* PhysObject = RigidBodySetup(PhysicalMaterial, RestCollection, GTDynamicCollection);
-
-		Chaos::FPBDRigidsSolver* Solver = FChaosSolversModule::GetModule()->CreateSolver(nullptr, ESolverFlags::Standalone);
-		Solver->RegisterObject(PhysObject);
-		// no floor
-		Solver->SetEnabled(true);
-		PhysObject->ActivateBodies();
-
-		Solver->AddDirtyProxy(PhysObject);
-		Solver->PushPhysicsState(Module->GetDispatcher());
-
-		Solver->AdvanceSolverBy(Dt);
-
-		// Calls BufferPhysicsResults(), FlipBuffer(), and PullFromPhysicsState() on each proxy.
-		FinalizeSolver(*Solver);
-
-		// never touched
-		TManagedArray<FTransform>& RestTransform = RestCollection->Transform;
-		EXPECT_LT(FMath::Abs(RestTransform[0].GetTranslation().Z), SMALL_THRESHOLD);
-
-
-		// simulated
-		TManagedArray<FTransform>& Transform = GTDynamicCollection->Transform;
-		EXPECT_EQ(Transform.Num(), 1);
-		EXPECT_LT(Transform[0].GetTranslation().Z, 0.f);
-
-		float TwiceGravityDisplacement = /* 0.5* */-980.f * Dt * Dt; // we seem to be twice gravity
-		EXPECT_NEAR(Transform[0].GetTranslation().Z, TwiceGravityDisplacement, 1e-2);
-
-		FChaosSolversModule::GetModule()->DestroySolver(Solver);
-		
-		delete PhysObject;
+		GeometryCollectionWrapper* Collection = NewSimulationObject<GeometryType::GeometryCollectionWithSingleCube>()->As<GeometryCollectionWrapper>();
+
+		Framework<T> UnitTest;
+		UnitTest.AddSimulationObject(Collection);
+		UnitTest.Initialize();
+		UnitTest.Advance();
+
+		{ // test results
+			EXPECT_LT(FMath::Abs(Collection->RestCollection->Transform[0].GetTranslation().Z), SMALL_THRESHOLD); // rest never touched
+			EXPECT_EQ(Collection->DynamicCollection->Transform.Num(), 1); // simulated is falling
+			EXPECT_LT(Collection->DynamicCollection->Transform[0].GetTranslation().Z, 0.f);
+			EXPECT_NEAR(Collection->DynamicCollection->Transform[0].GetTranslation().Z, -980.f * UnitTest.Dt * UnitTest.Dt, 1e-2);// we seem to be twice gravity
+		}
 	}
-	template void RigidBodiesFallingUnderGravity<float>();
+	template void RigidBodies_SingleFallingUnderGravity<float>();
 
 
 
 	template<class T>
 	void RigidBodiesCollidingWithSolverFloor()
 	{
-		TUniquePtr<Chaos::FChaosPhysicsMaterial> PhysicalMaterial = nullptr;
-		TSharedPtr<FGeometryCollection> RestCollection = nullptr;
-		TSharedPtr<FGeometryDynamicCollection> DynamicCollection = nullptr;
+		float Dt = 1 / 24.;
+		FChaosSolversModule* Module = FChaosSolversModule::GetModule();
+		Module->ChangeThreadingMode(EChaosThreadingMode::SingleThread);
 
 		//
 		//  Rigid Body Setup
 		//
+		TUniquePtr<Chaos::FChaosPhysicsMaterial> PhysicalMaterial = nullptr;
+		TSharedPtr<FGeometryCollection> RestCollection = nullptr;
+		TSharedPtr<FGeometryDynamicCollection> DynamicCollection = nullptr;
+
 		InitCollectionsParameters InitParams = { FTransform::Identity, FVector(1.0), nullptr, (int32)EObjectStateTypeEnum::Chaos_Object_Kinematic };
 		InitCollections(PhysicalMaterial, RestCollection, DynamicCollection, InitParams);
 				
@@ -115,9 +73,8 @@ namespace GeometryCollectionExample
 		FGeometryCollectionPhysicsProxy* PhysObject = RigidBodySetup(PhysicalMaterial, RestCollection, DynamicCollection, CustomFunc);
 
 		Chaos::FPBDRigidsSolver* Solver = FChaosSolversModule::GetModule()->CreateSolver(nullptr, ESolverFlags::Standalone);
-#if CHAOS_PARTICLEHANDLE_TODO
 		Solver->RegisterObject(PhysObject);
-#endif
+
 		////Solver->SetHasFloor(true);
 		////Solver->SetIsFloorAnalytic(true);
 		Solver->SetEnabled(true);
diff --git a/Engine/Plugins/Experimental/GeometryCollectionPlugin/Source/GeometryCollectionExampleCore/Private/GeometryCollection/GeometryCollectionExampleSimulationSolver.cpp b/Engine/Plugins/Experimental/GeometryCollectionPlugin/Source/GeometryCollectionExampleCore/Private/GeometryCollection/GeometryCollectionExampleSimulationSolver.cpp
--- a/Engine/Plugins/Experimental/GeometryCollectionPlugin/Source/GeometryCollectionExampleCore/Private/GeometryCollection/GeometryCollectionExampleSimulationSolver.cpp
+++ b/Engine/Plugins/Experimental/GeometryCollectionPlugin/Source/GeometryCollectionExampleCore/Private/GeometryCollection/GeometryCollectionExampleSimulationSolver.cpp
@@ -154,7 +154,7 @@ namespace GeometryCollectionExample
 	void Solver_ValidateReverseMapping()
 	{
 		TUniquePtr<Chaos::FChaosPhysicsMaterial> PhysicalMaterial = MakeUnique<Chaos::FChaosPhysicsMaterial>();
-		InitMaterialToZero(PhysicalMaterial);
+		InitMaterialToZero(PhysicalMaterial.Get());
 
 		Chaos::FPBDRigidsSolver* Solver = FChaosSolversModule::GetModule()->CreateSolver(nullptr, ESolverFlags::Standalone);
 		//Solver->SetHasFloor(true);
@@ -215,7 +215,7 @@ namespace GeometryCollectionExample
 	void CommonInit(int32 NumObjects, bool UseClusters, Chaos::FPBDRigidsSolver** SolverInOut, TUniquePtr<Chaos::FChaosPhysicsMaterial>& PhysicalMaterial, TArray<TSharedPtr<FGeometryCollection>>& RestArray, TArray<TSharedPtr<FGeometryDynamicCollection>>& DynamicArray)
 	{
 		PhysicalMaterial = MakeUnique<Chaos::FChaosPhysicsMaterial>();
-		InitMaterialToZero(PhysicalMaterial);
+		InitMaterialToZero(PhysicalMaterial.Get());
 
 		*SolverInOut = FChaosSolversModule::GetModule()->CreateSolver(nullptr, ESolverFlags::Standalone);
 		Chaos::FPBDRigidsSolver* Solver = *SolverInOut;
diff --git a/Engine/Plugins/Experimental/GeometryCollectionPlugin/Source/GeometryCollectionExampleCore/Private/GeometryCollection/GeometryCollectionExampleUtility.cpp b/Engine/Plugins/Experimental/GeometryCollectionPlugin/Source/GeometryCollectionExampleCore/Private/GeometryCollection/GeometryCollectionExampleUtility.cpp
--- a/Engine/Plugins/Experimental/GeometryCollectionPlugin/Source/GeometryCollectionExampleCore/Private/GeometryCollection/GeometryCollectionExampleUtility.cpp
+++ b/Engine/Plugins/Experimental/GeometryCollectionPlugin/Source/GeometryCollectionExampleCore/Private/GeometryCollection/GeometryCollectionExampleUtility.cpp
@@ -200,7 +200,7 @@ namespace GeometryCollectionExample {
 		return RestCollection;
 	}
 
-	void InitMaterialToZero(TUniquePtr<Chaos::FChaosPhysicsMaterial> const &PhysicalMaterial)
+	void InitMaterialToZero(Chaos::FChaosPhysicsMaterial * PhysicalMaterial)
 	{
 		PhysicalMaterial->Friction = 0;
 		PhysicalMaterial->Restitution = 0;
@@ -221,7 +221,37 @@ namespace GeometryCollectionExample {
 		if (PhysicalMaterial == nullptr)
 		{
 			PhysicalMaterial = MakeUnique<Chaos::FChaosPhysicsMaterial>();
-			InitMaterialToZero(PhysicalMaterial);
+			InitMaterialToZero(PhysicalMaterial.Get());
+		}
+
+		if (RestCollection == nullptr)
+		{
+			//Default initialization is a cube of the specified center and size. 
+			RestCollection = GeometryCollection::MakeCubeElement(InitParams.RestCenter, InitParams.RestScale);
+			if (InitParams.RestInitFunc != nullptr)
+			{
+				InitParams.RestInitFunc(RestCollection);
+			}
+		}
+
+		if (DynamicCollection == nullptr)
+		{
+			DynamicCollection = GeometryCollectionToGeometryDynamicCollection(RestCollection.Get(), InitParams.DynamicStateDefault);
+		}
+	}
+
+	void InitCollections(
+		TSharedPtr<Chaos::FChaosPhysicsMaterial>& PhysicalMaterial,
+		TSharedPtr<FGeometryCollection>& RestCollection,
+		TSharedPtr<FGeometryDynamicCollection>& DynamicCollection,
+		InitCollectionsParameters& InitParams
+	)
+	{
+		// Allow for customized initialization of these objects in the calling function. 
+		if (PhysicalMaterial == nullptr)
+		{
+			PhysicalMaterial = MakeShared<Chaos::FChaosPhysicsMaterial>();
+			InitMaterialToZero(PhysicalMaterial.Get());
 		}
 
 		if (RestCollection == nullptr)
diff --git a/Engine/Plugins/Experimental/GeometryCollectionPlugin/Source/GeometryCollectionExampleCore/Public/GeometryCollection/GeometryCollectionExampleFramework.h b/Engine/Plugins/Experimental/GeometryCollectionPlugin/Source/GeometryCollectionExampleCore/Public/GeometryCollection/GeometryCollectionExampleFramework.h
new file mode 100644
--- /dev/null
+++ b/Engine/Plugins/Experimental/GeometryCollectionPlugin/Source/GeometryCollectionExampleCore/Public/GeometryCollection/GeometryCollectionExampleFramework.h
@@ -0,0 +1,96 @@
+// Copyright Epic Games, Inc. All Rights Reserved.
+#pragma once
+
+#include "ChaosSolversModule.h"
+#include "GeometryCollection/GeometryCollection.h"
+#include "GeometryCollection/GeometryDynamicCollection.h"
+#include "PhysicsProxy/PhysicsProxies.h"
+
+namespace GeometryCollectionExample
+{
+	enum WrapperType
+	{
+		RigidBody,
+		GeometryCollection
+	};
+
+	struct WrapperBase
+	{
+		WrapperType Type;
+		WrapperBase(WrapperType TypeIn) : Type(TypeIn) {}
+		template<class AS_T> AS_T* As() { return AS_T::StaticType() == Type ? static_cast<AS_T*>(this) : nullptr; }
+		template<class AS_T> const AS_T* As() const { return AS_T::StaticType() == Type ? static_cast<const AS_T*>(this) : nullptr; }
+	};
+
+	struct GeometryCollectionWrapper : public WrapperBase
+	{
+		GeometryCollectionWrapper() : WrapperBase(WrapperType::GeometryCollection) {}
+		GeometryCollectionWrapper(
+			TSharedPtr<Chaos::FChaosPhysicsMaterial> PhysicalMaterialIn,
+			TSharedPtr<FGeometryCollection> RestCollectionIn,
+			TSharedPtr<FGeometryDynamicCollection> DynamicCollectionIn,
+			FGeometryCollectionPhysicsProxy* PhysObjectIn)
+			: WrapperBase(WrapperType::GeometryCollection)
+			, PhysicalMaterial(PhysicalMaterialIn)
+			, RestCollection(RestCollectionIn)
+			, DynamicCollection(DynamicCollectionIn)
+			, PhysObject(PhysObjectIn) {}
+		static WrapperType StaticType() { return WrapperType::GeometryCollection; }
+		TSharedPtr<Chaos::FChaosPhysicsMaterial> PhysicalMaterial;
+		TSharedPtr<FGeometryCollection> RestCollection;
+		TSharedPtr<FGeometryDynamicCollection> DynamicCollection;
+		FGeometryCollectionPhysicsProxy* PhysObject;
+	};
+
+	struct RigidBodyWrapper : public WrapperBase
+	{
+		RigidBodyWrapper() : WrapperBase(WrapperType::RigidBody) {}
+		static WrapperType StaticType() { return WrapperType::RigidBody; }
+		TUniquePtr<Chaos::FChaosPhysicsMaterial> PhysicalMaterial;
+		FSingleParticlePhysicsProxy<float>* PhysObject;
+	};
+
+
+	struct CreationParameters
+	{
+		FTransform Position = FTransform::Identity;
+		FVector Scale = FVector(1.0);
+		int32 DynamicState = (int32)EObjectStateTypeEnum::Chaos_Object_Dynamic;
+	};
+
+
+	enum GeometryType {
+		GeometryCollectionWithSingleCube,
+		RigidBodyAnalyticCube
+	};
+
+
+	template<GeometryType>
+	WrapperBase* NewSimulationObject(const CreationParameters Params = CreationParameters());
+
+
+	struct FrameworkParameters
+	{
+		float Dt = 1 / 24.;
+		Chaos::EThreadingMode ThreadingMode = Chaos::EThreadingMode::SingleThread;
+	};
+
+	template<class T>
+	class Framework
+	{
+	public:
+
+		Framework(FrameworkParameters Properties = FrameworkParameters());
+		virtual ~Framework();
+
+		void AddSimulationObject(WrapperBase* Object);
+		void Initialize();
+		void Advance();
+
+		T Dt;
+		FChaosSolversModule* Module = FChaosSolversModule::GetModule();
+		Chaos::FPBDRigidsSolver* Solver;
+		TArray< WrapperBase* > PhysicsObjects;
+	};
+
+}
\ No newline at end of file
diff --git a/Engine/Plugins/Experimental/GeometryCollectionPlugin/Source/GeometryCollectionExampleCore/Public/GeometryCollection/GeometryCollectionExampleSimulation.h b/Engine/Plugins/Experimental/GeometryCollectionPlugin/Source/GeometryCollectionExampleCore/Public/GeometryCollection/GeometryCollectionExampleSimulation.h
--- a/Engine/Plugins/Experimental/GeometryCollectionPlugin/Source/GeometryCollectionExampleCore/Public/GeometryCollection/GeometryCollectionExampleSimulation.h
+++ b/Engine/Plugins/Experimental/GeometryCollectionPlugin/Source/GeometryCollectionExampleCore/Public/GeometryCollection/GeometryCollectionExampleSimulation.h
@@ -8,7 +8,7 @@ namespace GeometryCollectionExample
 {	
 	
 	template<class T>
-	void RigidBodiesFallingUnderGravity();
+	void RigidBodies_SingleFallingUnderGravity();
 
 	template<class T>
 	void RigidBodiesCollidingWithSolverFloor();
diff --git a/Engine/Plugins/Experimental/GeometryCollectionPlugin/Source/GeometryCollectionExampleCore/Public/GeometryCollection/GeometryCollectionExampleUtility.h b/Engine/Plugins/Experimental/GeometryCollectionPlugin/Source/GeometryCollectionExampleCore/Public/GeometryCollection/GeometryCollectionExampleUtility.h
--- a/Engine/Plugins/Experimental/GeometryCollectionPlugin/Source/GeometryCollectionExampleCore/Public/GeometryCollection/GeometryCollectionExampleUtility.h
+++ b/Engine/Plugins/Experimental/GeometryCollectionPlugin/Source/GeometryCollectionExampleCore/Public/GeometryCollection/GeometryCollectionExampleUtility.h
@@ -32,7 +32,7 @@ namespace GeometryCollectionExample {
 	TSharedPtr<FGeometryCollection>	CreateClusteredBody_ThreeByTwo_ThreeTransform(FVector Position);
 	TSharedPtr<FGeometryCollection>	CreateClusteredBody_FracturedGeometry(FVector Position = FVector(0));
 
-	void InitMaterialToZero(TUniquePtr<Chaos::FChaosPhysicsMaterial> const &PhysicalMaterial);
+	void InitMaterialToZero(Chaos::FChaosPhysicsMaterial * PhysicalMaterial);
 
 	using FCollectionInitFunc = TFunction<void(TSharedPtr<FGeometryCollection>&)>;
 	struct InitCollectionsParameters
@@ -50,6 +50,13 @@ namespace GeometryCollectionExample {
 		InitCollectionsParameters& InitParams
 	);
 
+	void InitCollections(
+		TSharedPtr<Chaos::FChaosPhysicsMaterial>& PhysicalMaterial,
+		TSharedPtr<FGeometryCollection>& RestCollection,
+		TSharedPtr<FGeometryDynamicCollection>& DynamicCollection,
+		InitCollectionsParameters& InitParams
+	);
+
 	using FInitFunc = TFunction<void(FSimulationParameters&)>;
 	FGeometryCollectionPhysicsProxy* RigidBodySetup(
 		TUniquePtr<Chaos::FChaosPhysicsMaterial> &PhysicalMaterial,
