Numerous non-unity fixes

diff --git a/Engine/Plugins/Editor/EditorDebugTools/Source/EditorDebugTools/Private/SDebugPanel.cpp b/Engine/Plugins/Editor/EditorDebugTools/Source/EditorDebugTools/Private/SDebugPanel.cpp
--- a/Engine/Plugins/Editor/EditorDebugTools/Source/EditorDebugTools/Private/SDebugPanel.cpp
+++ b/Engine/Plugins/Editor/EditorDebugTools/Source/EditorDebugTools/Private/SDebugPanel.cpp
@@ -4,7 +4,11 @@
 #include "ISlateReflectorModule.h"
 #include "Widgets/Testing/SStarshipSuite.h"
 #include "Widgets/Testing/STestSuite.h"
-
+#include "Widgets\SBoxPanel.h"
+#include "Widgets\Input\SButton.h"
+#include "Styling\CoreStyle.h"
+#include "Modules\ModuleManager.h"
+#include "Framework\Application\SlateApplication.h"
 
 void SDebugPanel::Construct(const FArguments& InArgs)
 {
diff --git a/Engine/Plugins/Enterprise/DatasmithImporter/Source/DatasmithImporter/Public/Utility/DatasmithImporterUtils.h b/Engine/Plugins/Enterprise/DatasmithImporter/Source/DatasmithImporter/Public/Utility/DatasmithImporterUtils.h
--- a/Engine/Plugins/Enterprise/DatasmithImporter/Source/DatasmithImporter/Public/Utility/DatasmithImporterUtils.h
+++ b/Engine/Plugins/Enterprise/DatasmithImporter/Source/DatasmithImporter/Public/Utility/DatasmithImporterUtils.h
@@ -7,6 +7,7 @@
 #include "IDatasmithSceneElements.h"
 #include "InterchangeManager.h"
 #include "UObject/UObjectHash.h"
+#include "Engine/Texture.h"
 
 class ADatasmithSceneActor;
 class IDatasmithScene;
diff --git a/Engine/Plugins/Experimental/ControlRig/Source/ControlRig/Public/Tools/ControlRigPoseProjectSettings.h b/Engine/Plugins/Experimental/ControlRig/Source/ControlRig/Public/Tools/ControlRigPoseProjectSettings.h
--- a/Engine/Plugins/Experimental/ControlRig/Source/ControlRig/Public/Tools/ControlRigPoseProjectSettings.h
+++ b/Engine/Plugins/Experimental/ControlRig/Source/ControlRig/Public/Tools/ControlRigPoseProjectSettings.h
@@ -5,6 +5,7 @@
 */
 #include "CoreMinimal.h"
 #include "UObject/Object.h"
+#include "Engine/EngineTypes.h"
 #include "ControlRigPoseProjectSettings.generated.h"
 
 
diff --git a/Engine/Plugins/Experimental/ControlRig/Source/ControlRigEditor/Private/EditMode/ControlRigEditModeToolkit.cpp b/Engine/Plugins/Experimental/ControlRig/Source/ControlRigEditor/Private/EditMode/ControlRigEditModeToolkit.cpp
--- a/Engine/Plugins/Experimental/ControlRig/Source/ControlRigEditor/Private/EditMode/ControlRigEditModeToolkit.cpp
+++ b/Engine/Plugins/Experimental/ControlRig/Source/ControlRigEditor/Private/EditMode/ControlRigEditModeToolkit.cpp
@@ -11,6 +11,7 @@
 #include "EditorModeManager.h"
 #include "SControlRigEditModeTools.h"
 #include "ControlRigEditMode.h"
+#include "Modules/ModuleManager.h"
 
 #define LOCTEXT_NAMESPACE "FControlRigEditModeToolkit"
 
diff --git a/Engine/Plugins/Experimental/ControlRig/Source/ControlRigEditor/Private/EditMode/SControlRigBaseListWidget.cpp b/Engine/Plugins/Experimental/ControlRig/Source/ControlRigEditor/Private/EditMode/SControlRigBaseListWidget.cpp
--- a/Engine/Plugins/Experimental/ControlRig/Source/ControlRigEditor/Private/EditMode/SControlRigBaseListWidget.cpp
+++ b/Engine/Plugins/Experimental/ControlRig/Source/ControlRigEditor/Private/EditMode/SControlRigBaseListWidget.cpp
@@ -9,6 +9,7 @@
 #include "EditorStyleSet.h"
 
 #include "Widgets/Input/SSearchBox.h"
+#include "Widgets/Input/SButton.h"
 #include "Widgets/Text/SInlineEditableTextBlock.h"
 #include "ScopedTransaction.h"
 
diff --git a/Engine/Plugins/Experimental/ControlRig/Source/ControlRigEditor/Private/EditMode/SControlRigControlViews.cpp b/Engine/Plugins/Experimental/ControlRig/Source/ControlRigEditor/Private/EditMode/SControlRigControlViews.cpp
--- a/Engine/Plugins/Experimental/ControlRig/Source/ControlRigEditor/Private/EditMode/SControlRigControlViews.cpp
+++ b/Engine/Plugins/Experimental/ControlRig/Source/ControlRigEditor/Private/EditMode/SControlRigControlViews.cpp
@@ -17,6 +17,7 @@
 #include "Tools/ControlRigPoseMirrorSettings.h"
 #include "ControlRigEditMode.h"
 #include "EditorModeManager.h"
+#include "Widgets/Input/SButton.h"
 
 #define LOCTEXT_NAMESPACE "ControlRigBaseListWidget"
 
diff --git a/Engine/Plugins/Experimental/ControlRig/Source/ControlRigEditor/Private/EditMode/SControlRigEditModeTools.cpp b/Engine/Plugins/Experimental/ControlRig/Source/ControlRigEditor/Private/EditMode/SControlRigEditModeTools.cpp
--- a/Engine/Plugins/Experimental/ControlRig/Source/ControlRigEditor/Private/EditMode/SControlRigEditModeTools.cpp
+++ b/Engine/Plugins/Experimental/ControlRig/Source/ControlRigEditor/Private/EditMode/SControlRigEditModeTools.cpp
@@ -24,6 +24,7 @@
 #include "DetailWidgetRow.h"
 #include "Rigs/FKControlRig.h"
 #include "SControlRigBaseListWidget.h"
+#include "Framework/MultiBox/MultiBoxBuilder.h"
 
 #define LOCTEXT_NAMESPACE "ControlRigRootCustomization"
 
diff --git a/Engine/Plugins/Experimental/ControlRig/Source/ControlRigEditor/Private/EditMode/SControlRigEditModeTools.h b/Engine/Plugins/Experimental/ControlRig/Source/ControlRigEditor/Private/EditMode/SControlRigEditModeTools.h
--- a/Engine/Plugins/Experimental/ControlRig/Source/ControlRigEditor/Private/EditMode/SControlRigEditModeTools.h
+++ b/Engine/Plugins/Experimental/ControlRig/Source/ControlRigEditor/Private/EditMode/SControlRigEditModeTools.h
@@ -15,6 +15,7 @@ class SControlPicker;
 class SExpandableArea;
 class SControlHierarchy;
 class UControlRig;
+class FToolBarBuilder;
 
 class SControlRigEditModeTools : public SCompoundWidget, public IDetailKeyframeHandler
 {
diff --git a/Engine/Plugins/Experimental/ControlRig/Source/ControlRigEditor/Private/Editor/SControlRigStackView.cpp b/Engine/Plugins/Experimental/ControlRig/Source/ControlRigEditor/Private/Editor/SControlRigStackView.cpp
--- a/Engine/Plugins/Experimental/ControlRig/Source/ControlRigEditor/Private/Editor/SControlRigStackView.cpp
+++ b/Engine/Plugins/Experimental/ControlRig/Source/ControlRigEditor/Private/Editor/SControlRigStackView.cpp
@@ -12,6 +12,7 @@
 #include "ControlRigBlueprintGeneratedClass.h"
 #include "Framework/MultiBox/MultiBoxBuilder.h"
 #include "EditorStyleSet.h"
+#include "Widgets/Input/SSearchBox.h"
 
 #define LOCTEXT_NAMESPACE "SControlRigStackView"
 
diff --git a/Engine/Plugins/Experimental/ControlRig/Source/ControlRigEditor/Private/Editor/SControlRigStackView.h b/Engine/Plugins/Experimental/ControlRig/Source/ControlRigEditor/Private/Editor/SControlRigStackView.h
--- a/Engine/Plugins/Experimental/ControlRig/Source/ControlRigEditor/Private/Editor/SControlRigStackView.h
+++ b/Engine/Plugins/Experimental/ControlRig/Source/ControlRigEditor/Private/Editor/SControlRigStackView.h
@@ -12,6 +12,7 @@
 class FControlRigEditor;
 class SControlRigStackView;
 class FUICommandList;
+class SSearchBox;
 
 namespace ERigStackEntry
 {
diff --git a/Engine/Plugins/Experimental/ControlRig/Source/ControlRigEditor/Public/ControlRigBlueprintEditorLibrary.h b/Engine/Plugins/Experimental/ControlRig/Source/ControlRigEditor/Public/ControlRigBlueprintEditorLibrary.h
--- a/Engine/Plugins/Experimental/ControlRig/Source/ControlRigEditor/Public/ControlRigBlueprintEditorLibrary.h
+++ b/Engine/Plugins/Experimental/ControlRig/Source/ControlRigEditor/Public/ControlRigBlueprintEditorLibrary.h
@@ -10,6 +10,7 @@
 #include "CoreMinimal.h"
 #include "UObject/ObjectMacros.h"
 #include "ControlRigBlueprint.h"
+#include "Kismet/BlueprintFunctionLibrary.h"
 #include "ControlRigBlueprintEditorLibrary.generated.h"
 
 UCLASS(meta=(ScriptName="ControlRigBlueprintLibrary"))
diff --git a/Engine/Plugins/Experimental/GeometryCollectionPlugin/Source/GeometryCollectionEditor/Public/GeometryCollection/GeometryCollectionConversion.h b/Engine/Plugins/Experimental/GeometryCollectionPlugin/Source/GeometryCollectionEditor/Public/GeometryCollection/GeometryCollectionConversion.h
--- a/Engine/Plugins/Experimental/GeometryCollectionPlugin/Source/GeometryCollectionEditor/Public/GeometryCollection/GeometryCollectionConversion.h
+++ b/Engine/Plugins/Experimental/GeometryCollectionPlugin/Source/GeometryCollectionEditor/Public/GeometryCollection/GeometryCollectionConversion.h
@@ -6,6 +6,7 @@
 class UStaticMesh;
 class USkeletalMesh;
 class UGeometryCollection;
+class UMaterialInterface;
 
 typedef TTuple<const UStaticMesh *, const UStaticMeshComponent *, FTransform> GeometryCollectionStaticMeshConversionTuple;
 typedef TTuple<const USkeletalMesh *, const USkeletalMeshComponent *, FTransform> GeometryCollectionSkeletalMeshConversionTuple;
diff --git a/Engine/Plugins/FX/Niagara/Source/NiagaraEditor/Private/Widgets/SNiagaraGeneratedCodeView.h b/Engine/Plugins/FX/Niagara/Source/NiagaraEditor/Private/Widgets/SNiagaraGeneratedCodeView.h
--- a/Engine/Plugins/FX/Niagara/Source/NiagaraEditor/Private/Widgets/SNiagaraGeneratedCodeView.h
+++ b/Engine/Plugins/FX/Niagara/Source/NiagaraEditor/Private/Widgets/SNiagaraGeneratedCodeView.h
@@ -18,6 +18,7 @@
 #include "Widgets/Input/SSearchBox.h"
 #include "Widgets/Input/SComboButton.h"
 #include "Widgets/Input/SMultiLineEditableTextBox.h"
+#include "Widgets/Docking/SDockTab.h"
 
 class SNiagaraGeneratedCodeView : public SCompoundWidget
 {
diff --git a/Engine/Plugins/Runtime/Interchange/Source/Export/Private/InterchangeExportPlugin.cpp b/Engine/Plugins/Runtime/Interchange/Source/Export/Private/InterchangeExportPlugin.cpp
--- a/Engine/Plugins/Runtime/Interchange/Source/Export/Private/InterchangeExportPlugin.cpp
+++ b/Engine/Plugins/Runtime/Interchange/Source/Export/Private/InterchangeExportPlugin.cpp
@@ -5,6 +5,7 @@
 #include "IInterchangeExportPlugin.h"
 #include "InterchangeManager.h"
 #include "InterchangeTextureWriter.h"
+#include "Engine/Engine.h"
 
 class FInterchangeExportPlugin : public IInterchangeExportPlugin
 {
diff --git a/Engine/Plugins/Runtime/Interchange/Source/Import/Public/Material/InterchangeMaterialPayloadInterface.h b/Engine/Plugins/Runtime/Interchange/Source/Import/Public/Material/InterchangeMaterialPayloadInterface.h
--- a/Engine/Plugins/Runtime/Interchange/Source/Import/Public/Material/InterchangeMaterialPayloadInterface.h
+++ b/Engine/Plugins/Runtime/Interchange/Source/Import/Public/Material/InterchangeMaterialPayloadInterface.h
@@ -6,6 +6,7 @@
 #include "Material/InterchangeMaterialPayloadData.h"
 #include "UObject/Object.h"
 #include "UObject/ObjectMacros.h"
+#include "UObject/Interface.h"
 
 #include "InterchangeMaterialPayloadInterface.generated.h"
 
diff --git a/Engine/Plugins/Runtime/Interchange/Source/Import/Public/Mesh/InterchangeStaticMeshPayloadInterface.h b/Engine/Plugins/Runtime/Interchange/Source/Import/Public/Mesh/InterchangeStaticMeshPayloadInterface.h
--- a/Engine/Plugins/Runtime/Interchange/Source/Import/Public/Mesh/InterchangeStaticMeshPayloadInterface.h
+++ b/Engine/Plugins/Runtime/Interchange/Source/Import/Public/Mesh/InterchangeStaticMeshPayloadInterface.h
@@ -6,6 +6,7 @@
 #include "Mesh/InterchangeStaticMeshPayload.h"
 #include "UObject/Object.h"
 #include "UObject/ObjectMacros.h"
+#include "UObject/Interface.h"
 
 #include "InterchangeStaticMeshPayloadInterface.generated.h"
 
diff --git a/Engine/Plugins/Runtime/Oculus/OculusVR/Source/OculusMR/Private/OculusMR_CastingCameraActor.h b/Engine/Plugins/Runtime/Oculus/OculusVR/Source/OculusMR/Private/OculusMR_CastingCameraActor.h
--- a/Engine/Plugins/Runtime/Oculus/OculusVR/Source/OculusMR/Private/OculusMR_CastingCameraActor.h
+++ b/Engine/Plugins/Runtime/Oculus/OculusVR/Source/OculusMR/Private/OculusMR_CastingCameraActor.h
@@ -18,6 +18,7 @@ class AOculusMR_BoundaryActor;
 class UTextureRenderTarget2D;
 class UOculusMR_Settings;
 class UOculusMR_State;
+class UMaterialInstanceDynamic;
 
 /**
 * The camera actor in the level that tracks the binded physical camera in game
diff --git a/Engine/Plugins/XGEController/Source/Private/XGEControllerModule.cpp b/Engine/Plugins/XGEController/Source/Private/XGEControllerModule.cpp
--- a/Engine/Plugins/XGEController/Source/Private/XGEControllerModule.cpp
+++ b/Engine/Plugins/XGEController/Source/Private/XGEControllerModule.cpp
@@ -10,6 +10,7 @@
 #include "Misc/Paths.h"
 #include "Misc/Guid.h"
 #include "Serialization/MemoryWriter.h"
+#include "Modules/ModuleManager.h"
 
 // Comma separated list of executable file names which should be intercepted by XGE.
 // Update this list if adding new tasks.
diff --git a/Engine/Plugins/XGEController/Source/Public/XGEControllerModule.h b/Engine/Plugins/XGEController/Source/Public/XGEControllerModule.h
--- a/Engine/Plugins/XGEController/Source/Public/XGEControllerModule.h
+++ b/Engine/Plugins/XGEController/Source/Public/XGEControllerModule.h
@@ -3,6 +3,7 @@
 
 #include "DistributedBuildInterface/Public/DistributedBuildControllerInterface.h"
 #include "Windows/WindowsPlatformNamedPipe.h"
+#include "Containers/Queue.h"
 
 class FXGEControllerModule : public IDistributedBuildController
 {
diff --git a/Engine/Source/Developer/SlateReflector/Private/VisualTreeCapture.h b/Engine/Source/Developer/SlateReflector/Private/VisualTreeCapture.h
--- a/Engine/Source/Developer/SlateReflector/Private/VisualTreeCapture.h
+++ b/Engine/Source/Developer/SlateReflector/Private/VisualTreeCapture.h
@@ -12,6 +12,7 @@ class FSlateWindowElementList;
 class FPaintArgs;
 struct FGeometry;
 class FSlateRect;
+class FSlateInvalidationRoot;
 
 class FVisualEntry
 {
diff --git a/Engine/Source/Developer/TraceServices/Private/Model/CallstacksProvider.cpp b/Engine/Source/Developer/TraceServices/Private/Model/CallstacksProvider.cpp
--- a/Engine/Source/Developer/TraceServices/Private/Model/CallstacksProvider.cpp
+++ b/Engine/Source/Developer/TraceServices/Private/Model/CallstacksProvider.cpp
@@ -4,6 +4,8 @@
 #include "Misc/ScopeRWLock.h"
 #include "TraceServices/Model/AnalysisSession.h"
 #include "Algo/Unique.h"
+#include "Containers/ArrayView.h"
+#include "UObject/NameTypes.h"
 
 namespace TraceServices
 {
diff --git a/Engine/Source/Developer/TraceServices/Private/Model/CallstacksProvider.h b/Engine/Source/Developer/TraceServices/Private/Model/CallstacksProvider.h
--- a/Engine/Source/Developer/TraceServices/Private/Model/CallstacksProvider.h
+++ b/Engine/Source/Developer/TraceServices/Private/Model/CallstacksProvider.h
@@ -5,6 +5,7 @@
 #include "TraceServices/Model/Callstack.h"
 #include "HAL/CriticalSection.h"
 #include "Common/PagedArray.h"
+#include "Containers/Map.h"
 
 namespace TraceServices
 {
diff --git a/Engine/Source/Developer/TraceServices/Public/TraceServices/Utils.h b/Engine/Source/Developer/TraceServices/Public/TraceServices/Utils.h
--- a/Engine/Source/Developer/TraceServices/Public/TraceServices/Utils.h
+++ b/Engine/Source/Developer/TraceServices/Public/TraceServices/Utils.h
@@ -2,6 +2,9 @@
 
 #pragma once
 
+#include "Trace/Analyzer.h"
+#include "Containers/UnrealString.h"
+
 namespace TraceServices {
 namespace FTraceAnalyzerUtils {
 
diff --git a/Engine/Source/Editor/AssetPlacementEdMode/Private/AssetPlacementEdModeModule.cpp b/Engine/Source/Editor/AssetPlacementEdMode/Private/AssetPlacementEdModeModule.cpp
--- a/Engine/Source/Editor/AssetPlacementEdMode/Private/AssetPlacementEdModeModule.cpp
+++ b/Engine/Source/Editor/AssetPlacementEdMode/Private/AssetPlacementEdModeModule.cpp
@@ -3,6 +3,7 @@
 #include "AssetPlacementEdModeModule.h"
 #include "AssetPlacementEdModeCommands.h"
 #include "AssetPlacementEdModeStyle.h"
+#include "Modules/ModuleManager.h"
 
 #define LOCTEXT_NAMESPACE "AssetPlacementEdMode"
 
diff --git a/Engine/Source/Editor/LandscapeEditor/Private/LandscapeEditorObject.cpp b/Engine/Source/Editor/LandscapeEditor/Private/LandscapeEditorObject.cpp
--- a/Engine/Source/Editor/LandscapeEditor/Private/LandscapeEditorObject.cpp
+++ b/Engine/Source/Editor/LandscapeEditor/Private/LandscapeEditorObject.cpp
@@ -12,6 +12,7 @@
 #include "LandscapeMaterialInstanceConstant.h"
 #include "Misc/ConfigCacheIni.h"
 #include "EngineUtils.h"
+#include "Misc/MessageDialog.h"
 
 //#define LOCTEXT_NAMESPACE "LandscapeEditor"
 
diff --git a/Engine/Source/Editor/Layers/Private/SLayersViewRow.cpp b/Engine/Source/Editor/Layers/Private/SLayersViewRow.cpp
--- a/Engine/Source/Editor/Layers/Private/SLayersViewRow.cpp
+++ b/Engine/Source/Editor/Layers/Private/SLayersViewRow.cpp
@@ -10,6 +10,7 @@
 #include "DragAndDrop/FolderDragDropOp.h"
 #include "Widgets/Text/SInlineEditableTextBlock.h"
 #include "EditorActorFolders.h"
+#include "Engine/World.h"
 
 #define LOCTEXT_NAMESPACE "LayersView"
 
diff --git a/Engine/Source/Editor/NewLevelDialog/Private/NewLevelDialogModule.cpp b/Engine/Source/Editor/NewLevelDialog/Private/NewLevelDialogModule.cpp
--- a/Engine/Source/Editor/NewLevelDialog/Private/NewLevelDialogModule.cpp
+++ b/Engine/Source/Editor/NewLevelDialog/Private/NewLevelDialogModule.cpp
@@ -21,12 +21,14 @@
 #include "Widgets/Layout/SBox.h"
 #include "Widgets/Input/SButton.h"
 #include "Widgets/Layout/SScrollBox.h"
+#include "Widgets/Input/SCheckBox.h"
 #include "EditorStyleSet.h"
 #include "RenderingThread.h"
 #include "Editor/UnrealEdEngine.h"
 #include "Engine/Texture2D.h"
 #include "UnrealEdGlobals.h"
 
+
 #define LOCTEXT_NAMESPACE "NewLevelDialog"
 
 /**
diff --git a/Engine/Source/Editor/Persona/Private/BlendProfileCustomization.cpp b/Engine/Source/Editor/Persona/Private/BlendProfileCustomization.cpp
--- a/Engine/Source/Editor/Persona/Private/BlendProfileCustomization.cpp
+++ b/Engine/Source/Editor/Persona/Private/BlendProfileCustomization.cpp
@@ -10,6 +10,7 @@
 #include "ISkeletonEditorModule.h"
 #include "Kismet2/BlueprintEditorUtils.h"
 #include "Animation/AnimBlueprint.h"
+#include "Modules/ModuleManager.h"
 
 #define LOCTEXT_NAMESPACE "BlendProfileCustomization"
 
diff --git a/Engine/Source/Editor/PlacementMode/Private/SPlacementModeTools.cpp b/Engine/Source/Editor/PlacementMode/Private/SPlacementModeTools.cpp
--- a/Engine/Source/Editor/PlacementMode/Private/SPlacementModeTools.cpp
+++ b/Engine/Source/Editor/PlacementMode/Private/SPlacementModeTools.cpp
@@ -16,6 +16,7 @@
 #include "ContentBrowserDataDragDropOp.h"
 #include "EditorClassUtils.h"
 #include "Widgets/Input/SSearchBox.h"
+#include "Widgets/Docking/SDockTab.h"
 
 
 struct FSortPlaceableItems
@@ -92,6 +93,7 @@ public:
 	SLATE_ARGUMENT( TOptional<FLinearColor>, AssetTypeColorOverride )
 	SLATE_END_ARGS()
 
+
 	void Construct( const FArguments& InArgs, const FAssetData& InAsset)
 	{
 		Asset = InAsset;
diff --git a/Engine/Source/Editor/UnrealEd/Public/Tools/LegacyEdModeInterfaces.h b/Engine/Source/Editor/UnrealEd/Public/Tools/LegacyEdModeInterfaces.h
--- a/Engine/Source/Editor/UnrealEd/Public/Tools/LegacyEdModeInterfaces.h
+++ b/Engine/Source/Editor/UnrealEd/Public/Tools/LegacyEdModeInterfaces.h
@@ -15,6 +15,8 @@ enum EModeTools : int8;
 class FModeTool;
 struct FBox;
 struct FMatrix;
+class FSceneView;
+class FPrimitiveDrawInterface;
 
 UINTERFACE(NotBlueprintable, MinimalAPI)
 class ULegacyEdModeSelectInterface : public UInterface
diff --git a/Engine/Source/Runtime/Engine/Private/LevelInstance/Packed/PackedLevelInstanceActor.cpp b/Engine/Source/Runtime/Engine/Private/LevelInstance/Packed/PackedLevelInstanceActor.cpp
--- a/Engine/Source/Runtime/Engine/Private/LevelInstance/Packed/PackedLevelInstanceActor.cpp
+++ b/Engine/Source/Runtime/Engine/Private/LevelInstance/Packed/PackedLevelInstanceActor.cpp
@@ -3,6 +3,7 @@
 #include "LevelInstance/Packed/PackedLevelInstanceActor.h"
 #include "LevelInstance/Packed/PackedLevelInstanceBuilder.h"
 #include "LevelInstance/LevelInstanceSubsystem.h"
+#include "Components/InstancedStaticMeshComponent.h"
 
 APackedLevelInstance::APackedLevelInstance(const FObjectInitializer& ObjectInitializer)
 	: Super(ObjectInitializer)
diff --git a/Engine/Source/Runtime/Engine/Private/ShaderCompiler/ShaderGenerationUtil.cpp b/Engine/Source/Runtime/Engine/Private/ShaderCompiler/ShaderGenerationUtil.cpp
--- a/Engine/Source/Runtime/Engine/Private/ShaderCompiler/ShaderGenerationUtil.cpp
+++ b/Engine/Source/Runtime/Engine/Private/ShaderCompiler/ShaderGenerationUtil.cpp
@@ -13,6 +13,8 @@
 #include "Interfaces/ITargetPlatformManagerModule.h"
 #include "RenderUtils.h"
 #include "RHI.h"
+#include "SceneManagement.h"
+#include "Engine/EngineTypes.h"
 
 #include "GBufferInfo.h"
 
diff --git a/Engine/Source/Runtime/Engine/Public/LevelInstance/Packed/LevelInstanceISMPacker.h b/Engine/Source/Runtime/Engine/Public/LevelInstance/Packed/LevelInstanceISMPacker.h
--- a/Engine/Source/Runtime/Engine/Public/LevelInstance/Packed/LevelInstanceISMPacker.h
+++ b/Engine/Source/Runtime/Engine/Public/LevelInstance/Packed/LevelInstanceISMPacker.h
@@ -11,6 +11,7 @@ class AActor;
 class APackedLevelInstance;
 class UActorComponent;
 class FPackedLevelInstanceBuilderContext;
+class UStaticMeshComponent;
 
 class FLevelInstanceISMPacker : public ILevelInstancePacker
 {
diff --git a/Engine/Source/Runtime/Interchange/Engine/Public/InterchangeManager.h b/Engine/Source/Runtime/Interchange/Engine/Public/InterchangeManager.h
--- a/Engine/Source/Runtime/Interchange/Engine/Public/InterchangeManager.h
+++ b/Engine/Source/Runtime/Interchange/Engine/Public/InterchangeManager.h
@@ -18,6 +18,7 @@
 #include "UObject/Object.h"
 #include "UObject/ObjectMacros.h"
 #include "UObject/StrongObjectPtr.h"
+#include "Containers/Queue.h"
 
 #include "InterchangeManager.generated.h"
 
diff --git a/Engine/Source/Runtime/RHI/Public/RHIValidationCommon.h b/Engine/Source/Runtime/RHI/Public/RHIValidationCommon.h
--- a/Engine/Source/Runtime/RHI/Public/RHIValidationCommon.h
+++ b/Engine/Source/Runtime/RHI/Public/RHIValidationCommon.h
@@ -14,6 +14,8 @@ const bool GRHIValidationEnabled = false;
 
 #if ENABLE_RHI_VALIDATION
 
+#include "RHI.h"
+
 class FRHIUniformBuffer;
 
 // Forward declaration of function defined in RHIUtilities.h
diff --git a/Engine/Source/Runtime/RawMesh/Private/RawMesh.cpp b/Engine/Source/Runtime/RawMesh/Private/RawMesh.cpp
--- a/Engine/Source/Runtime/RawMesh/Private/RawMesh.cpp
+++ b/Engine/Source/Runtime/RawMesh/Private/RawMesh.cpp
@@ -6,6 +6,7 @@
 #include "Misc/SecureHash.h"
 #include "Modules/ModuleManager.h"
 #include "Serialization/BulkDataReader.h"
+#include "Misc/ScopeRWLock.h"
 
 IMPLEMENT_MODULE(FDefaultModuleImpl, RawMesh);
 
