GeometryCollection -- report error instead of crashing when converting a non-rigid SkeletalMesh.

diff --git a/Engine/Plugins/Experimental/GeometryCollectionPlugin/Source/GeometryCollectionEditor/Private/GeometryCollection/GeometryCollectionConversion.cpp b/Engine/Plugins/Experimental/GeometryCollectionPlugin/Source/GeometryCollectionEditor/Private/GeometryCollection/GeometryCollectionConversion.cpp
--- a/Engine/Plugins/Experimental/GeometryCollectionPlugin/Source/GeometryCollectionEditor/Private/GeometryCollection/GeometryCollectionConversion.cpp
+++ b/Engine/Plugins/Experimental/GeometryCollectionPlugin/Source/GeometryCollectionEditor/Private/GeometryCollection/GeometryCollectionConversion.cpp
@@ -272,11 +272,23 @@ void FGeometryCollectionConversion::AppendSkeletalMesh(const USkeletalMesh* Skel
 				{
 					if (SkelMeshRenderData->LODRenderData.Num())
 					{
-						const FSkeletalMeshLODRenderData & SkeletalMeshLODRenderData = SkelMeshRenderData->LODRenderData[0];
-						const FSkinWeightVertexBuffer & SkinWeightVertexBuffer = *SkeletalMeshLODRenderData.GetSkinWeightVertexBuffer();
-
-						const FSkelMeshRenderSection & RenderSection = SkeletalMeshLODRenderData.RenderSections[0];
-						const TArray<FBoneIndexType> & SkeletalBoneMap = RenderSection.BoneMap;
+						const FSkeletalMeshLODRenderData& SkeletalMeshLODRenderData = SkelMeshRenderData->LODRenderData[0];
+						const FSkinWeightVertexBuffer& SkinWeightVertexBuffer = *SkeletalMeshLODRenderData.GetSkinWeightVertexBuffer();
+						const FStaticMeshVertexBuffers& VertexBuffers = SkeletalMeshLODRenderData.StaticVertexBuffers;
+						const FPositionVertexBuffer& PositionVertexBuffer = VertexBuffers.PositionVertexBuffer;
+						const int32 VertexCount = PositionVertexBuffer.GetNumVertices();
+						// Check that all vertex weightings are rigid. 
+						for (int32 VertexIndex = 0; VertexIndex < VertexCount; VertexIndex++)
+						{
+							int32 SkeletalBoneIndex = -1;
+							if (!SkinWeightVertexBuffer.GetRigidWeightBone(VertexIndex, SkeletalBoneIndex))
+							{
+								UE_LOG(UGeometryCollectionConversionLogging, Error, TEXT("Non-rigid weighting found on vertex %d: Cannot convert to GeometryCollection."), VertexIndex);
+								return;
+							}
+						}
+						const FSkelMeshRenderSection& RenderSection = SkeletalMeshLODRenderData.RenderSections[0];
+						const TArray<FBoneIndexType>& SkeletalBoneMap = RenderSection.BoneMap;
 
 						//
 						// The Component transform for each Mesh will become the FTransform that drives
@@ -338,16 +350,11 @@ void FGeometryCollectionConversion::AppendSkeletalMesh(const USkeletalMesh* Skel
 						TManagedArray<FLinearColor>& BoneColor = GeometryCollection->BoneColor;
 						TManagedArray<FString>& BoneName = GeometryCollection->BoneName;
 
-						//
-						// Transform Attributes
-						//
+						// 
+						// Transform Attributes 
+						// 
 						TManagedArray<int32>& Parent = GeometryCollection->Parent;
 						TManagedArray<int32>& SimulationType = GeometryCollection->SimulationType;
-
-						const FStaticMeshVertexBuffers & VertexBuffers = SkeletalMeshLODRenderData.StaticVertexBuffers;
-						const FPositionVertexBuffer & PositionVertexBuffer = VertexBuffers.PositionVertexBuffer;
-
-						const int32 VertexCount = PositionVertexBuffer.GetNumVertices();
 						int InitialNumVertices = GeometryCollection->NumElements(FGeometryCollection::VerticesGroup);
 						int VertexBaseIndex = GeometryCollection->AddElements(VertexCount, FGeometryCollection::VerticesGroup);
 						for (int32 VertexIndex = 0; VertexIndex < VertexCount; VertexIndex++)
@@ -355,7 +362,6 @@ void FGeometryCollectionConversion::AppendSkeletalMesh(const USkeletalMesh* Skel
 							int VertexOffset = VertexBaseIndex + VertexIndex;
 							BoneMap[VertexOffset] = -1;
 							int32 SkeletalBoneIndex = -1;
-							check(SkinWeightVertexBuffer.GetRigidWeightBone(VertexIndex, SkeletalBoneIndex));
 							if (SkeletalBoneIndex > -1)
 							{
 								BoneMap[VertexOffset] = SkeletalBoneIndex + TransformBaseIndex;
@@ -365,8 +371,7 @@ void FGeometryCollectionConversion::AppendSkeletalMesh(const USkeletalMesh* Skel
 							TangentU[VertexOffset] = VertexBuffers.StaticMeshVertexBuffer.VertexTangentX(VertexIndex);
 							TangentV[VertexOffset] = VertexBuffers.StaticMeshVertexBuffer.VertexTangentY(VertexIndex);
 							Normal[VertexOffset] = VertexBuffers.StaticMeshVertexBuffer.VertexTangentZ(VertexIndex);
-
-							// @todo : Support multiple UV's per vertex based on MAX_STATIC_TEXCOORDS
+							// @todo : Support multiple UV's per vertex based on MAX_STATIC_TEXCOORDS 
 							UV[VertexOffset] = VertexBuffers.StaticMeshVertexBuffer.GetVertexUV(VertexIndex, 0);
 							if (VertexBuffers.ColorVertexBuffer.GetNumVertices() == VertexCount)
 								Color[VertexOffset] = VertexBuffers.ColorVertexBuffer.VertexColor(VertexIndex);
