GeometryCollectionConversion -- important function call in SkeletalMeshAppend was missing, which lost the ability to evaluate rigidly bound vertices.

diff --git a/Engine/Plugins/Experimental/GeometryCollectionPlugin/Source/GeometryCollectionEditor/Private/GeometryCollection/GeometryCollectionConversion.cpp b/Engine/Plugins/Experimental/GeometryCollectionPlugin/Source/GeometryCollectionEditor/Private/GeometryCollection/GeometryCollectionConversion.cpp
--- a/Engine/Plugins/Experimental/GeometryCollectionPlugin/Source/GeometryCollectionEditor/Private/GeometryCollection/GeometryCollectionConversion.cpp
+++ b/Engine/Plugins/Experimental/GeometryCollectionPlugin/Source/GeometryCollectionEditor/Private/GeometryCollection/GeometryCollectionConversion.cpp
@@ -362,6 +362,7 @@ void FGeometryCollectionConversion::AppendSkeletalMesh(const USkeletalMesh* Skel
 							int VertexOffset = VertexBaseIndex + VertexIndex;
 							BoneMap[VertexOffset] = -1;
 							int32 SkeletalBoneIndex = -1;
+							SkinWeightVertexBuffer.GetRigidWeightBone(VertexIndex, SkeletalBoneIndex);
 							if (SkeletalBoneIndex > -1)
 							{
 								BoneMap[VertexOffset] = SkeletalBoneIndex + TransformBaseIndex;
