Fracture Tool -- Merging Geometry Collections and baking in scale is now possible via the New button in the Fracture Tool.

diff --git a/Engine/Plugins/Experimental/ChaosEditor/Source/FractureEditor/Private/FractureToolEmbed.cpp b/Engine/Plugins/Experimental/ChaosEditor/Source/FractureEditor/Private/FractureToolEmbed.cpp
--- a/Engine/Plugins/Experimental/ChaosEditor/Source/FractureEditor/Private/FractureToolEmbed.cpp
+++ b/Engine/Plugins/Experimental/ChaosEditor/Source/FractureEditor/Private/FractureToolEmbed.cpp
@@ -90,16 +90,15 @@ void UFractureToolAddEmbeddedGeometry::Execute(TWeakPtr<FFractureEditorModeToolk
 				}
 
 				FTransform ComponentTransform = SMActorTransform.GetRelativeTransform(TargetActorTransform);
-				decltype(FGeometryCollectionSource::SourceMaterial) SourceMaterials(SelectedStaticMeshComponent->GetMaterials());
-
+				
 				// Add static mesh geometry to context's geometry collection.
-				FGeometryCollectionConversion::AppendStaticMesh(ComponentStaticMesh, SourceMaterials, ComponentTransform, FracturedGeometryCollection, true);
+				FGeometryCollectionConversion::AppendStaticMesh(ComponentStaticMesh, SelectedStaticMeshComponent, ComponentTransform, FracturedGeometryCollection, true);
 
 				int32 NewTransformIndex = Parent.Num() - 1;
 				check(NewTransformIndex > -1)
 
-					// Reparent the new transform to the selected bone
-					Parent[NewTransformIndex] = SelectedBone;
+				// Reparent the new transform to the selected bone
+				Parent[NewTransformIndex] = SelectedBone;
 				Children[SelectedBone].Add(NewTransformIndex);
 
 				// Set properties appropriate to embedded geometry
diff --git a/Engine/Plugins/Experimental/ChaosEditor/Source/FractureEditor/Private/FractureToolGenerators.cpp b/Engine/Plugins/Experimental/ChaosEditor/Source/FractureEditor/Private/FractureToolGenerators.cpp
--- a/Engine/Plugins/Experimental/ChaosEditor/Source/FractureEditor/Private/FractureToolGenerators.cpp
+++ b/Engine/Plugins/Experimental/ChaosEditor/Source/FractureEditor/Private/FractureToolGenerators.cpp
@@ -11,6 +11,7 @@
 #include "Layers/LayersSubsystem.h"
 #include "Engine/SkeletalMesh.h"
 #include "Engine/StaticMesh.h"
+#include "Engine/StaticMeshActor.h"
 
 #include "GeometryCollection/GeometryCollectionComponent.h"
 #include "GeometryCollection/GeometryCollectionActor.h"
@@ -46,8 +47,7 @@ void UFractureToolGenerateAsset::RegisterUICommand(FFractureEditorCommands* Bind
 
 bool UFractureToolGenerateAsset::CanExecute() const
 {
-	// We can execute this command if only static meshes are selected.
-	return (IsStaticMeshSelected() && (!IsGeometryCollectionSelected()));
+	return (IsStaticMeshSelected() || IsGeometryCollectionSelected());
 }
 
 void UFractureToolGenerateAsset::Execute(TWeakPtr<FFractureEditorModeToolkit> InToolkit)
@@ -116,7 +116,8 @@ void UFractureToolGenerateAsset::OnGenerateAssetPathChosen(const FString& InAsse
 		AActor* FirstActor = Actors[0];
 
 		AGeometryCollectionActor* GeometryCollectionActor = Cast<AGeometryCollectionActor>(FirstActor);
-		GeometryCollectionActor = ConvertStaticMeshToGeometryCollection(InAssetPath, Actors);
+		
+		GeometryCollectionActor = ConvertActorsToGeometryCollection(InAssetPath, Actors);
 
 		GeometryCollectionComponent = GeometryCollectionActor->GetGeometryCollectionComponent();
 
@@ -140,6 +141,11 @@ void UFractureToolGenerateAsset::OnGenerateAssetPathChosen(const FString& InAsse
 			SharedToolkit->SetBoneSelection(GeometryCollectionComponent, EditBoneColor.GetSelectedBones(), true);
 
 			SharedToolkit->OnSetLevelViewValue(-1);
+
+			SharedToolkit->RegenerateOutliner();
+			SharedToolkit->RegenerateHistogram();
+
+			SharedToolkit->UpdateExplodedVectors(GeometryCollectionComponent);
 		}
 		
 		GeometryCollectionComponent->MarkRenderDynamicDataDirty();
@@ -152,7 +158,7 @@ void UFractureToolGenerateAsset::OnGenerateAssetPathChosen(const FString& InAsse
 	}
 }
 
-AGeometryCollectionActor* UFractureToolGenerateAsset::ConvertStaticMeshToGeometryCollection(const FString& InAssetPath, TArray<AActor*>& Actors)
+AGeometryCollectionActor* UFractureToolGenerateAsset::ConvertActorsToGeometryCollection(const FString& InAssetPath, TArray<AActor*>& Actors)
 {
 	ensure(Actors.Num() > 0);
 	AActor* FirstActor = Actors[0];
@@ -172,36 +178,76 @@ AGeometryCollectionActor* UFractureToolGenerateAsset::ConvertStaticMeshToGeometr
 
 		check(FracturedGeometryCollection);
 
-		TArray<UStaticMeshComponent*> StaticMeshComponents;
-		Actor->GetComponents<UStaticMeshComponent>(StaticMeshComponents, true);
-		for (int32 ii = 0, ni = StaticMeshComponents.Num(); ii < ni; ++ii)
+		if (const AStaticMeshActor* StaticMeshActor = Cast<AStaticMeshActor>(Actor))
 		{
-			// We're partial to static mesh components, here
-			UStaticMeshComponent* StaticMeshComponent = StaticMeshComponents[ii];
-			if (StaticMeshComponent != nullptr)
+			TArray<UStaticMeshComponent*> StaticMeshComponents;
+			Actor->GetComponents<UStaticMeshComponent>(StaticMeshComponents, true);
+			for (int32 ii = 0, ni = StaticMeshComponents.Num(); ii < ni; ++ii)
 			{
-				UStaticMesh* ComponentStaticMesh = StaticMeshComponent->GetStaticMesh();
-				if (ComponentStaticMesh != nullptr)
+				// We're partial to static mesh components, here
+				UStaticMeshComponent* StaticMeshComponent = StaticMeshComponents[ii];
+				if (StaticMeshComponent != nullptr)
 				{
-					// If any of the static meshes have Nanite enabled, also enable on the new geometry collection asset for convenience.
-					FracturedGeometryCollection->EnableNanite |= ComponentStaticMesh->NaniteSettings.bEnabled;
-				}
+					UStaticMesh* ComponentStaticMesh = StaticMeshComponent->GetStaticMesh();
+					if (ComponentStaticMesh != nullptr)
+					{
+						// If any of the static meshes have Nanite enabled, also enable on the new geometry collection asset for convenience.
+						FracturedGeometryCollection->EnableNanite |= ComponentStaticMesh->NaniteSettings.bEnabled;
+					}
 
-				FTransform ComponentTransform(StaticMeshComponent->GetComponentTransform());
-				ComponentTransform.SetTranslation((ComponentTransform.GetTranslation() - ActorTransform.GetTranslation()) + ActorOffset);
+					FTransform ComponentTransform(StaticMeshComponent->GetComponentTransform());
+					ComponentTransform.SetTranslation((ComponentTransform.GetTranslation() - ActorTransform.GetTranslation()) + ActorOffset);
 
-				// Record the contributing source on the asset.
-				FSoftObjectPath SourceSoftObjectPath(ComponentStaticMesh);
-				decltype(FGeometryCollectionSource::SourceMaterial) SourceMaterials(StaticMeshComponent->GetMaterials());
-				FracturedGeometryCollection->GeometrySource.Add({ SourceSoftObjectPath, ComponentTransform, SourceMaterials });
+					// Record the contributing source on the asset.
+					FSoftObjectPath SourceSoftObjectPath(ComponentStaticMesh);
+					decltype(FGeometryCollectionSource::SourceMaterial) SourceMaterials(StaticMeshComponent->GetMaterials());
+					FracturedGeometryCollection->GeometrySource.Add({ SourceSoftObjectPath, ComponentTransform, SourceMaterials });
 
-				FGeometryCollectionConversion::AppendStaticMesh(ComponentStaticMesh, SourceMaterials, ComponentTransform, FracturedGeometryCollection, true);
+					FGeometryCollectionConversion::AppendStaticMesh(ComponentStaticMesh, SourceMaterials, ComponentTransform, FracturedGeometryCollection, true);
+				}
 			}
 		}
+		else if (const AGeometryCollectionActor* GeometryActor = Cast<AGeometryCollectionActor>(Actor))
+		{
+			TArray<UGeometryCollectionComponent*> GeometryCollectionComponents;
+			Actor->GetComponents<UGeometryCollectionComponent>(GeometryCollectionComponents, true);
+			for (int32 ii = 0, ni = GeometryCollectionComponents.Num(); ii < ni; ++ii)
+			{
+				UGeometryCollectionComponent* GeometryCollectionComponent = GeometryCollectionComponents[ii];
+				if (GeometryCollectionComponent != nullptr)
+				{
+					const UGeometryCollection* RestCollection = GeometryCollectionComponent->GetRestCollection();
+					if (RestCollection != nullptr)
+					{
+						// If any of the static meshes have Nanite enabled, also enable on the new geometry collection asset for convenience.
+						FracturedGeometryCollection->EnableNanite |= RestCollection->EnableNanite;
+					}
 
-		FracturedGeometryCollection->InitializeMaterials();
+					FTransform ComponentTransform(GeometryCollectionComponent->GetComponentTransform());
+					ComponentTransform.SetTranslation((ComponentTransform.GetTranslation() - ActorTransform.GetTranslation()) + ActorOffset);
+
+					// Record the contributing source on the asset.
+					FSoftObjectPath SourceSoftObjectPath(RestCollection);
+
+					// We're not interested in recording the final material of the collection since it's inevitably the Selection material.
+					int32 NumMaterials = GeometryCollectionComponent->GetNumMaterials() - 1;
+					TArray<TObjectPtr<UMaterialInterface>> SourceMaterials;
+					SourceMaterials.SetNum(NumMaterials);
+					for (int32 MaterialIndex = 0; MaterialIndex < NumMaterials; ++MaterialIndex)
+					{
+						SourceMaterials[MaterialIndex] = GeometryCollectionComponent->GetMaterial(MaterialIndex);
+					}
+					FracturedGeometryCollection->GeometrySource.Add({ SourceSoftObjectPath, ComponentTransform, SourceMaterials });
+
+					FGeometryCollectionConversion::AppendGeometryCollection(RestCollection, GeometryCollectionComponent, ComponentTransform, FracturedGeometryCollection, true);
+
+				}
+			}
+		}	
 	}
 
+	FracturedGeometryCollection->InitializeMaterials();
+
 	AddSingleRootNodeIfRequired(FracturedGeometryCollection);
 
 	if (FracturedGeometryCollection->EnableNanite)
@@ -356,6 +402,10 @@ void UFractureToolResetAsset::Execute(TWeakPtr<FFractureEditorModeToolkit> InToo
 						// #todo (bmiller) Once we've settled on the right approach with static meshes, we'll need to apply the same strategy to skeletal mesh reconstruction.
 						// FGeometryCollectionConversion::AppendSkeletalMesh(SourceSkeletalMesh, Source.SourceMaterial, Source.LocalTransform, GeometryCollectionObject, true);
 					}
+					else if (const UGeometryCollection* SourceGeometryCollection = Cast<UGeometryCollection>(SourceMesh))
+					{
+						FGeometryCollectionConversion::AppendGeometryCollection(SourceGeometryCollection, Source.SourceMaterial, Source.LocalTransform, GeometryCollectionObject, true);
+					}
 				}
 
 				GeometryCollectionObject->InitializeMaterials();
diff --git a/Engine/Plugins/Experimental/ChaosEditor/Source/FractureEditor/Private/FractureToolGenerators.h b/Engine/Plugins/Experimental/ChaosEditor/Source/FractureEditor/Private/FractureToolGenerators.h
--- a/Engine/Plugins/Experimental/ChaosEditor/Source/FractureEditor/Private/FractureToolGenerators.h
+++ b/Engine/Plugins/Experimental/ChaosEditor/Source/FractureEditor/Private/FractureToolGenerators.h
@@ -31,7 +31,7 @@ public:
 private:
 	void OpenGenerateAssetDialog(TArray<AActor*>& Actors);
 	void OnGenerateAssetPathChosen(const FString& InAssetPath, TArray<AActor*> Actors);
-	AGeometryCollectionActor* ConvertStaticMeshToGeometryCollection(const FString& InAssetPath, TArray<AActor*>& Actors);
+	AGeometryCollectionActor* ConvertActorsToGeometryCollection(const FString& InAssetPath, TArray<AActor*>& Actors);
 	class AGeometryCollectionActor* CreateNewGeometryActor(const FString& Name, const FTransform& Transform, bool AddMaterials /*= false*/);
 	static ULevel* GetSelectedLevel();
 	static AActor* AddActor(ULevel* InLevel, UClass* Class);
diff --git a/Engine/Plugins/Experimental/GeometryCollectionPlugin/Source/GeometryCollectionEditor/Private/GeometryCollection/GeometryCollectionConversion.cpp b/Engine/Plugins/Experimental/GeometryCollectionPlugin/Source/GeometryCollectionEditor/Private/GeometryCollection/GeometryCollectionConversion.cpp
--- a/Engine/Plugins/Experimental/GeometryCollectionPlugin/Source/GeometryCollectionEditor/Private/GeometryCollection/GeometryCollectionConversion.cpp
+++ b/Engine/Plugins/Experimental/GeometryCollectionPlugin/Source/GeometryCollectionEditor/Private/GeometryCollection/GeometryCollectionConversion.cpp
@@ -199,6 +199,8 @@ void FGeometryCollectionConversion::AppendStaticMesh(const UStaticMesh* StaticMe
 				CurrMaterial = UMaterial::GetDefaultMaterial(MD_Surface);
 			}
 
+			// We add the material twice, once for interior and again for exterior.
+			GeometryCollectionObject->Materials.Add(CurrMaterial);
 			GeometryCollectionObject->Materials.Add(CurrMaterial);
 		}
 
@@ -236,6 +238,251 @@ void FGeometryCollectionConversion::AppendStaticMesh(const UStaticMesh* StaticMe
 }
 
 
+void FGeometryCollectionConversion::AppendGeometryCollection(const UGeometryCollection* SourceGeometryCollection, const TArray<UMaterialInterface*>& Materials, const FTransform& GeometryCollectionTransform, UGeometryCollection* TargetGeometryCollectionObject, bool ReindexMaterials)
+{
+	if (SourceGeometryCollection == nullptr)
+	{
+		return;
+	}
+	const TSharedPtr<FGeometryCollection, ESPMode::ThreadSafe> SourceGeometryCollectionPtr = SourceGeometryCollection->GetGeometryCollection();
+
+	check(TargetGeometryCollectionObject);
+	TSharedPtr<FGeometryCollection, ESPMode::ThreadSafe> GeometryCollectionPtr = TargetGeometryCollectionObject->GetGeometryCollection();
+	FGeometryCollection* GeometryCollection = GeometryCollectionPtr.Get();
+	check(GeometryCollection);
+
+	// Assemble offsets and add elements
+	const int32 VertexCount = SourceGeometryCollectionPtr->Vertex.Num();
+	const int32 FaceCount = SourceGeometryCollectionPtr->Indices.Num();
+	const int32 TransformCount = SourceGeometryCollectionPtr->Transform.Num();
+	const int32 GeometryCount = SourceGeometryCollectionPtr->TransformIndex.Num();
+	const int32 SectionCount = SourceGeometryCollectionPtr->Sections.Num();
+
+	const int32 AssetMaterialStart = TargetGeometryCollectionObject->Materials.Num();
+
+	FVector Scale = GeometryCollectionTransform.GetScale3D();
+	FTransform AppliedTransform = GeometryCollectionTransform;
+	AppliedTransform.RemoveScaling();
+
+	const int32 VertexStart = GeometryCollection->AddElements(VertexCount, FGeometryCollection::VerticesGroup);
+	const int32 FaceStart = GeometryCollection->AddElements(FaceCount, FGeometryCollection::FacesGroup);
+	const int32 TransformStart = GeometryCollection->AddElements(TransformCount, FGeometryCollection::TransformGroup);
+	const int32 GeometryStart = GeometryCollection->AddElements(GeometryCount, FGeometryCollection::GeometryGroup);
+	const int32 SectionStart = GeometryCollection->AddElements(SectionCount, FGeometryCollection::MaterialGroup);
+
+	// source vertex information
+	const TManagedArray<FVector>& SourceVertex = SourceGeometryCollectionPtr->Vertex;
+	const TManagedArray<FVector>& SourceTangentU = SourceGeometryCollectionPtr->TangentU;
+	const TManagedArray<FVector>& SourceTangentV = SourceGeometryCollectionPtr->TangentV;
+	const TManagedArray<FVector>& SourceNormal = SourceGeometryCollectionPtr->Normal;
+	const TManagedArray<FVector2D>& SourceUV = SourceGeometryCollectionPtr->UV;
+	const TManagedArray<FLinearColor>& SourceColor = SourceGeometryCollectionPtr->Color;
+	const TManagedArray<int32>& SourceBoneMap = SourceGeometryCollectionPtr->BoneMap;
+
+	// target vertex information
+	TManagedArray<FVector>& TargetVertex = GeometryCollection->Vertex;
+	TManagedArray<FVector>& TargetTangentU = GeometryCollection->TangentU;
+	TManagedArray<FVector>& TargetTangentV = GeometryCollection->TangentV;
+	TManagedArray<FVector>& TargetNormal = GeometryCollection->Normal;
+	TManagedArray<FVector2D>& TargetUV = GeometryCollection->UV;
+	TManagedArray<FLinearColor>& TargetColor = GeometryCollection->Color;
+	TManagedArray<int32>& TargetBoneMap = GeometryCollection->BoneMap;
+
+	// append vertices
+	for (int32 VertexIndex = 0; VertexIndex < VertexCount; VertexIndex++)
+	{
+		const int32 VertexOffset = VertexStart + VertexIndex;
+		TargetVertex[VertexOffset] = SourceVertex[VertexIndex] * Scale;
+		
+		TargetTangentU[VertexOffset] = SourceTangentU[VertexIndex];
+		TargetTangentV[VertexOffset] = SourceTangentV[VertexIndex];
+		TargetNormal[VertexOffset] = SourceNormal[VertexIndex];
+		TargetUV[VertexOffset] = SourceUV[VertexIndex];
+		TargetColor[VertexOffset] = SourceColor[VertexIndex];
+
+		TargetBoneMap[VertexOffset] = SourceBoneMap[VertexIndex] + TransformStart;
+	}
+
+	// source face information
+	const TManagedArray<FIntVector>& SourceIndices = SourceGeometryCollectionPtr->Indices;
+	const TManagedArray<bool>& SourceVisible = SourceGeometryCollectionPtr->Visible;
+	const TManagedArray<int32>& SourceMaterialID = SourceGeometryCollectionPtr->MaterialID;
+	const TManagedArray<int32>& SourceMaterialIndex = SourceGeometryCollectionPtr->MaterialIndex;
+	
+	// target face information
+	TManagedArray<FIntVector>& TargetIndices = GeometryCollection->Indices;
+	TManagedArray<bool>& TargetVisible = GeometryCollection->Visible;
+	TManagedArray<int32>& TargetMaterialID = GeometryCollection->MaterialID;
+	TManagedArray<int32>& TargetMaterialIndex = GeometryCollection->MaterialIndex;
+
+	// append faces
+	for (int32 FaceIndex = 0; FaceIndex < FaceCount; ++FaceIndex)
+	{
+		const FIntVector& SourceFace = SourceIndices[FaceIndex];
+		const int32 FaceOffset = FaceStart + FaceIndex;
+		TargetIndices[FaceOffset] = FIntVector(
+			SourceFace[0] + VertexStart,
+			SourceFace[1] + VertexStart,
+			SourceFace[2] + VertexStart);
+		TargetVisible[FaceOffset] = SourceVisible[FaceIndex];
+		
+		TargetMaterialID[FaceOffset] = AssetMaterialStart + SourceMaterialID[FaceIndex];
+		TargetMaterialIndex[FaceOffset] = FaceOffset;
+	}
+
+	// source transform information
+	const TManagedArray<FTransform>& SourceTransform = SourceGeometryCollectionPtr->Transform;
+	const TManagedArray<FString>& SourceBoneName = SourceGeometryCollectionPtr->BoneName;
+	const TManagedArray<FLinearColor>& SourceBoneColor = SourceGeometryCollectionPtr->BoneColor;
+	const TManagedArray<int32>& SourceParent = SourceGeometryCollectionPtr->Parent;
+	const TManagedArray<TSet<int32>>& SourceChildren = SourceGeometryCollectionPtr->Children;
+	const TManagedArray<int32>& SourceTransformToGeometryIndex = SourceGeometryCollectionPtr->TransformToGeometryIndex;
+	const TManagedArray<int32>& SourceSimulationType = SourceGeometryCollectionPtr->SimulationType;
+	const TManagedArray<int32>& SourceStatusFlags = SourceGeometryCollectionPtr->StatusFlags;
+	const TManagedArray<int32>& SourceInitialDynamicState = SourceGeometryCollectionPtr->InitialDynamicState;
+	
+	// target transform information
+	TManagedArray<FTransform>& TargetTransform = GeometryCollection->Transform;
+	TManagedArray<FString>& TargetBoneName = GeometryCollection->BoneName;
+	TManagedArray<FLinearColor>& TargetBoneColor = GeometryCollection->BoneColor;
+	TManagedArray<int32>& TargetParent = GeometryCollection->Parent;
+	TManagedArray<TSet<int32>>& TargetChildren = GeometryCollection->Children;
+	TManagedArray<int32>& TargetTransformToGeometryIndex = GeometryCollection->TransformToGeometryIndex;
+	TManagedArray<int32>& TargetSimulationType = GeometryCollection->SimulationType;
+	TManagedArray<int32>& TargetStatusFlags = GeometryCollection->StatusFlags;
+	TManagedArray<int32>& TargetInitialDynamicState = GeometryCollection->InitialDynamicState;
+
+	// append transform hierarchy
+	for (int32 TransformIndex = 0; TransformIndex < TransformCount; ++TransformIndex)
+	{
+		const int32 TransformOffset = TransformStart + TransformIndex;
+
+		// Only apply the transform to the parent node. Child nodes only need scaling applied to translation offsets.
+		if (SourceParent[TransformIndex] == INDEX_NONE)
+		{
+			TargetTransform[TransformOffset] = SourceTransform[TransformIndex] * AppliedTransform;
+		}
+		else
+		{
+			FTransform ScaledTranslation = SourceTransform[TransformIndex];
+			ScaledTranslation.ScaleTranslation(Scale);
+			TargetTransform[TransformOffset] = ScaledTranslation;
+		}
+	
+		// #todo Get this Bone name to be unique
+		TargetBoneName[TransformOffset] = SourceBoneName[TransformIndex];
+
+		const FColor RandBoneColor(FMath::Rand() % 100 + 5, FMath::Rand() % 100 + 5, FMath::Rand() % 100 + 5, 255);
+		TargetBoneColor[TransformOffset] = FLinearColor(RandBoneColor);
+
+		TargetParent[TransformOffset] = (SourceParent[TransformIndex]==INDEX_NONE) ? INDEX_NONE: SourceParent[TransformIndex] + TransformStart;
+
+		const TSet<int32>& SourceChildrenSet = SourceChildren[TransformIndex];
+		for (int32 ChildIndex : SourceChildrenSet)
+		{
+			TargetChildren[TransformOffset].Add(ChildIndex + TransformStart);
+		}
+
+		TargetTransformToGeometryIndex[TransformOffset] = SourceTransformToGeometryIndex[TransformIndex] + GeometryStart;
+		TargetSimulationType[TransformOffset] = SourceSimulationType[TransformIndex];
+		TargetStatusFlags[TransformOffset] = SourceStatusFlags[TransformIndex];
+		TargetInitialDynamicState[TransformOffset] = SourceInitialDynamicState[TransformIndex];
+	}
+
+	// source geometry information
+	const TManagedArray<int32>& SourceTransformIndex = SourceGeometryCollectionPtr->TransformIndex;
+	const TManagedArray<int32>& SourceVertexStart = SourceGeometryCollectionPtr->VertexStart;
+	const TManagedArray<int32>& SourceVertexCount = SourceGeometryCollectionPtr->VertexCount;
+	const TManagedArray<int32>& SourceFaceStart = SourceGeometryCollectionPtr->FaceStart;
+	const TManagedArray<int32>& SourceFaceCount = SourceGeometryCollectionPtr->FaceCount;
+	
+	// target geometry information
+	TManagedArray<int32>& TargetTransformIndex = GeometryCollection->TransformIndex;
+	TManagedArray<FBox>& TargetBoundingBox = GeometryCollection->BoundingBox;
+	TManagedArray<float>& TargetInnerRadius = GeometryCollection->InnerRadius;
+	TManagedArray<float>& TargetOuterRadius = GeometryCollection->OuterRadius;
+	TManagedArray<int32>& TargetVertexStart = GeometryCollection->VertexStart;
+	TManagedArray<int32>& TargetVertexCount = GeometryCollection->VertexCount;
+	TManagedArray<int32>& TargetFaceStart = GeometryCollection->FaceStart;
+	TManagedArray<int32>& TargetFaceCount = GeometryCollection->FaceCount;
+
+	// append geometry
+	for (int32 GeometryIndex = 0; GeometryIndex < GeometryCount; ++GeometryIndex)
+	{
+		const int32 GeometryOffset = GeometryStart + GeometryIndex;
+
+		TargetTransformIndex[GeometryOffset] = SourceTransformIndex[GeometryIndex] + TransformStart;
+		
+		TargetVertexStart[GeometryOffset] = SourceVertexStart[GeometryIndex] + VertexStart;
+		TargetVertexCount[GeometryOffset] = SourceVertexCount[GeometryIndex];
+		TargetFaceStart[GeometryOffset] = SourceFaceStart[GeometryIndex] + FaceStart;
+		TargetFaceCount[GeometryOffset] = SourceFaceCount[GeometryIndex];
+
+		// Find centroid of geometry for inner/outer radius calculations
+		FVector Center(0);
+		for (int32 VertexIndex = TargetVertexStart[GeometryOffset]; VertexIndex < TargetVertexStart[GeometryOffset] + TargetVertexCount[GeometryOffset]; ++VertexIndex)
+		{
+			Center += TargetVertex[VertexIndex];
+		}
+		if (TargetVertexCount[GeometryOffset]) Center /= TargetVertexCount[GeometryOffset];
+		
+		TargetBoundingBox[GeometryOffset] = FBox(ForceInitToZero);
+		TargetInnerRadius[GeometryOffset] = FLT_MAX;
+		TargetOuterRadius[GeometryOffset] = -FLT_MAX;
+		for (int32 VertexIndex = TargetVertexStart[GeometryOffset]; VertexIndex < TargetVertexStart[GeometryOffset] + TargetVertexCount[GeometryOffset]; ++VertexIndex)
+		{
+			TargetBoundingBox[GeometryOffset] += TargetVertex[VertexIndex];
+
+			float Delta = (Center - TargetVertex[VertexIndex]).Size();
+			TargetInnerRadius[GeometryOffset] = FMath::Min(TargetInnerRadius[GeometryOffset], Delta);
+			TargetOuterRadius[GeometryOffset] = FMath::Max(TargetOuterRadius[GeometryOffset], Delta);
+		}
+	}
+
+	// for each material, add a reference in our GeometryCollectionObject
+	const int32 NumMeshMaterials = Materials.Num();
+	TargetGeometryCollectionObject->Materials.Reserve(AssetMaterialStart + NumMeshMaterials);
+
+	for (int32 Index = 0; Index < NumMeshMaterials; ++Index)
+	{
+		UMaterialInterface* CurrMaterial = Materials[Index];
+
+		// Possible we have a null entry - replace with default
+		if (CurrMaterial == nullptr)
+		{
+			CurrMaterial = UMaterial::GetDefaultMaterial(MD_Surface);
+		}
+
+		TargetGeometryCollectionObject->Materials.Add(CurrMaterial);
+	}
+
+	// source material information
+	const TManagedArray<FGeometryCollectionSection>& SourceSections = SourceGeometryCollectionPtr->Sections;
+
+	// target material information
+	TManagedArray<FGeometryCollectionSection>& TargetSections = GeometryCollection->Sections;
+
+	// append sections
+	for (int32 SectionIndex = 0; SectionIndex < SectionCount; ++SectionIndex)
+	{
+		int32 SectionOffset = SectionStart + SectionIndex;
+
+		TargetSections[SectionOffset].MaterialID = AssetMaterialStart + SourceSections[SectionIndex].MaterialID;
+
+		TargetSections[SectionOffset].FirstIndex = SourceSections[SectionIndex].FirstIndex + FaceStart * 3;
+		TargetSections[SectionOffset].MinVertexIndex = VertexStart + SourceSections[SectionIndex].MinVertexIndex;
+
+		TargetSections[SectionOffset].NumTriangles = SourceSections[SectionIndex].NumTriangles;
+		TargetSections[SectionOffset].MaxVertexIndex = VertexStart + SourceSections[SectionIndex].MaxVertexIndex;
+	}
+
+	if (ReindexMaterials) {
+		GeometryCollection->ReindexMaterials();
+	}
+
+}
+
+
 void FGeometryCollectionConversion::AppendStaticMesh(const UStaticMesh* StaticMesh, const UStaticMeshComponent* StaticMeshComponent, const FTransform& StaticMeshTransform, UGeometryCollection* GeometryCollectionObject, bool ReindexMaterials)
 {
 	if (StaticMesh == nullptr)
@@ -252,10 +499,40 @@ void FGeometryCollectionConversion::AppendStaticMesh(const UStaticMesh* StaticMe
 		Materials.Add(CurrMaterial);
 	}
 
+	// Geometry collections usually carry the selection material, which we'll delete before appending
+	UMaterialInterface* BoneSelectedMaterial = LoadObject<UMaterialInterface>(nullptr, UGeometryCollection::GetSelectedMaterialPath(), nullptr, LOAD_None, nullptr);
+	GeometryCollectionObject->Materials.Remove(BoneSelectedMaterial);
+	Materials.Remove(BoneSelectedMaterial);
+
 	AppendStaticMesh(StaticMesh, Materials, StaticMeshTransform, GeometryCollectionObject, ReindexMaterials);
 }
 
 
+void FGeometryCollectionConversion::AppendGeometryCollection(const UGeometryCollection* SourceGeometryCollection, const UGeometryCollectionComponent* GeometryCollectionComponent, const FTransform& GeometryCollectionTransform, UGeometryCollection* TargetGeometryCollectionObject, bool ReindexMaterials)
+{
+	if (SourceGeometryCollection == nullptr)
+	{
+		return;
+	}
+
+	TArray<UMaterialInterface*> Materials;
+	Materials.Reserve(SourceGeometryCollection->Materials.Num());
+
+	for (int32 Index = 0; Index < SourceGeometryCollection->Materials.Num(); ++Index)
+	{
+		UMaterialInterface* CurrMaterial = GeometryCollectionComponent ? GeometryCollectionComponent->GetMaterial(Index) : SourceGeometryCollection->Materials[Index].Get();
+		Materials.Add(CurrMaterial);
+	}
+
+	// Geometry collections usually carry the selection material, which we'll delete before appending
+	UMaterialInterface* BoneSelectedMaterial = LoadObject<UMaterialInterface>(nullptr, UGeometryCollection::GetSelectedMaterialPath(), nullptr, LOAD_None, nullptr);
+	TargetGeometryCollectionObject->Materials.Remove(BoneSelectedMaterial);
+	Materials.Remove(BoneSelectedMaterial);
+
+	AppendGeometryCollection(SourceGeometryCollection, Materials, GeometryCollectionTransform, TargetGeometryCollectionObject, ReindexMaterials);
+}
+
+
 void FGeometryCollectionConversion::AppendSkeletalMesh(const USkeletalMesh* SkeletalMesh, const USkeletalMeshComponent *SkeletalMeshComponent, const FTransform & SkeletalMeshTransform, UGeometryCollection * GeometryCollectionObject, bool ReindexMaterials)
 {
 	//UE_LOG(UGeometryCollectionConversionLogging, Log, TEXT("FGeometryCollectionConversion::AppendSkeletalMesh()"));
diff --git a/Engine/Plugins/Experimental/GeometryCollectionPlugin/Source/GeometryCollectionEditor/Public/GeometryCollection/GeometryCollectionConversion.h b/Engine/Plugins/Experimental/GeometryCollectionPlugin/Source/GeometryCollectionEditor/Public/GeometryCollection/GeometryCollectionConversion.h
--- a/Engine/Plugins/Experimental/GeometryCollectionPlugin/Source/GeometryCollectionEditor/Public/GeometryCollection/GeometryCollectionConversion.h
+++ b/Engine/Plugins/Experimental/GeometryCollectionPlugin/Source/GeometryCollectionEditor/Public/GeometryCollection/GeometryCollectionConversion.h
@@ -7,6 +7,7 @@ class UStaticMesh;
 class USkeletalMesh;
 class UGeometryCollection;
 class UMaterialInterface;
+class UGeometryCollectionComponent;
 
 typedef TTuple<const UStaticMesh *, const UStaticMeshComponent *, FTransform> GeometryCollectionStaticMeshConversionTuple;
 typedef TTuple<const USkeletalMesh *, const USkeletalMeshComponent *, FTransform> GeometryCollectionSkeletalMeshConversionTuple;
@@ -43,6 +44,24 @@ public:
 	*/
 	static void AppendSkeletalMesh(const USkeletalMesh* SkeletalMesh, const USkeletalMeshComponent *SkeletalMeshComponent, const FTransform & SkeletalMeshTransform, UGeometryCollection * GeometryCollection, bool ReindexMaterials = true);
 
+	/**
+	*  Appends a GeometryCollection to a GeometryCollectionComponent.
+	*  @param SourceGeometryCollection : Const GeometryCollection to read vertex/normals/index data from
+	*  @param Materials : Materials fetched from the GeometryCollectionComponent used to configure this geometry
+	*  @param GeometryCollectionTransform : GeometryCollection transform.
+	*  @param TargetGeometryCollection  : Collection to append the GeometryCollection into.
+	*/
+	static void AppendGeometryCollection(const UGeometryCollection* SourceGeometryCollection, const TArray<UMaterialInterface*>& Materials, const FTransform& GeometryCollectionTransform, UGeometryCollection* TargetGeometryCollectionObject, bool ReindexMaterials = true);
+
+	/**
+	*  Appends a GeometryCollection to a GeometryCollectionComponent.
+	*  @param GeometryCollectionComponent : Const GeometryCollection to read vertex/normals/index data from
+	*  @param GeometryCollectionTransform : GeometryCollection transform.
+	*  @param TargetGeometryCollection  : Collection to append the GeometryCollection into.
+	*/
+	static void AppendGeometryCollection(const UGeometryCollection* SourceGeometryCollection, const UGeometryCollectionComponent* GeometryCollectionComponent, const FTransform& GeometryCollectionTransform, UGeometryCollection* TargetGeometryCollectionObject, bool ReindexMaterials = true);
+
+
 	/**
 	*  Command invoked from "GeometryCollection.CreatGeometryCollection", uses the selected Actors to create a GeometryCollection Asset
 	*  @param World
diff --git a/Engine/Source/Runtime/Experimental/GeometryCollectionEngine/Private/GeometryCollection/GeometryCollectionObject.cpp b/Engine/Source/Runtime/Experimental/GeometryCollectionEngine/Private/GeometryCollection/GeometryCollectionObject.cpp
--- a/Engine/Source/Runtime/Experimental/GeometryCollectionEngine/Private/GeometryCollection/GeometryCollectionObject.cpp
+++ b/Engine/Source/Runtime/Experimental/GeometryCollectionEngine/Private/GeometryCollection/GeometryCollectionObject.cpp
@@ -238,27 +238,43 @@ void UGeometryCollection::InitializeMaterials()
 {
 	Modify();
 
-	// Consolidate materials
+	// Last Material is the selection one
+	UMaterialInterface* BoneSelectedMaterial = LoadObject<UMaterialInterface>(nullptr, GetSelectedMaterialPath(), nullptr, LOAD_None, nullptr);
 
-	// add all materials to a set
-	TSet<UMaterialInterface*> MaterialSet;
-	for (UMaterialInterface* Curr : Materials)
+	// Skip selection materials
+	Materials.Remove(BoneSelectedMaterial);
+	
+	// We're assuming that all materials are arranged in pairs, so first we collect these.
+	using FMaterialPair = TPair<UMaterialInterface*, UMaterialInterface*>;
+	TSet<FMaterialPair> MaterialSet;
+	for (int32 MaterialIndex = 0; MaterialIndex < Materials.Num(); ++MaterialIndex)
 	{
-		MaterialSet.Add(Curr);
-	}
+		UMaterialInterface* ExteriorMaterial = Materials[MaterialIndex];
+		
+		// If we have an odd number of materials, the last material duplicates itself.
+		UMaterialInterface* InteriorMaterial = Materials[MaterialIndex];
+		if (++MaterialIndex < Materials.Num())
+		{
+			InteriorMaterial = Materials[MaterialIndex];
+		}
 
+		MaterialSet.Add(FMaterialPair(ExteriorMaterial, InteriorMaterial));
+	}
+	
 	// create the final material array only containing unique materials
-	// and one slot for each internal material
-	TMap< UMaterialInterface *, int32> MaterialPtrToArrayIndex;
+	// alternating exterior and interior materials
+	TMap<UMaterialInterface*, int32> ExteriorMaterialPtrToArrayIndex;
+	TMap<UMaterialInterface*, int32> InteriorMaterialPtrToArrayIndex;
 	TArray<UMaterialInterface*> FinalMaterials;
-	for (UMaterialInterface *Curr : MaterialSet)
+	for (const FMaterialPair& Curr : MaterialSet)
 	{
 		// Add base material
-		TTuple< UMaterialInterface *, int32> CurrTuple(Curr, FinalMaterials.Add(Curr));
-		MaterialPtrToArrayIndex.Add(CurrTuple);
+		TTuple< UMaterialInterface*, int32> BaseTuple(Curr.Key, FinalMaterials.Add(Curr.Key));
+		ExteriorMaterialPtrToArrayIndex.Add(BaseTuple);
 
 		// Add interior material
-		FinalMaterials.Add(Curr);
+		TTuple< UMaterialInterface*, int32> InteriorTuple(Curr.Value, FinalMaterials.Add(Curr.Value));
+		InteriorMaterialPtrToArrayIndex.Add(InteriorTuple);
 	}
 
 	TManagedArray<int32>& MaterialID = GeometryCollection->MaterialID;
@@ -269,7 +285,14 @@ void UGeometryCollection::InitializeMaterials()
 		if (MaterialID[Material] < Materials.Num())
 		{
 			UMaterialInterface* OldMaterialPtr = Materials[MaterialID[Material]];
-			MaterialID[Material] = *MaterialPtrToArrayIndex.Find(OldMaterialPtr);
+			if (MaterialID[Material] % 2 == 0)
+			{
+				MaterialID[Material] = *ExteriorMaterialPtrToArrayIndex.Find(OldMaterialPtr);
+			}
+			else
+			{
+				MaterialID[Material] = *InteriorMaterialPtrToArrayIndex.Find(OldMaterialPtr);
+			}
 		}
 	}
 
@@ -277,13 +300,14 @@ void UGeometryCollection::InitializeMaterials()
 	Materials = FinalMaterials;
 
 	// Last Material is the selection one
-	UMaterialInterface* BoneSelectedMaterial = LoadObject<UMaterialInterface>(nullptr, GetSelectedMaterialPath(), nullptr, LOAD_None, nullptr);
 	BoneSelectedMaterialIndex = Materials.Add(BoneSelectedMaterial);
-	
+
 	GeometryCollection->ReindexMaterials();
 	InvalidateCollection();
 }
 
+
+
 /** Returns true if there is anything to render */
 bool UGeometryCollection::HasVisibleGeometry() const
 {
@@ -501,7 +525,7 @@ void UGeometryCollection::Serialize(FArchive& Ar)
 #endif
 }
 
-const TCHAR* UGeometryCollection::GetSelectedMaterialPath() const
+const TCHAR* UGeometryCollection::GetSelectedMaterialPath()
 {
 	return TEXT("/Engine/EditorMaterials/GeometryCollection/SelectedGeometryMaterial.SelectedGeometryMaterial");
 }
diff --git a/Engine/Source/Runtime/Experimental/GeometryCollectionEngine/Public/GeometryCollection/GeometryCollectionObject.h b/Engine/Source/Runtime/Experimental/GeometryCollectionEngine/Public/GeometryCollection/GeometryCollectionObject.h
--- a/Engine/Source/Runtime/Experimental/GeometryCollectionEngine/Public/GeometryCollection/GeometryCollectionObject.h
+++ b/Engine/Source/Runtime/Experimental/GeometryCollectionEngine/Public/GeometryCollection/GeometryCollectionObject.h
@@ -23,7 +23,7 @@ struct GEOMETRYCOLLECTIONENGINE_API FGeometryCollectionSource
 {
 	GENERATED_BODY()
 
-	UPROPERTY(BlueprintReadWrite, EditAnywhere, Category = "GeometrySource", meta=(AllowedClasses="StaticMesh, SkeletalMesh"))
+	UPROPERTY(BlueprintReadWrite, EditAnywhere, Category = "GeometrySource", meta=(AllowedClasses="StaticMesh, SkeletalMesh, GeometryCollection"))
 	FSoftObjectPath SourceGeometryObject;
 
 	UPROPERTY(BlueprintReadWrite, EditAnywhere, Category = "GeometrySource")
@@ -189,6 +189,7 @@ public:
 	/** appends the standard materials to this uobject */
 	void InitializeMaterials();
 
+
 	/** Returns true if there is anything to render */
 	bool HasVisibleGeometry() const;
 
@@ -344,7 +345,7 @@ public:
 	FORCEINLINE const int32 GetBoneSelectedMaterialIndex() const { return BoneSelectedMaterialIndex; }
 
 	/** Returns the asset path for the automatically populated selected material. */
-	const TCHAR* GetSelectedMaterialPath() const;
+	static const TCHAR* GetSelectedMaterialPath();
 
 #if WITH_EDITORONLY_DATA
 	/** Information for thumbnail rendering */
