GeometryCollection -- double Material ID on static mesh conversion to skip newly minted internal materials.

diff --git a/Engine/Plugins/Experimental/GeometryCollectionPlugin/Source/GeometryCollectionEditor/Private/GeometryCollection/GeometryCollectionConversion.cpp b/Engine/Plugins/Experimental/GeometryCollectionPlugin/Source/GeometryCollectionEditor/Private/GeometryCollection/GeometryCollectionConversion.cpp
--- a/Engine/Plugins/Experimental/GeometryCollectionPlugin/Source/GeometryCollectionEditor/Private/GeometryCollection/GeometryCollectionConversion.cpp
+++ b/Engine/Plugins/Experimental/GeometryCollectionPlugin/Source/GeometryCollectionEditor/Private/GeometryCollection/GeometryCollectionConversion.cpp
@@ -215,7 +215,8 @@ void FGeometryCollectionConversion::AppendStaticMesh(const UStaticMesh* StaticMe
 			// create new section
 			int32 SectionIndex = GeometryCollection->AddElements(1, FGeometryCollection::MaterialGroup);
 
-			Sections[SectionIndex].MaterialID = MaterialStart + CurrSection.MaterialIndex;
+			// Materials are doubled on ingestions: we multiply the incoming ID by 2 to skip internal materials.
+			Sections[SectionIndex].MaterialID = MaterialStart + (CurrSection.MaterialIndex * 2);
 
 			Sections[SectionIndex].FirstIndex = IndicesStart * 3 + CurrSection.FirstIndex;
 			Sections[SectionIndex].MinVertexIndex = VertexStart + CurrSection.MinVertexIndex;
