GeometryCollectionConversion
-- Refactored conversion from StaticMesh code to use FMeshDescription instead of render buffer so that Nanite Static Mesh can be converted to Geometry Collection.

diff --git a/Engine/Plugins/Experimental/ChaosEditor/Source/FractureEditor/Private/FractureToolGenerators.cpp b/Engine/Plugins/Experimental/ChaosEditor/Source/FractureEditor/Private/FractureToolGenerators.cpp
--- a/Engine/Plugins/Experimental/ChaosEditor/Source/FractureEditor/Private/FractureToolGenerators.cpp
+++ b/Engine/Plugins/Experimental/ChaosEditor/Source/FractureEditor/Private/FractureToolGenerators.cpp
@@ -254,9 +254,11 @@ AGeometryCollectionActor* UFractureToolGenerateAsset::ConvertActorsToGeometryCol
 	if (FracturedGeometryCollection->EnableNanite)
 	{
 		FracturedGeometryCollection->InvalidateCollection();
-		FracturedGeometryCollection->EnsureDataIsCooked(false /* init resources */);
+		FracturedGeometryCollection->EnsureDataIsCooked(true /* init resources */);
 	}
 
+	NewActor->GetGeometryCollectionComponent()->MarkRenderStateDirty();
+
 	// Add and initialize guids
 	::GeometryCollection::GenerateTemporaryGuids(FracturedGeometryCollection->GetGeometryCollection().Get(), 0 , true);
 
diff --git a/Engine/Plugins/Experimental/ChaosEditor/Source/FractureEditor/Private/SGeometryCollectionOutliner.cpp b/Engine/Plugins/Experimental/ChaosEditor/Source/FractureEditor/Private/SGeometryCollectionOutliner.cpp
--- a/Engine/Plugins/Experimental/ChaosEditor/Source/FractureEditor/Private/SGeometryCollectionOutliner.cpp
+++ b/Engine/Plugins/Experimental/ChaosEditor/Source/FractureEditor/Private/SGeometryCollectionOutliner.cpp
@@ -416,6 +416,11 @@ void FGeometryCollectionTreeItemComponent::GetChildrenForBone(FGeometryCollectio
 	{
 		if (FGeometryCollection* Collection = RestCollection->GetGeometryCollection().Get())
 		{
+			if (!Collection->HasAttribute("GUID", "Transform"))
+			{
+				GeometryCollection::GenerateTemporaryGuids(Collection);
+			}
+
 			if (const int32* BoneIndex = GuidIndexMap.Find(BoneItem.GetGuid()))
 			{
 				const TManagedArray<TSet<int32>>& Children = Collection->Children;
diff --git a/Engine/Plugins/Experimental/GeometryCollectionPlugin/Source/GeometryCollectionEditor/GeometryCollectionEditor.Build.cs b/Engine/Plugins/Experimental/GeometryCollectionPlugin/Source/GeometryCollectionEditor/GeometryCollectionEditor.Build.cs
--- a/Engine/Plugins/Experimental/GeometryCollectionPlugin/Source/GeometryCollectionEditor/GeometryCollectionEditor.Build.cs
+++ b/Engine/Plugins/Experimental/GeometryCollectionPlugin/Source/GeometryCollectionEditor/GeometryCollectionEditor.Build.cs
@@ -29,7 +29,9 @@ namespace UnrealBuildTool.Rules
 				    "SceneOutliner",
 					"EditorStyle",
 					"ToolMenus",
-					"Chaos"
+					"Chaos",
+					"MeshDescription",
+					"StaticMeshDescription"
 				}
 			);
 
diff --git a/Engine/Plugins/Experimental/GeometryCollectionPlugin/Source/GeometryCollectionEditor/Private/GeometryCollection/GeometryCollectionConversion.cpp b/Engine/Plugins/Experimental/GeometryCollectionPlugin/Source/GeometryCollectionEditor/Private/GeometryCollection/GeometryCollectionConversion.cpp
--- a/Engine/Plugins/Experimental/GeometryCollectionPlugin/Source/GeometryCollectionEditor/Private/GeometryCollection/GeometryCollectionConversion.cpp
+++ b/Engine/Plugins/Experimental/GeometryCollectionPlugin/Source/GeometryCollectionEditor/Private/GeometryCollection/GeometryCollectionConversion.cpp
@@ -22,10 +22,36 @@
 #include "Logging/LogMacros.h"
 #include "Materials/Material.h"
 #include "Rendering/SkeletalMeshRenderData.h"
+#include "MeshDescription.h"
+#include "StaticMeshAttributes.h"
+#include "StaticMeshOperations.h"
 
 
 DEFINE_LOG_CATEGORY_STATIC(UGeometryCollectionConversionLogging, Log, All);
 
+struct FUniqueVertex
+{
+	FVector Normal;
+	FVector Tangent;
+	FVector2D UV;
+
+	bool operator==(const FUniqueVertex& Other) const
+	{
+		return ((this->Normal == Other.Normal) &&
+			(this->Tangent == Other.Tangent) &&
+			(this->UV == Other.UV));
+	}
+};
+
+FORCEINLINE uint32 GetTypeHash(const FUniqueVertex& UniqueVertex)
+{
+	uint32 VertexHash = GetTypeHash(UniqueVertex.Normal);
+	VertexHash = HashCombine(VertexHash, GetTypeHash(UniqueVertex.Tangent));
+	VertexHash = HashCombine(VertexHash, GetTypeHash(UniqueVertex.UV));
+
+	return VertexHash;
+}
+
 void FGeometryCollectionConversion::AppendStaticMesh(const UStaticMesh* StaticMesh, const TArray<UMaterialInterface*>& Materials, const FTransform& StaticMeshTransform, UGeometryCollection* GeometryCollectionObject, bool ReindexMaterials)
 {
 	if (StaticMesh == nullptr)
@@ -33,70 +59,138 @@ void FGeometryCollectionConversion::AppendStaticMesh(const UStaticMesh* StaticMe
 		return;
 	}
 
+	// Prefer the HiRes description, although this isn't always available.
+	FMeshDescription* MeshDescription = nullptr;
+	if (StaticMesh->IsHiResMeshDescriptionValid())
+	{
+		MeshDescription = StaticMesh->GetHiResMeshDescription();
+	}
+	else
+	{
+		MeshDescription = StaticMesh->GetMeshDescription(0);
+	}
+
 	check(GeometryCollectionObject);
 	TSharedPtr<FGeometryCollection, ESPMode::ThreadSafe> GeometryCollectionPtr = GeometryCollectionObject->GetGeometryCollection();
 	FGeometryCollection* GeometryCollection = GeometryCollectionPtr.Get();
 	check(GeometryCollection);
 
-	// @todo : Discuss how to handle multiple LOD's
-	if (StaticMesh->GetRenderData() && StaticMesh->GetRenderData()->LODResources.Num() > 0)
+	if (MeshDescription)
 	{
-		const FStaticMeshVertexBuffers& VertexBuffer = StaticMesh->GetRenderData()->LODResources[0].VertexBuffers;
-
-		// vertex information
-		TManagedArray<FVector3f>& Vertex = GeometryCollection->Vertex;
-		TManagedArray<FVector3f>& TangentU = GeometryCollection->TangentU;
-		TManagedArray<FVector3f>& TangentV = GeometryCollection->TangentV;
-		TManagedArray<FVector3f>& Normal = GeometryCollection->Normal;
-		TManagedArray<FVector2D>& UV = GeometryCollection->UV;
-		TManagedArray<FLinearColor>& Color = GeometryCollection->Color;
-		TManagedArray<int32>& BoneMap = GeometryCollection->BoneMap;
-		TManagedArray<FLinearColor>& BoneColor = GeometryCollection->BoneColor;
-		TManagedArray<FString>& BoneName = GeometryCollection->BoneName;
-
-		const int32 VertexCount = VertexBuffer.PositionVertexBuffer.GetNumVertices();
-		int InitialNumVertices = GeometryCollection->NumElements(FGeometryCollection::VerticesGroup);
-		int VertexStart = GeometryCollection->AddElements(VertexCount, FGeometryCollection::VerticesGroup);
+		FStaticMeshOperations::ComputeTriangleTangentsAndNormals(*MeshDescription);
+		FStaticMeshOperations::RecomputeNormalsAndTangentsIfNeeded(*MeshDescription, EComputeNTBsFlags::UseMikkTSpace);
+
+		// source vertex information
+		FStaticMeshAttributes Attributes(*MeshDescription);
+		TArrayView<const FVector> SourcePosition = Attributes.GetVertexPositions().GetRawArray();
+		TArrayView<const FVector> SourceTangent = Attributes.GetVertexInstanceTangents().GetRawArray();
+		TArrayView<const float> SourceBinormalSign = Attributes.GetVertexInstanceBinormalSigns().GetRawArray();
+		TArrayView<const FVector> SourceNormal = Attributes.GetVertexInstanceNormals().GetRawArray();
+		TArrayView<const FVector4> SourceColor = Attributes.GetVertexInstanceColors().GetRawArray();
+
+		TVertexInstanceAttributesConstRef<FVector2D> InstanceUVs = Attributes.GetVertexInstanceUVs();
+		TArrayView<const FVector2D> SourceUV = InstanceUVs.GetRawArray();
+		const int NumUVLayers = InstanceUVs.GetNumChannels();
+		
+		// target vertex information
+		TManagedArray<FVector>& TargetVertex = GeometryCollection->Vertex;
+		TManagedArray<FVector>& TargetTangentU = GeometryCollection->TangentU;
+		TManagedArray<FVector>& TargetTangentV = GeometryCollection->TangentV;
+		TManagedArray<FVector>& TargetNormal = GeometryCollection->Normal;
+		TManagedArray<FVector2D>& TargetUV = GeometryCollection->UV;
+		TManagedArray<FLinearColor>& TargetColor = GeometryCollection->Color;
+		TManagedArray<int32>& TargetBoneMap = GeometryCollection->BoneMap;
+		TManagedArray<FLinearColor>& TargetBoneColor = GeometryCollection->BoneColor;
+		TManagedArray<FString>& TargetBoneName = GeometryCollection->BoneName;
+
+		const int32 VertexStart = GeometryCollection->NumElements(FGeometryCollection::VerticesGroup);
+		int32 VertexCount = 0;
+		
 		FVector Scale = StaticMeshTransform.GetScale3D();
-		for (int32 VertexIndex = 0; VertexIndex < VertexCount; VertexIndex++)
+		
+		// We'll need to re-introduce UV seams, etc. by splitting vertices.
+		// A new mapping of MeshDescription vertex instances to the split vertices is maintained.
+		TMap<FVertexInstanceID, int32> VertexInstanceToGeometryCollectionVertex;
+		VertexInstanceToGeometryCollectionVertex.Reserve(Attributes.GetVertexInstanceNormals().GetNumElements());
+		
+		for (const FVertexID VertexIndex : MeshDescription->Vertices().GetElementIDs())
 		{
-			int VertexOffset = VertexStart + VertexIndex;
-			Vertex[VertexOffset] = VertexBuffer.PositionVertexBuffer.VertexPosition(VertexIndex) * Scale;
-			BoneMap[VertexOffset] = GeometryCollection->NumElements(FGeometryCollection::TransformGroup);
-
-			TangentU[VertexOffset] = VertexBuffer.StaticMeshVertexBuffer.VertexTangentX(VertexIndex);
-			TangentV[VertexOffset] = VertexBuffer.StaticMeshVertexBuffer.VertexTangentY(VertexIndex);
-			Normal[VertexOffset] = VertexBuffer.StaticMeshVertexBuffer.VertexTangentZ(VertexIndex);
-
-			// @todo : Support multiple UV's per vertex based on MAX_STATIC_TEXCOORDS
-			UV[VertexOffset] = VertexBuffer.StaticMeshVertexBuffer.GetVertexUV(VertexIndex, 0);
-			if (VertexBuffer.ColorVertexBuffer.GetNumVertices() == VertexCount)
-				Color[VertexOffset] = VertexBuffer.ColorVertexBuffer.VertexColor(VertexIndex);
-			else
-				Color[VertexOffset] = FLinearColor::White;
+			
+			TArrayView<const FVertexInstanceID> ReferencingVertexInstances = MeshDescription->GetVertexVertexInstanceIDs(VertexIndex);
+
+			// Generate per instance hash of splittable attributes.
+			TMap<FUniqueVertex, TArray<FVertexInstanceID>> SplitVertices;
+			for (const FVertexInstanceID& InstanceID : ReferencingVertexInstances)
+			{
+				FUniqueVertex UniqueVertex{ SourceNormal[InstanceID], SourceTangent[InstanceID], SourceUV[InstanceID] };
+				TArray<FVertexInstanceID>& SplitVertex = SplitVertices.FindOrAdd(UniqueVertex);
+				SplitVertex.Add(InstanceID);
+			}
+
+			int32 CurrentVertex = GeometryCollection->AddElements(SplitVertices.Num(), FGeometryCollection::VerticesGroup);
+			
+			// Create a new vertex for each split vertex and map the mesh description instance to it.
+			for (const TTuple<FUniqueVertex,TArray<FVertexInstanceID>>& SplitVertex : SplitVertices)
+			{
+				const TArray<FVertexInstanceID>& InstanceIDs = SplitVertex.Value;
+				const FVertexInstanceID& ExemplarInstanceID = InstanceIDs[0];
+
+				TargetVertex[CurrentVertex] = SourcePosition[VertexIndex] * Scale;
+				TargetBoneMap[CurrentVertex] = GeometryCollection->NumElements(FGeometryCollection::TransformGroup);
+
+				TargetNormal[CurrentVertex] = SourceNormal[ExemplarInstanceID];
+				TargetTangentU[CurrentVertex] = SourceTangent[ExemplarInstanceID];
+				TargetTangentV[CurrentVertex] = SourceBinormalSign[ExemplarInstanceID] * FVector::CrossProduct(TargetNormal[CurrentVertex], TargetTangentU[CurrentVertex]);
+
+				// @todo : Support multiple UV's per vertex based on MAX_STATIC_TEXCOORDS
+				TargetUV[CurrentVertex] = SourceUV[ExemplarInstanceID];
+
+				if (SourceColor.Num() > 0)
+				{
+					TargetColor[CurrentVertex] = FLinearColor(SourceColor[ExemplarInstanceID]);
+				}
+				else
+				{
+					TargetColor[CurrentVertex] = FLinearColor::White;
+				}
+
+				for (const FVertexInstanceID& InstanceID : InstanceIDs)
+				{
+					VertexInstanceToGeometryCollectionVertex.Add(InstanceID, CurrentVertex);
+				}
+
+				++CurrentVertex;
+				++VertexCount;
+			}
 		}
 
-		// Triangle Indices
-		TManagedArray<FIntVector>& Indices = GeometryCollection->Indices;
-		TManagedArray<bool>& Visible = GeometryCollection->Visible;
-		TManagedArray<int32>& MaterialID = GeometryCollection->MaterialID;
-		TManagedArray<int32>& MaterialIndex = GeometryCollection->MaterialIndex;
-
-		const FRawStaticIndexBuffer& IndexBuffer = StaticMesh->GetRenderData()->LODResources[0].IndexBuffer;
-		FIndexArrayView IndexBufferView = IndexBuffer.GetArrayView();
-		const int32 IndicesCount = IndexBuffer.GetNumIndices() / 3;
-		int InitialNumIndices = GeometryCollection->NumElements(FGeometryCollection::FacesGroup);
-		int IndicesStart = GeometryCollection->AddElements(IndicesCount, FGeometryCollection::FacesGroup);
-		for (int32 IndicesIndex = 0, StaticIndex = 0; IndicesIndex < IndicesCount; IndicesIndex++, StaticIndex += 3)
+		// target triangle indices
+		TManagedArray<FIntVector>& TargetIndices = GeometryCollection->Indices;
+		TManagedArray<bool>& TargetVisible = GeometryCollection->Visible;
+		TManagedArray<int32>& TargetMaterialID = GeometryCollection->MaterialID;
+		TManagedArray<int32>& TargetMaterialIndex = GeometryCollection->MaterialIndex;
+
+		const int32 IndicesCount = MeshDescription->Triangles().Num();
+		const int32 InitialNumIndices = GeometryCollection->NumElements(FGeometryCollection::FacesGroup);
+		const int32 IndicesStart = GeometryCollection->AddElements(IndicesCount, FGeometryCollection::FacesGroup);
+		int32 TargetIndex = IndicesStart;
+		for (const int32 TriangleIndex : MeshDescription->Triangles().GetElementIDs())
 		{
-			int32 IndicesOffset = IndicesStart + IndicesIndex;
-			Indices[IndicesOffset] = FIntVector(
-				IndexBufferView[StaticIndex] + VertexStart,
-				IndexBufferView[StaticIndex + 1] + VertexStart,
-				IndexBufferView[StaticIndex + 2] + VertexStart);
-			Visible[IndicesOffset] = true;
-			MaterialID[IndicesOffset] = 0;
-			MaterialIndex[IndicesOffset] = IndicesOffset;
+			TArrayView<const FVertexInstanceID> TriangleVertices = MeshDescription->GetTriangleVertexInstances(TriangleIndex);
+
+			TargetIndices[TargetIndex] = FIntVector(
+				VertexInstanceToGeometryCollectionVertex[TriangleVertices[0]],
+				VertexInstanceToGeometryCollectionVertex[TriangleVertices[1]],
+				VertexInstanceToGeometryCollectionVertex[TriangleVertices[2]]
+			);
+
+			TargetVisible[TargetIndex] = true;
+			TargetMaterialID[TargetIndex] = 0;
+
+			// Is this right?
+			TargetMaterialIndex[TargetIndex] = TargetIndex;
+
+			++TargetIndex;
 		}
 
 		// Geometry transform
@@ -113,8 +207,8 @@ void FGeometryCollectionConversion::AppendStaticMesh(const UStaticMesh* StaticMe
 		SimulationType[TransformIndex1] = FGeometryCollection::ESimulationTypes::FST_Rigid;
 
 		const FColor RandBoneColor(FMath::Rand() % 100 + 5, FMath::Rand() % 100 + 5, FMath::Rand() % 100 + 5, 255);
-		BoneColor[TransformIndex1] = FLinearColor(RandBoneColor);
-		BoneName[TransformIndex1] = StaticMesh->GetName();
+		TargetBoneColor[TransformIndex1] = FLinearColor(RandBoneColor);
+		TargetBoneName[TransformIndex1] = StaticMesh->GetName();
 
 		// GeometryGroup
 		int GeometryIndex = GeometryCollection->AddElements(1, FGeometryCollection::GeometryGroup);
@@ -128,8 +222,8 @@ void FGeometryCollectionConversion::AppendStaticMesh(const UStaticMesh* StaticMe
 		TManagedArray<int32>& FaceStartArray = GeometryCollection->FaceStart;
 		TManagedArray<int32>& FaceCountArray = GeometryCollection->FaceCount;
 
-		TransformIndex[GeometryIndex] = BoneMap[VertexStart];
-		VertexStartArray[GeometryIndex] = InitialNumVertices;
+		TransformIndex[GeometryIndex] = TargetBoneMap[VertexStart];
+		VertexStartArray[GeometryIndex] = VertexStart;
 		VertexCountArray[GeometryIndex] = VertexCount;
 		FaceStartArray[GeometryIndex] = InitialNumIndices;
 		FaceCountArray[GeometryIndex] = IndicesCount;
@@ -141,7 +235,7 @@ void FGeometryCollectionConversion::AppendStaticMesh(const UStaticMesh* StaticMe
 		FVector Center(0);
 		for (int32 VertexIndex = VertexStart; VertexIndex < VertexStart + VertexCount; VertexIndex++)
 		{
-			Center += Vertex[VertexIndex];
+			Center += TargetVertex[VertexIndex];
 		}
 		if (VertexCount) Center /= VertexCount;
 
@@ -151,9 +245,9 @@ void FGeometryCollectionConversion::AppendStaticMesh(const UStaticMesh* StaticMe
 		OuterRadius[GeometryIndex] = -FLT_MAX;
 		for (int32 VertexIndex = VertexStart; VertexIndex < VertexStart + VertexCount; VertexIndex++)
 		{
-			BoundingBox[GeometryIndex] += Vertex[VertexIndex];
+			BoundingBox[GeometryIndex] += TargetVertex[VertexIndex];
 
-			float Delta = (Center - Vertex[VertexIndex]).Size();
+			float Delta = (Center - TargetVertex[VertexIndex]).Size();
 			InnerRadius[GeometryIndex] = FMath::Min(InnerRadius[GeometryIndex], Delta);
 			OuterRadius[GeometryIndex] = FMath::Max(OuterRadius[GeometryIndex], Delta);
 		}
@@ -164,7 +258,7 @@ void FGeometryCollectionConversion::AppendStaticMesh(const UStaticMesh* StaticMe
 			FVector Centroid(0);
 			for (int e = 0; e < 3; e++)
 			{
-				Centroid += Vertex[Indices[fdx][e]];
+				Centroid += TargetVertex[TargetIndices[fdx][e]];
 			}
 			Centroid /= 3;
 
@@ -179,7 +273,7 @@ void FGeometryCollectionConversion::AppendStaticMesh(const UStaticMesh* StaticMe
 			for (int e = 0; e < 3; e++)
 			{
 				int i = e, j = (e + 1) % 3;
-				FVector Edge = Vertex[Indices[fdx][i]] + 0.5 * (Vertex[Indices[fdx][j]] - Vertex[Indices[fdx][i]]);
+				FVector Edge = TargetVertex[TargetIndices[fdx][i]] + 0.5 * (TargetVertex[TargetIndices[fdx][j]] - TargetVertex[TargetIndices[fdx][i]]);
 				float Delta = (Center - Edge).Size();
 				InnerRadius[GeometryIndex] = FMath::Min(InnerRadius[GeometryIndex], Delta);
 				OuterRadius[GeometryIndex] = FMath::Max(OuterRadius[GeometryIndex], Delta);
@@ -230,7 +324,7 @@ void FGeometryCollectionConversion::AppendStaticMesh(const UStaticMesh* StaticMe
 			// note the divide by 3 - the GeometryCollection stores indices in tuples of 3 rather than in a flat array
 			for (int32 i = Sections[SectionIndex].FirstIndex / 3; i < Sections[SectionIndex].FirstIndex / 3 + Sections[SectionIndex].NumTriangles; ++i)
 			{
-				MaterialID[i] = Sections[SectionIndex].MaterialID;
+				TargetMaterialID[i] = Sections[SectionIndex].MaterialID;
 			}
 		}
 
@@ -240,7 +334,6 @@ void FGeometryCollectionConversion::AppendStaticMesh(const UStaticMesh* StaticMe
 	}
 }
 
-
 void FGeometryCollectionConversion::AppendGeometryCollection(const UGeometryCollection* SourceGeometryCollection, const TArray<UMaterialInterface*>& Materials, const FTransform& GeometryCollectionTransform, UGeometryCollection* TargetGeometryCollectionObject, bool ReindexMaterials)
 {
 	if (SourceGeometryCollection == nullptr)
