LWC: Type fixes for geometry collection conversion with LWC enabled

diff --git a/Engine/Plugins/Experimental/GeometryCollectionPlugin/Source/GeometryCollectionEditor/Private/GeometryCollection/GeometryCollectionConversion.cpp b/Engine/Plugins/Experimental/GeometryCollectionPlugin/Source/GeometryCollectionEditor/Private/GeometryCollection/GeometryCollectionConversion.cpp
--- a/Engine/Plugins/Experimental/GeometryCollectionPlugin/Source/GeometryCollectionEditor/Private/GeometryCollection/GeometryCollectionConversion.cpp
+++ b/Engine/Plugins/Experimental/GeometryCollectionPlugin/Source/GeometryCollectionEditor/Private/GeometryCollection/GeometryCollectionConversion.cpp
@@ -82,10 +82,10 @@ void FGeometryCollectionConversion::AppendStaticMesh(const UStaticMesh* StaticMe
 
 		// source vertex information
 		FStaticMeshAttributes Attributes(*MeshDescription);
-		TArrayView<const FVector> SourcePosition = Attributes.GetVertexPositions().GetRawArray();
-		TArrayView<const FVector> SourceTangent = Attributes.GetVertexInstanceTangents().GetRawArray();
+		TArrayView<const FVector3f> SourcePosition = Attributes.GetVertexPositions().GetRawArray();
+		TArrayView<const FVector3f> SourceTangent = Attributes.GetVertexInstanceTangents().GetRawArray();
 		TArrayView<const float> SourceBinormalSign = Attributes.GetVertexInstanceBinormalSigns().GetRawArray();
-		TArrayView<const FVector> SourceNormal = Attributes.GetVertexInstanceNormals().GetRawArray();
+		TArrayView<const FVector3f> SourceNormal = Attributes.GetVertexInstanceNormals().GetRawArray();
 		TArrayView<const FVector4> SourceColor = Attributes.GetVertexInstanceColors().GetRawArray();
 
 		TVertexInstanceAttributesConstRef<FVector2D> InstanceUVs = Attributes.GetVertexInstanceUVs();
@@ -93,10 +93,10 @@ void FGeometryCollectionConversion::AppendStaticMesh(const UStaticMesh* StaticMe
 		const int NumUVLayers = InstanceUVs.GetNumChannels();
 		
 		// target vertex information
-		TManagedArray<FVector>& TargetVertex = GeometryCollection->Vertex;
-		TManagedArray<FVector>& TargetTangentU = GeometryCollection->TangentU;
-		TManagedArray<FVector>& TargetTangentV = GeometryCollection->TangentV;
-		TManagedArray<FVector>& TargetNormal = GeometryCollection->Normal;
+		TManagedArray<FVector3f>& TargetVertex = GeometryCollection->Vertex;
+		TManagedArray<FVector3f>& TargetTangentU = GeometryCollection->TangentU;
+		TManagedArray<FVector3f>& TargetTangentV = GeometryCollection->TangentV;
+		TManagedArray<FVector3f>& TargetNormal = GeometryCollection->Normal;
 		TManagedArray<FVector2D>& TargetUV = GeometryCollection->UV;
 		TManagedArray<FLinearColor>& TargetColor = GeometryCollection->Color;
 		TManagedArray<int32>& TargetBoneMap = GeometryCollection->BoneMap;
