GeometryCollection : EditorStyle
-Added vitrtual destructor.

diff --git a/Engine/Plugins/Experimental/GeometryCollectionPlugin/Source/GeometryCollectionEditor/Private/GeometryCollection/GeometryCollectionEditorStyle.h b/Engine/Plugins/Experimental/GeometryCollectionPlugin/Source/GeometryCollectionEditor/Private/GeometryCollection/GeometryCollectionEditorStyle.h
--- a/Engine/Plugins/Experimental/GeometryCollectionPlugin/Source/GeometryCollectionEditor/Private/GeometryCollection/GeometryCollectionEditorStyle.h
+++ b/Engine/Plugins/Experimental/GeometryCollectionPlugin/Source/GeometryCollectionEditor/Private/GeometryCollection/GeometryCollectionEditorStyle.h
@@ -22,7 +22,7 @@ public:
 		FSlateStyleRegistry::RegisterSlateStyle(*this);
 	}
 
-	~FGeometryCollectionEditorStyle()
+	virtual ~FGeometryCollectionEditorStyle()
 	{
 		FSlateStyleRegistry::UnRegisterSlateStyle(*this);
 	}
