GeometryCollection - Dataflow Editor
- Configure the dataflow asset editor for the geometry collection object. Currently requires the enabling the pvar (p.Chaos.GeometryCollection.DataflowEditor 1)

diff --git a/Engine/Plugins/Experimental/GeometryCollectionPlugin/GeometryCollectionPlugin.uplugin b/Engine/Plugins/Experimental/GeometryCollectionPlugin/GeometryCollectionPlugin.uplugin
--- a/Engine/Plugins/Experimental/GeometryCollectionPlugin/GeometryCollectionPlugin.uplugin
+++ b/Engine/Plugins/Experimental/GeometryCollectionPlugin/GeometryCollectionPlugin.uplugin
@@ -39,6 +39,10 @@
 		{
 			"Name": "ChaosSolverPlugin",
 			"Enabled": true
+		},
+		{
+			"Name": "Dataflow",
+			"Enabled": true
 		}
 	]
 }
diff --git a/Engine/Plugins/Experimental/GeometryCollectionPlugin/Source/GeometryCollectionEditor/GeometryCollectionEditor.Build.cs b/Engine/Plugins/Experimental/GeometryCollectionPlugin/Source/GeometryCollectionEditor/GeometryCollectionEditor.Build.cs
--- a/Engine/Plugins/Experimental/GeometryCollectionPlugin/Source/GeometryCollectionEditor/GeometryCollectionEditor.Build.cs
+++ b/Engine/Plugins/Experimental/GeometryCollectionPlugin/Source/GeometryCollectionEditor/GeometryCollectionEditor.Build.cs
@@ -34,6 +34,10 @@ namespace UnrealBuildTool.Rules
 					"StaticMeshDescription",
 					"LevelEditor",
 					"InputCore",
+					"GraphEditor",
+					"DataflowCore",
+					"DataflowEngine",
+					"DataflowEditor",
 				}
 			);
 
diff --git a/Engine/Plugins/Experimental/GeometryCollectionPlugin/Source/GeometryCollectionEditor/Private/GeometryCollection/AssetTypeActions_GeometryCollection.cpp b/Engine/Plugins/Experimental/GeometryCollectionPlugin/Source/GeometryCollectionEditor/Private/GeometryCollection/AssetTypeActions_GeometryCollection.cpp
--- a/Engine/Plugins/Experimental/GeometryCollectionPlugin/Source/GeometryCollectionEditor/Private/GeometryCollection/AssetTypeActions_GeometryCollection.cpp
+++ b/Engine/Plugins/Experimental/GeometryCollectionPlugin/Source/GeometryCollectionEditor/Private/GeometryCollection/AssetTypeActions_GeometryCollection.cpp
@@ -1,10 +1,17 @@
 // Copyright Epic Games, Inc. All Rights Reserved.
 
 #include "GeometryCollection/AssetTypeActions_GeometryCollection.h"
+
+#include "GeometryCollection/Dataflow/GeometryCollectionEditorToolkit.h"
+#include "GeometryCollection/GeometryCollectionEditorPlugin.h"
 #include "ThumbnailRendering/SceneThumbnailInfo.h"
 #include "GeometryCollection/GeometryCollectionObject.h"
 #include "ToolMenus.h"
 
+bool bGeometryCollectionDataflowEditor = false;
+FAutoConsoleVariableRef CVarGeometryCollectionDataflowEditor(TEXT("p.Chaos.GeometryCollection.DataflowEditor"), bGeometryCollectionDataflowEditor, TEXT("Enable dataflow asset editor on geometry collection assets(Curently Dev-Only)"));
+
+
 #define LOCTEXT_NAMESPACE "AssetTypeActions"
 
 UClass* FAssetTypeActions_GeometryCollection::GetSupportedClass() const
@@ -33,5 +40,24 @@ void FAssetTypeActions_GeometryCollection::GetActions(const TArray<UObject*>& In
 	// IconPath = Plugin->GetBaseDir() / TEXT("Resources/Icon128.png");
 }
 
+void FAssetTypeActions_GeometryCollection::OpenAssetEditor(const TArray<UObject*>& InObjects, TSharedPtr<class IToolkitHost> EditWithinLevelEditor)
+{
+	if (bGeometryCollectionDataflowEditor)
+	{
+		EToolkitMode::Type Mode = EditWithinLevelEditor.IsValid() ? EToolkitMode::WorldCentric : EToolkitMode::Standalone;
+		for (auto ObjIt = InObjects.CreateConstIterator(); ObjIt; ++ObjIt)
+		{
+			if (auto Object = Cast<UGeometryCollection>(*ObjIt))
+			{
+				IGeometryCollectionEditorPlugin* EditorModule = &FModuleManager::LoadModuleChecked<IGeometryCollectionEditorPlugin>("GeometryCollectionEditor");
+				EditorModule->CreateGeometryCollectionAssetEditor(Mode, EditWithinLevelEditor, Object);
+			}
+		}
+	}
+	else
+	{ 
+		FAssetTypeActions_Base::OpenAssetEditor(InObjects,EditWithinLevelEditor);
+	}
+}
 
 #undef LOCTEXT_NAMESPACE
diff --git a/Engine/Plugins/Experimental/GeometryCollectionPlugin/Source/GeometryCollectionEditor/Private/GeometryCollection/Dataflow/GeometryCollectionEditorToolkit.cpp b/Engine/Plugins/Experimental/GeometryCollectionPlugin/Source/GeometryCollectionEditor/Private/GeometryCollection/Dataflow/GeometryCollectionEditorToolkit.cpp
new file mode 100644
--- /dev/null
+++ b/Engine/Plugins/Experimental/GeometryCollectionPlugin/Source/GeometryCollectionEditor/Private/GeometryCollection/Dataflow/GeometryCollectionEditorToolkit.cpp
@@ -0,0 +1,251 @@
+// Copyright Epic Games, Inc. All Rights Reserved.
+
+#include "GeometryCollection/Dataflow/GeometryCollectionEditorToolkit.h"
+
+#include "Dataflow/DataflowEditorActions.h"
+#include "Dataflow/DataflowEdNode.h"
+#include "Dataflow/DataflowNodeFactory.h"
+#include "Dataflow/DataflowObject.h"
+#include "Dataflow/DataflowSchema.h"
+#include "Dataflow/DataflowCore.h"
+#include "EditorStyleSet.h"
+#include "EditorViewportTabContent.h"
+#include "EditorViewportLayout.h"
+#include "EditorViewportCommands.h"
+#include "GraphEditorActions.h"
+#include "Modules/ModuleManager.h"
+#include "PropertyEditorModule.h"
+#include "Styling/SlateStyleRegistry.h"
+#include "Framework/Commands/GenericCommands.h"
+
+#define LOCTEXT_NAMESPACE "GeometryCollectionEditorToolkit"
+
+//DEFINE_LOG_CATEGORY_STATIC(FGeometryCollectionEditorToolkitLog, Log, All);
+
+
+const FName FGeometryCollectionEditorToolkit::GraphCanvasTabId(TEXT("GeometryCollectionEditor_GraphCanvas"));
+const FName FGeometryCollectionEditorToolkit::PropertiesTabId(TEXT("GeometryCollectionEditor_Properties"));
+
+void FGeometryCollectionEditorToolkit::InitGeometryCollectionAssetEditor(const EToolkitMode::Type Mode, const TSharedPtr<IToolkitHost>& InitToolkitHost, UObject* ObjectToEdit)
+{
+	Dataflow = nullptr;
+	GeometryCollection = CastChecked<UGeometryCollection>(ObjectToEdit);
+	if (GeometryCollection != nullptr)
+	{
+		if (GeometryCollection->Dataflow == nullptr)
+		{
+			const FName NodeName = MakeUniqueObjectName(GeometryCollection, UDataflow::StaticClass(), FName("GeometryCollectionDataflowAsset"));
+			GeometryCollection->Dataflow = NewObject<UDataflow>(GeometryCollection, NodeName);
+		}
+		Dataflow = GeometryCollection->Dataflow;
+		GeometryCollection->Dataflow->Schema = UDataflowSchema::StaticClass();
+
+		GraphEditor = CreateGraphEditorWidget(Dataflow);
+		PropertiesEditor = CreatePropertiesEditorWidget(ObjectToEdit);
+
+		const TSharedRef<FTabManager::FLayout> StandaloneDefaultLayout = FTabManager::NewLayout("GeometryCollectionDataflowEditor_Layout")
+			->AddArea
+			(
+				FTabManager::NewPrimaryArea()->SetOrientation(Orient_Vertical)
+				->Split
+				(
+					FTabManager::NewSplitter()->SetOrientation(Orient_Vertical)
+					->Split
+					(
+						FTabManager::NewSplitter()->SetOrientation(Orient_Horizontal)
+						->SetSizeCoefficient(0.9f)
+						->Split
+						(
+							FTabManager::NewStack()
+							->SetSizeCoefficient(0.6f)
+							->AddTab(GraphCanvasTabId, ETabState::OpenedTab)
+						)
+						->Split
+						(
+							FTabManager::NewSplitter()->SetOrientation(Orient_Vertical)
+							->SetSizeCoefficient(0.2f)
+							->Split
+							(
+								FTabManager::NewStack()
+								->SetSizeCoefficient(0.7f)
+								->AddTab(PropertiesTabId, ETabState::OpenedTab)
+							)
+						)
+					)
+				)
+			);
+
+		const bool bCreateDefaultStandaloneMenu = true;
+		const bool bCreateDefaultToolbar = true;
+		FAssetEditorToolkit::InitAssetEditor(Mode, InitToolkitHost, FName(TEXT("GeometryCollectionEditorApp")), StandaloneDefaultLayout, bCreateDefaultToolbar, bCreateDefaultStandaloneMenu, ObjectToEdit);
+	}
+}
+
+void FGeometryCollectionEditorToolkit::EvaluateNode()
+{
+	Dataflow::FGeometryCollectionContext Context(GeometryCollection, FGameTime::GetTimeSinceAppStart().GetRealTimeSeconds());
+	FDataflowEditorCommands::EvaluateNodes(GetSelectedNodes(), Context);
+}
+
+void FGeometryCollectionEditorToolkit::DeleteNode()
+{
+	if (UDataflow* Graph = dynamic_cast<UDataflow*>(GraphEditor->GetCurrentGraph()))
+	{
+		FDataflowEditorCommands::DeleteNodes(Graph, GetSelectedNodes());
+	}
+}
+
+FGraphPanelSelectionSet FGeometryCollectionEditorToolkit::GetSelectedNodes() const
+{
+	if (GraphEditor.IsValid())
+	{
+		return GraphEditor->GetSelectedNodes();
+	}
+	return FGraphPanelSelectionSet();
+}
+
+void FGeometryCollectionEditorToolkit::OnSelectedNodesChanged(const TSet<UObject*>& NewSelection)
+{
+	if (UDataflow* Graph = dynamic_cast<UDataflow*>(GraphEditor->GetCurrentGraph()))
+	{
+		FDataflowEditorCommands::OnSelectedNodesChanged(PropertiesEditor, Dataflow, Graph, NewSelection);
+	}
+}
+
+
+TSharedRef<SGraphEditor> FGeometryCollectionEditorToolkit::CreateGraphEditorWidget(UDataflow* DataflowToEdit)
+{
+	ensure(DataflowToEdit);
+
+	FDataflowEditorCommands::Register();
+	FGraphEditorCommands::Register();
+
+	// No need to regenerate the commands.
+	if (!GraphEditorCommands.IsValid())
+	{
+		GraphEditorCommands = MakeShareable(new FUICommandList);
+		{
+			GraphEditorCommands->MapAction(FGenericCommands::Get().Delete,
+				FExecuteAction::CreateSP(this, &FGeometryCollectionEditorToolkit::DeleteNode)
+			);
+			GraphEditorCommands->MapAction(FDataflowEditorCommands::Get().EvaluateNode,
+				FExecuteAction::CreateSP(this, &FGeometryCollectionEditorToolkit::EvaluateNode)
+			);
+		}
+	}
+
+
+	FGraphAppearanceInfo AppearanceInfo;
+	AppearanceInfo.CornerText = LOCTEXT("AppearanceCornerText_GeometryCollectionEditor", "Dataflow");
+
+	SGraphEditor::FGraphEditorEvents InEvents;
+	InEvents.OnSelectionChanged = SGraphEditor::FOnSelectionChanged::CreateSP(this, &FGeometryCollectionEditorToolkit::OnSelectedNodesChanged);
+
+	return SNew(SGraphEditor)
+		.AdditionalCommands(GraphEditorCommands)
+		.IsEditable(true)
+		.Appearance(AppearanceInfo)
+		.GraphToEdit(DataflowToEdit)
+		.GraphEvents(InEvents)
+		.ShowGraphStateOverlay(false);
+}
+
+TSharedPtr<IDetailsView> FGeometryCollectionEditorToolkit::CreatePropertiesEditorWidget(UObject* ObjectToEdit)
+{
+	ensure(ObjectToEdit);
+	FPropertyEditorModule& PropertyEditorModule = FModuleManager::LoadModuleChecked<FPropertyEditorModule>(TEXT("PropertyEditor"));
+
+	FDetailsViewArgs DetailsViewArgs;
+	DetailsViewArgs.bAllowSearch = true;
+	DetailsViewArgs.bLockable = false;
+	DetailsViewArgs.bUpdatesFromSelection = false;
+	DetailsViewArgs.NameAreaSettings = FDetailsViewArgs::HideNameArea;
+	DetailsViewArgs.NotifyHook = this;
+
+	TSharedPtr<IDetailsView> DetailsView = PropertyEditorModule.CreateDetailView(DetailsViewArgs);
+	DetailsView->SetObject(ObjectToEdit);
+	return DetailsView;
+
+}
+
+
+TSharedRef<SDockTab> FGeometryCollectionEditorToolkit::SpawnTab_GraphCanvas(const FSpawnTabArgs& Args)
+{
+	check(Args.GetTabId() == GraphCanvasTabId);
+
+	TSharedRef<SDockTab> SpawnedTab = SNew(SDockTab)
+		.Label(LOCTEXT("GeometryCollectionEditor_Dataflow_TabTitle", "Graph"))
+		[
+			GraphEditor.ToSharedRef()
+		];
+
+	return SpawnedTab;
+}
+
+TSharedRef<SDockTab> FGeometryCollectionEditorToolkit::SpawnTab_Properties(const FSpawnTabArgs& Args)
+{
+	check(Args.GetTabId() == PropertiesTabId);
+
+	return SNew(SDockTab)
+		.Label(LOCTEXT("GeometryCollectionEditor_Properties_TabTitle", "Details"))
+		[
+			PropertiesEditor.ToSharedRef()
+		];
+}
+
+
+void FGeometryCollectionEditorToolkit::RegisterTabSpawners(const TSharedRef<FTabManager>& InTabManager)
+{
+	TSharedRef<FWorkspaceItem> WorkspaceMenuCategoryRef = InTabManager->AddLocalWorkspaceMenuCategory(LOCTEXT("WorkspaceMenu_GeometryCollectionEditor", "Dataflow Editor"));
+
+	InTabManager->RegisterTabSpawner(GraphCanvasTabId, FOnSpawnTab::CreateSP(this, &FGeometryCollectionEditorToolkit::SpawnTab_GraphCanvas))
+		.SetDisplayName(LOCTEXT("DataflowTab", "Graph"))
+		.SetGroup(WorkspaceMenuCategoryRef)
+		.SetIcon(FSlateIcon(FAppStyle::GetAppStyleSetName(), "GraphEditor.EventGraph_16x"));
+
+	InTabManager->RegisterTabSpawner(PropertiesTabId, FOnSpawnTab::CreateSP(this, &FGeometryCollectionEditorToolkit::SpawnTab_Properties))
+		.SetDisplayName(LOCTEXT("PropertiesTab", "Details"))
+		.SetGroup(WorkspaceMenuCategoryRef)
+		.SetIcon(FSlateIcon(FAppStyle::GetAppStyleSetName(), "LevelEditor.Tabs.Details"));
+
+	FAssetEditorToolkit::RegisterTabSpawners(InTabManager);
+}
+
+
+FName FGeometryCollectionEditorToolkit::GetToolkitFName() const
+{
+	return FName("GeometryCollectionEditor");
+}
+
+FText FGeometryCollectionEditorToolkit::GetBaseToolkitName() const
+{
+	return LOCTEXT("AppLabel", "Geometry Collection Editor");
+}
+
+FString FGeometryCollectionEditorToolkit::GetWorldCentricTabPrefix() const
+{
+	return LOCTEXT("WorldCentricTabPrefix", "GeometryCollection").ToString();
+}
+
+FLinearColor FGeometryCollectionEditorToolkit::GetWorldCentricTabColorScale() const
+{
+	return FLinearColor(0.3f, 0.2f, 0.5f, 0.5f);
+}
+
+FString FGeometryCollectionEditorToolkit::GetReferencerName() const
+{
+	return TEXT("GeometryCollectionEditorToolkit");
+}
+
+void FGeometryCollectionEditorToolkit::AddReferencedObjects(FReferenceCollector& Collector)
+{
+	if (Dataflow)
+	{
+		Collector.AddReferencedObject(Dataflow);
+	}
+	if (GeometryCollection)
+	{
+		Collector.AddReferencedObject(GeometryCollection);
+	}
+}
+#undef LOCTEXT_NAMESPACE
diff --git a/Engine/Plugins/Experimental/GeometryCollectionPlugin/Source/GeometryCollectionEditor/Private/GeometryCollection/Dataflow/GeometryCollectionNodeConnectionTypes.h b/Engine/Plugins/Experimental/GeometryCollectionPlugin/Source/GeometryCollectionEditor/Private/GeometryCollection/Dataflow/GeometryCollectionNodeConnectionTypes.h
new file mode 100644
--- /dev/null
+++ b/Engine/Plugins/Experimental/GeometryCollectionPlugin/Source/GeometryCollectionEditor/Private/GeometryCollection/Dataflow/GeometryCollectionNodeConnectionTypes.h
@@ -0,0 +1,17 @@
+// Copyright Epic Games, Inc. All Rights Reserved.
+
+#pragma once
+
+#include "CoreMinimal.h"
+
+#include "Dataflow/DataflowCore.h"
+#include "GeometryCollection/ManagedArrayCollection.h"
+
+
+namespace Dataflow
+{
+	typedef TSharedPtr<FManagedArrayCollection> FManagedArrayCollectionSharedPtr;
+	 
+	DATAFLOW_CONNECTION_TYPE_SHARED_POINTER(GEOMETRYCOLLECTIONEDITOR_API, FManagedArrayCollectionSharedPtr, ManagedArrayCollectionSharedPtr)
+	
+}
diff --git a/Engine/Plugins/Experimental/GeometryCollectionPlugin/Source/GeometryCollectionEditor/Private/GeometryCollection/Dataflow/GeometryCollectionNodes.cpp b/Engine/Plugins/Experimental/GeometryCollectionPlugin/Source/GeometryCollectionEditor/Private/GeometryCollection/Dataflow/GeometryCollectionNodes.cpp
new file mode 100644
--- /dev/null
+++ b/Engine/Plugins/Experimental/GeometryCollectionPlugin/Source/GeometryCollectionEditor/Private/GeometryCollection/Dataflow/GeometryCollectionNodes.cpp
@@ -0,0 +1,14 @@
+// Copyright Epic Games, Inc. All Rights Reserved.
+
+#include "GeometryCollectionNodes.h"
+
+namespace Dataflow
+{
+	void GeometryCollectionEngineAssetNodes()
+	{
+		DATAFLOW_NODE_REGISTER_CREATION_FACTORY(GetCollectionAssetNode);
+		DATAFLOW_NODE_REGISTER_CREATION_FACTORY(ExampleCollectionEditNode);
+		DATAFLOW_NODE_REGISTER_CREATION_FACTORY(SetCollectionAssetNode);
+	}
+}
+
diff --git a/Engine/Plugins/Experimental/GeometryCollectionPlugin/Source/GeometryCollectionEditor/Private/GeometryCollection/Dataflow/GeometryCollectionNodes.h b/Engine/Plugins/Experimental/GeometryCollectionPlugin/Source/GeometryCollectionEditor/Private/GeometryCollection/Dataflow/GeometryCollectionNodes.h
new file mode 100644
--- /dev/null
+++ b/Engine/Plugins/Experimental/GeometryCollectionPlugin/Source/GeometryCollectionEditor/Private/GeometryCollection/Dataflow/GeometryCollectionNodes.h
@@ -0,0 +1,122 @@
+// Copyright Epic Games, Inc. All Rights Reserved.
+
+#pragma once 
+
+#include "GeometryCollectionNodeConnectionTypes.h"
+
+#include "CoreMinimal.h"
+#include "GeometryCollection/Dataflow/GeometryCollectionEditorToolkit.h"
+#include "Dataflow/DataflowObjectInterface.h"
+#include "Dataflow/DataflowEngine.h"
+#include "Dataflow/DataflowProperty.h"
+#include "GeometryCollection/GeometryCollectionObject.h"
+#include "GeometryCollection/ManagedArrayCollection.h"
+#include "GeometryCollection/GeometryCollection.h"
+
+
+namespace Dataflow
+{
+
+	class GetCollectionAssetNode : public FNode
+	{
+		DATAFLOW_NODE_DEFINE_INTERNAL(GetCollectionAssetNode)
+
+	public:
+		typedef FManagedArrayCollectionSharedPtr DataType;
+
+		TSharedPtr< class TOutput<DataType> > Output;
+
+		GetCollectionAssetNode(const FNodeParameters& InParam, FGuid InGuid = FGuid::NewGuid())
+			: FNode(InParam, InGuid)
+			, Output(new TOutput<DataType>(TOutputParameters<DataType>({ FName("CollectionOut"), this })))
+		{}
+
+		virtual void Evaluate(const FContext& Context, FConnection* Out) const override
+		{
+			if (Output.Get() == Out)
+			{
+				const FGeometryCollectionContext* EngineContext = (const FGeometryCollectionContext*)(&Context);
+				if (UGeometryCollection* CollectionAsset = EngineContext->Asset)
+				{
+					if (const TSharedPtr<FGeometryCollection, ESPMode::ThreadSafe> AssetCollection = CollectionAsset->GetGeometryCollection())
+					{
+						FManagedArrayCollection* NewCollection = AssetCollection->NewCopy<FManagedArrayCollection>();
+						Output->SetValue(FManagedArrayCollectionSharedPtr(NewCollection), Context);
+					}
+				}
+			}
+		}
+	};
+
+
+	class ExampleCollectionEditNode : public FNode
+	{
+		DATAFLOW_NODE_DEFINE_INTERNAL(ExampleCollectionEditNode)
+
+	public:
+		typedef FManagedArrayCollectionSharedPtr DataType;
+		TSharedPtr< class TInput<DataType> > Input;
+		TSharedPtr< class TOutput<DataType> > Output;
+		TProperty<bool> Active;
+		TProperty<float> Scale;
+
+		ExampleCollectionEditNode(const FNodeParameters& InParam, FGuid InGuid = FGuid::NewGuid())
+			: FNode(InParam, InGuid)
+			, Input(new TInput<DataType>(TInputParameters<DataType>({ FName("CollectionIn"), this })))
+			, Output(new TOutput<DataType>(TOutputParameters<DataType>({ FName("CollectionOut"), this })))
+			, Active({ FName("Active"), true, this })
+			, Scale({ FName("Scale"), 10.f, this })
+		{}
+
+		virtual void Evaluate(const FContext& Context, FConnection* Out) const override
+		{
+			if (Output.Get() == Out)
+			{
+				DataType Collection = Input->GetValue(Context);
+				if (Active.GetValue())
+				{
+					TManagedArray<FVector3f>* Vertex = Collection->FindAttribute<FVector3f>("Vertex", "Vertices");
+					for (int i = 0; i < Vertex->Num(); i++)
+					{
+						(*Vertex)[i][1] *= Scale.GetValue();
+					}
+				}
+				Output->SetValue(Collection,Context);
+			}
+		}
+	};
+
+	class SetCollectionAssetNode : public FNode
+	{
+		DATAFLOW_NODE_DEFINE_INTERNAL(SetCollectionAssetNode)
+
+	public:
+		typedef FManagedArrayCollectionSharedPtr DataType;
+		TSharedPtr< class TInput<DataType> > Input;
+
+		SetCollectionAssetNode(const FNodeParameters& InParam, FGuid InGuid = FGuid::NewGuid())
+			: FNode(InParam, InGuid)
+			, Input(new TInput<DataType>(TInputParameters<DataType>({ FName("CollectionIn"), this })))
+		{}
+
+		virtual void Evaluate(const FContext& Context, FConnection* Out) const override
+		{
+			if (Out==nullptr)
+			{
+				DataType Collection = Input->GetValue(Context);
+				const FGeometryCollectionContext* EngineContext = (FGeometryCollectionContext*)(&Context);
+				if (UGeometryCollection* CollectionAsset = EngineContext->Asset)
+				{
+					TSharedPtr<FGeometryCollection, ESPMode::ThreadSafe> NewCollection(Collection->NewCopy<FGeometryCollection>());
+					CollectionAsset->SetGeometryCollection(NewCollection);
+				}
+			}
+		}
+	};
+
+
+	void GeometryCollectionEngineAssetNodes();
+
+
+}
+
diff --git a/Engine/Plugins/Experimental/GeometryCollectionPlugin/Source/GeometryCollectionEditor/Private/GeometryCollection/GeometryCollectionEditorPlugin.cpp b/Engine/Plugins/Experimental/GeometryCollectionPlugin/Source/GeometryCollectionEditor/Private/GeometryCollection/GeometryCollectionEditorPlugin.cpp
--- a/Engine/Plugins/Experimental/GeometryCollectionPlugin/Source/GeometryCollectionEditor/Private/GeometryCollection/GeometryCollectionEditorPlugin.cpp
+++ b/Engine/Plugins/Experimental/GeometryCollectionPlugin/Source/GeometryCollectionEditor/Private/GeometryCollection/GeometryCollectionEditorPlugin.cpp
@@ -5,6 +5,7 @@
 
 #include "AssetToolsModule.h"
 #include "CoreMinimal.h"
+#include "Dataflow/GeometryCollectionNodes.h"
 #include "LevelEditor.h"
 #include "Modules/ModuleManager.h"
 #include "PropertyEditorModule.h"
@@ -18,6 +19,7 @@
 #include "GeometryCollection/GeometryCollectionThumbnailRenderer.h"
 #include "GeometryCollection/GeometryCollectionSelectRigidBodyEdMode.h"
 #include "GeometryCollection/GeometryCollectionSelectionCommands.h"
+#include "GeometryCollection/Dataflow/GeometryCollectionEditorToolkit.h"
 #include "HAL/ConsoleManager.h"
 #include "Features/IModularFeatures.h"
 #include "GeometryCollection/DetailCustomizations/GeomComponentCacheCustomization.h"
@@ -27,6 +29,8 @@
 #include "Styling/SlateStyleRegistry.h"
 #include "Styling/CoreStyle.h"
 #include "ToolMenus.h"
+#include "Toolkits/AssetEditorToolkit.h"
+#include "Toolkits/IToolkitHost.h"
 
 IMPLEMENT_MODULE( IGeometryCollectionEditorPlugin, GeometryCollectionEditor )
 
@@ -257,6 +261,10 @@ void IGeometryCollectionEditorPlugin::StartupModule()
 	// Register tool menus - delayed until ToolMenus module is loaded
 	UToolMenus::RegisterStartupCallback(
 		FSimpleMulticastDelegate::FDelegate::CreateRaw(this, &IGeometryCollectionEditorPlugin::RegisterMenus));
+
+	// Dataflow 
+	Dataflow::GeometryCollectionEngineAssetNodes();
+
 }
 
 
@@ -329,4 +337,12 @@ void IGeometryCollectionEditorPlugin::RegisterMenus()
 		Section.AddMenuEntry(FGeometryCollectionSelectionCommands::Get().SelectNone);
 		Section.AddMenuEntry(FGeometryCollectionSelectionCommands::Get().SelectInverseGeometry);
 	}
+}
+
+
+TSharedRef<FAssetEditorToolkit> IGeometryCollectionEditorPlugin::CreateGeometryCollectionAssetEditor(const EToolkitMode::Type Mode, const TSharedPtr<IToolkitHost>& InitToolkitHost, UObject* GeometryCollectionAsset)
+{
+	TSharedPtr<FGeometryCollectionEditorToolkit> NewGeometryCollectionAssetEditor = MakeShared<FGeometryCollectionEditorToolkit>();
+	NewGeometryCollectionAssetEditor->InitGeometryCollectionAssetEditor(Mode, InitToolkitHost, GeometryCollectionAsset);
+	return StaticCastSharedPtr<FAssetEditorToolkit>(NewGeometryCollectionAssetEditor).ToSharedRef();
 }
\ No newline at end of file
diff --git a/Engine/Plugins/Experimental/GeometryCollectionPlugin/Source/GeometryCollectionEditor/Public/GeometryCollection/AssetTypeActions_GeometryCollection.h b/Engine/Plugins/Experimental/GeometryCollectionPlugin/Source/GeometryCollectionEditor/Public/GeometryCollection/AssetTypeActions_GeometryCollection.h
--- a/Engine/Plugins/Experimental/GeometryCollectionPlugin/Source/GeometryCollectionEditor/Public/GeometryCollection/AssetTypeActions_GeometryCollection.h
+++ b/Engine/Plugins/Experimental/GeometryCollectionPlugin/Source/GeometryCollectionEditor/Public/GeometryCollection/AssetTypeActions_GeometryCollection.h
@@ -14,7 +14,7 @@ public:
 	virtual void GetActions(const TArray<UObject*>& InObjects, struct FToolMenuSection& Section) override;
 	virtual FColor GetTypeColor() const override { return FColor(255, 192, 128); }
 	virtual UClass* GetSupportedClass() const override;
-	//virtual void OpenAssetEditor(const TArray<UObject*>& InObjects, TSharedPtr<class IToolkitHost> EditWithinLevelEditor = TSharedPtr<IToolkitHost>()) override;
+	virtual void OpenAssetEditor(const TArray<UObject*>& InObjects, TSharedPtr<class IToolkitHost> EditWithinLevelEditor = TSharedPtr<IToolkitHost>()) override;
 	virtual uint32 GetCategories() override { return EAssetTypeCategories::Physics; }
 	virtual UThumbnailInfo* GetThumbnailInfo(UObject* Asset) const override;
 };
diff --git a/Engine/Plugins/Experimental/GeometryCollectionPlugin/Source/GeometryCollectionEditor/Public/GeometryCollection/Dataflow/GeometryCollectionEditorToolkit.h b/Engine/Plugins/Experimental/GeometryCollectionPlugin/Source/GeometryCollectionEditor/Public/GeometryCollection/Dataflow/GeometryCollectionEditorToolkit.h
new file mode 100644
--- /dev/null
+++ b/Engine/Plugins/Experimental/GeometryCollectionPlugin/Source/GeometryCollectionEditor/Public/GeometryCollection/Dataflow/GeometryCollectionEditorToolkit.h
@@ -0,0 +1,78 @@
+// Copyright Epic Games, Inc. All Rights Reserved.
+
+#pragma once
+
+#include "CoreMinimal.h"
+#include "Dataflow/DataflowNodeParameters.h"
+#include "GeometryCollection/GeometryCollectionObject.h"
+#include "GraphEditor.h"
+#include "Toolkits/SimpleAssetEditor.h"
+#include "Misc/NotifyHook.h"
+
+class IDetailsView;
+class FTabManager;
+class IToolkitHost;
+class UDataflow;
+
+namespace Dataflow
+{
+	class GEOMETRYCOLLECTIONEDITOR_API FGeometryCollectionContext : public FContext
+	{
+	public:
+		UGeometryCollection* Asset;
+
+		FGeometryCollectionContext(UGeometryCollection* InAsset, float InTime)
+			: FContext(InTime)
+			, Asset(InAsset)
+		{}
+
+	};
+}
+
+class FGeometryCollectionEditorToolkit : public FAssetEditorToolkit, public FNotifyHook, public FGCObject
+{
+public:
+	void InitGeometryCollectionAssetEditor(const EToolkitMode::Type Mode, const TSharedPtr<IToolkitHost>& InitToolkitHost, UObject* ObjectToEdit);
+
+	// IToolkit Interface
+	virtual FName GetToolkitFName() const override;
+	virtual FText GetBaseToolkitName() const override;
+	virtual FString GetWorldCentricTabPrefix() const override;
+	virtual FLinearColor GetWorldCentricTabColorScale() const override;
+
+	// FGCObject Interface
+	virtual void AddReferencedObjects(FReferenceCollector& Collector) override;
+	virtual FString GetReferencerName() const override; 
+
+	// Tab spawners 
+	virtual void RegisterTabSpawners(const TSharedRef<FTabManager>& TabManager) override;
+	TSharedRef<SDockTab> SpawnTab_GraphCanvas(const FSpawnTabArgs& Args);
+	TSharedRef<SDockTab> SpawnTab_Properties(const FSpawnTabArgs& Args);
+
+	// Graph Editor Operations
+	void OnSelectedNodesChanged(const TSet<UObject*>& NewSelection);
+	void DeleteNode();
+	void EvaluateNode();
+	FGraphPanelSelectionSet GetSelectedNodes() const;
+
+	UDataflow* GetDataflow() { return Dataflow; }
+	const UDataflow* GetDataflow() const { return Dataflow; }
+
+private:
+
+
+	static const FName GraphCanvasTabId;
+	TSharedPtr<SGraphEditor> GraphEditor;
+	TSharedRef<SGraphEditor> CreateGraphEditorWidget(UDataflow* ObjectToEdit);
+
+	static const FName PropertiesTabId;
+	TSharedPtr<IDetailsView> PropertiesEditor;
+	TSharedPtr<IDetailsView> CreatePropertiesEditorWidget(UObject* ObjectToEdit);
+
+
+	UDataflow* Dataflow = nullptr;
+	UGeometryCollection* GeometryCollection = nullptr;
+
+	/** Command list for the graph editor */
+	TSharedPtr<FUICommandList> GraphEditorCommands;
+};
diff --git a/Engine/Plugins/Experimental/GeometryCollectionPlugin/Source/GeometryCollectionEditor/Public/GeometryCollection/GeometryCollectionEditorPlugin.h b/Engine/Plugins/Experimental/GeometryCollectionPlugin/Source/GeometryCollectionEditor/Public/GeometryCollection/GeometryCollectionEditorPlugin.h
--- a/Engine/Plugins/Experimental/GeometryCollectionPlugin/Source/GeometryCollectionEditor/Public/GeometryCollection/GeometryCollectionEditorPlugin.h
+++ b/Engine/Plugins/Experimental/GeometryCollectionPlugin/Source/GeometryCollectionEditor/Public/GeometryCollection/GeometryCollectionEditorPlugin.h
@@ -6,6 +6,8 @@
 #include "Modules/ModuleInterface.h"
 #include "Modules/ModuleManager.h"
 #include "GeometryCollection/GeometryCollectionProviderEditor.h"
+#include "Templates/SharedPointer.h"
+#include "Toolkits/IToolkit.h"
 
 
 class ISlateStyle;
@@ -13,6 +15,8 @@ class FSlateStyleSet;
 class FAssetTypeActions_GeometryCollection;
 class FAssetTypeActions_GeometryCollectionCache;
 class FGeometryCollectionAssetBroker;
+class FAssetEditorToolkit;
+class IToolkitHost;
 
 /**
  * The public interface to this module
@@ -23,6 +27,12 @@ public:
 	virtual void StartupModule();
 	virtual void ShutdownModule();
 
+	/**
+	*  Create a custom asset editor.
+	*
+	*/
+	TSharedRef<FAssetEditorToolkit> CreateGeometryCollectionAssetEditor(const EToolkitMode::Type Mode, const TSharedPtr<IToolkitHost>& InitToolkitHost, UObject* GeometryCollectionAsset);
+
 
 	/**
 	 * Singleton-like access to this module's interface.  This is just for convenience!
diff --git a/Engine/Source/Runtime/Experimental/GeometryCollectionEngine/GeometryCollectionEngine.Build.cs b/Engine/Source/Runtime/Experimental/GeometryCollectionEngine/GeometryCollectionEngine.Build.cs
--- a/Engine/Source/Runtime/Experimental/GeometryCollectionEngine/GeometryCollectionEngine.Build.cs
+++ b/Engine/Source/Runtime/Experimental/GeometryCollectionEngine/GeometryCollectionEngine.Build.cs
@@ -23,9 +23,11 @@ namespace UnrealBuildTool.Rules
                     "FieldSystemEngine",
 	                "ChaosSolverEngine",
 					"NetCore",
-                    "IntelISPC"
-                }
-                );
+					"IntelISPC",
+					"DataflowCore",
+					"DataflowEngine"
+				}
+				);
 
 			PrivateIncludePathModuleNames.Add("DerivedDataCache");
 
diff --git a/Engine/Source/Runtime/Experimental/GeometryCollectionEngine/Public/GeometryCollection/GeometryCollectionObject.h b/Engine/Source/Runtime/Experimental/GeometryCollectionEngine/Public/GeometryCollection/GeometryCollectionObject.h
--- a/Engine/Source/Runtime/Experimental/GeometryCollectionEngine/Public/GeometryCollection/GeometryCollectionObject.h
+++ b/Engine/Source/Runtime/Experimental/GeometryCollectionEngine/Public/GeometryCollection/GeometryCollectionObject.h
@@ -7,6 +7,7 @@
 #include "Rendering/NaniteResources.h"
 #include "InstanceUniformShaderParameters.h"
 #include "GeometryCollection/ManagedArray.h"
+#include "GeometryCollection/GeometryCollectionObject.h"
 #include "GeometryCollection/GeometryCollectionSimulationTypes.h"
 #include "Chaos/ChaosSolverActor.h"
 
@@ -18,6 +19,7 @@ class FGeometryCollection;
 class FManagedArrayCollection;
 struct FGeometryCollectionSection;
 struct FSharedSimulationParameters;
+class UDataflow;
 
 USTRUCT(BlueprintType)
 struct GEOMETRYCOLLECTIONENGINE_API FGeometryCollectionSource
@@ -606,6 +608,13 @@ public:
 	*/
 	void UpdateConvexGeometry();
 
+
+	//
+	// Dataflow
+	//
+	UPROPERTY(EditAnywhere, Category = "Procedural")
+	TObjectPtr<UDataflow> Dataflow;
+
 private:
 #if WITH_EDITOR
 	void CreateSimulationDataImp(bool bCopyFromDDC);
