Dataflow - Add specialized GraphEdtior

diff --git a/Engine/Plugins/Experimental/Dataflow/Source/DataflowEditor/Private/Dataflow/DataflowEditorActions.cpp b/Engine/Plugins/Experimental/Dataflow/Source/DataflowEditor/Private/Dataflow/DataflowEditorActions.cpp
--- a/Engine/Plugins/Experimental/Dataflow/Source/DataflowEditor/Private/Dataflow/DataflowEditorActions.cpp
+++ b/Engine/Plugins/Experimental/Dataflow/Source/DataflowEditor/Private/Dataflow/DataflowEditorActions.cpp
@@ -22,8 +22,8 @@
 				this->AsShared(),
 				AddNode,
 				NodeName, //FName("UseCreationFormToggle"),
-				NSLOCTEXT("DataFlow", "DataflowButton", "New Dataflow Node"),
-				NSLOCTEXT("DataFlow", "NewDataflowNodeTooltip", "New Dataflow Node Tooltip"),
+				NSLOCTEXT("DataFlow1", "DataflowButton", "New Dataflow Node"),
+				NSLOCTEXT("DataFlow2", "NewDataflowNodeTooltip", "New Dataflow Node Tooltip"),
 				FSlateIcon(),
 				EUserInterfaceActionType::Button,
 				FInputChord()
@@ -50,7 +50,7 @@
 
 
 
-void FDataflowEditorCommands::EvaluateNodes(const FGraphPanelSelectionSet& SelectedNodes, const Dataflow::FContext& InContext)
+void FDataflowEditorCommands::EvaluateNodes(const FGraphPanelSelectionSet& SelectedNodes, FDataflowEditorCommands::FGraphEvaluationCallback Evaluate)
 {
 	for (UObject* Ode : SelectedNodes)
 	{
@@ -64,12 +64,12 @@
 					{
 						for (Dataflow::FConnection* NodeOutput : DataflowNode->GetOutputs())
 						{
-							DataflowNode->Evaluate(InContext, NodeOutput);
+							Evaluate(DataflowNode.Get(),NodeOutput);
 						}
 					}
 					else
 					{
-						DataflowNode->Evaluate(InContext, nullptr);
+						Evaluate(DataflowNode.Get(), nullptr);
 					}
 				}
 			}
diff --git a/Engine/Plugins/Experimental/Dataflow/Source/DataflowEditor/Private/Dataflow/DataflowEditorToolkit.cpp b/Engine/Plugins/Experimental/Dataflow/Source/DataflowEditor/Private/Dataflow/DataflowEditorToolkit.cpp
--- a/Engine/Plugins/Experimental/Dataflow/Source/DataflowEditor/Private/Dataflow/DataflowEditorToolkit.cpp
+++ b/Engine/Plugins/Experimental/Dataflow/Source/DataflowEditor/Private/Dataflow/DataflowEditorToolkit.cpp
@@ -3,6 +3,7 @@
 #include "DataflowEditorToolkit.h"
 
 #include "Dataflow/DataflowEditorActions.h"
+#include "Dataflow/DataflowGraphEditor.h"
 #include "Dataflow/DataflowEdNode.h"
 #include "Dataflow/DataflowNodeFactory.h"
 #include "Dataflow/DataflowObject.h"
@@ -20,7 +21,6 @@
 
 //DEFINE_LOG_CATEGORY_STATIC(FDataflowEditorToolkitLog, Log, All);
 
-
 const FName FDataflowEditorToolkit::GraphCanvasTabId(TEXT("DataflowEditor_GraphCanvas"));
 const FName FDataflowEditorToolkit::PropertiesTabId(TEXT("DataflowEditor_Properties"));
 
@@ -74,73 +74,12 @@ void FDataflowEditorToolkit::InitDataflowEditor(const EToolkitMode::Type Mode, c
 
 
 
-void FDataflowEditorToolkit::EvaluateNode()
-{
-	Dataflow::FContext Context(FGameTime::GetTimeSinceAppStart().GetRealTimeSeconds());
-	FDataflowEditorCommands::EvaluateNodes(GetSelectedNodes(), Context);
-}
-
-void FDataflowEditorToolkit::DeleteNode()
-{
-	if (UDataflow* Graph = dynamic_cast<UDataflow*>(GraphEditor->GetCurrentGraph()))
-	{
-		FDataflowEditorCommands::DeleteNodes(Graph, GetSelectedNodes());
-	}
-}
-
-FGraphPanelSelectionSet FDataflowEditorToolkit::GetSelectedNodes() const
-{
-	if (GraphEditor.IsValid())
-	{
-		return GraphEditor->GetSelectedNodes();
-	}
-	return FGraphPanelSelectionSet();
-}
-
-void FDataflowEditorToolkit::OnSelectedNodesChanged(const TSet<UObject*>& NewSelection)
-{
-	if (UDataflow* Graph = dynamic_cast<UDataflow*>(GraphEditor->GetCurrentGraph()))
-	{
-		FDataflowEditorCommands::OnSelectedNodesChanged(PropertiesEditor, Dataflow, Graph, NewSelection);
-	}
-}
-
-
 TSharedRef<SGraphEditor> FDataflowEditorToolkit::CreateGraphEditorWidget(UDataflow* DataflowToEdit)
 {
 	ensure(DataflowToEdit);
-
-	FDataflowEditorCommands::Register();
-	FGraphEditorCommands::Register();
-
-	// No need to regenerate the commands.
-	if (!GraphEditorCommands.IsValid())
-	{
-		GraphEditorCommands = MakeShareable(new FUICommandList);
-		{
-			GraphEditorCommands->MapAction(FGenericCommands::Get().Delete,
-				FExecuteAction::CreateSP(this, &FDataflowEditorToolkit::DeleteNode)
-			);
-			GraphEditorCommands->MapAction(FDataflowEditorCommands::Get().EvaluateNode,
-				FExecuteAction::CreateSP(this, &FDataflowEditorToolkit::EvaluateNode)
-			);
-		}
-	}
-
-
-	FGraphAppearanceInfo AppearanceInfo;
-	AppearanceInfo.CornerText = LOCTEXT("AppearanceCornerText_DataflowEditor", "Dataflow");
-
-	SGraphEditor::FGraphEditorEvents InEvents;
-	InEvents.OnSelectionChanged = SGraphEditor::FOnSelectionChanged::CreateSP(this, &FDataflowEditorToolkit::OnSelectedNodesChanged);
-
-	return SNew(SGraphEditor)
-		.AdditionalCommands(GraphEditorCommands)
-		.IsEditable(true)
-		.Appearance(AppearanceInfo)
+	return SNew(SDataflowGraphEditor, DataflowToEdit)
 		.GraphToEdit(DataflowToEdit)
-		.GraphEvents(InEvents)
-		.ShowGraphStateOverlay(false);
+		.DetailsView(GetPropertiesEditor());
 }
 
 TSharedPtr<IDetailsView> FDataflowEditorToolkit::CreatePropertiesEditorWidget(UObject* ObjectToEdit)
diff --git a/Engine/Plugins/Experimental/Dataflow/Source/DataflowEditor/Private/Dataflow/DataflowEditorToolkit.h b/Engine/Plugins/Experimental/Dataflow/Source/DataflowEditor/Private/Dataflow/DataflowEditorToolkit.h
--- a/Engine/Plugins/Experimental/Dataflow/Source/DataflowEditor/Private/Dataflow/DataflowEditorToolkit.h
+++ b/Engine/Plugins/Experimental/Dataflow/Source/DataflowEditor/Private/Dataflow/DataflowEditorToolkit.h
@@ -32,28 +32,26 @@ public:
 	TSharedRef<SDockTab> SpawnTab_GraphCanvas(const FSpawnTabArgs& Args);
 	TSharedRef<SDockTab> SpawnTab_Properties(const FSpawnTabArgs& Args);
 
-	// Graph Editor Operations
-	void OnSelectedNodesChanged(const TSet<UObject*>& NewSelection);
-	void DeleteNode();
-	void EvaluateNode();
-	FGraphPanelSelectionSet GetSelectedNodes() const;
-
+	// Member Access
 	UDataflow* GetDataflow() { return Dataflow; }
 	const UDataflow* GetDataflow() const { return Dataflow; }
 
+	TSharedPtr<IDetailsView> GetPropertiesEditor() {return PropertiesEditor;}
+	const TSharedPtr<IDetailsView> GetPropertiesEditor() const { return PropertiesEditor; }
+
+	TSharedPtr<SGraphEditor> GetGraphEditor() { return GraphEditor; }
+	const TSharedPtr<SGraphEditor> GetGraphEditor() const { return GraphEditor; }
+
 private:
 
+	UDataflow* Dataflow = nullptr;
+
 	static const FName GraphCanvasTabId;
 	TSharedPtr<SGraphEditor> GraphEditor;
+	TSharedPtr<FUICommandList> GraphEditorCommands;
 	TSharedRef<SGraphEditor> CreateGraphEditorWidget(UDataflow* ObjectToEdit);
 
 	static const FName PropertiesTabId;
 	TSharedPtr<IDetailsView> PropertiesEditor;
 	TSharedPtr<IDetailsView> CreatePropertiesEditorWidget(UObject* ObjectToEdit);
-
-
-	UDataflow* Dataflow = nullptr;
-
-	/** Command list for the graph editor */
-	TSharedPtr<FUICommandList> GraphEditorCommands;
 };
diff --git a/Engine/Plugins/Experimental/Dataflow/Source/DataflowEditor/Private/Dataflow/DataflowGraphEditor.cpp b/Engine/Plugins/Experimental/Dataflow/Source/DataflowEditor/Private/Dataflow/DataflowGraphEditor.cpp
new file mode 100644
--- /dev/null
+++ b/Engine/Plugins/Experimental/Dataflow/Source/DataflowEditor/Private/Dataflow/DataflowGraphEditor.cpp
@@ -0,0 +1,100 @@
+// Copyright Epic Games, Inc. All Rights Reserved.
+
+#include "Dataflow/DataflowGraphEditor.h"
+
+#include "DataflowEditorToolkit.h"
+#include "Dataflow/DataflowEngine.h"
+#include "Dataflow/DataflowSNodeFactories.h"
+
+
+#define LOCTEXT_NAMESPACE "DataflowGraphEditor"
+
+
+void SDataflowGraphEditor::Construct(const FArguments& InArgs, UObject* InAssetOwner)
+{
+	check(InArgs._GraphToEdit);
+	AssetOwner = InAssetOwner; // nullptr is valid
+	DataflowAsset = Cast<UDataflow>(InArgs._GraphToEdit);
+	DetailsView = InArgs._DetailsView;
+	EvaluateGraphCallback = InArgs._EvaluateGraph;
+
+	FGraphAppearanceInfo AppearanceInfo;
+	AppearanceInfo.CornerText = FText::FromString("Dataflow");
+
+	FGraphEditorCommands::Register();
+	FDataflowEditorCommands::Register();
+	if (!GraphEditorCommands.IsValid())
+	{
+		GraphEditorCommands = MakeShareable(new FUICommandList);
+		{
+			GraphEditorCommands->MapAction(FGenericCommands::Get().Delete,
+				FExecuteAction::CreateSP(this, &SDataflowGraphEditor::DeleteNode)
+			);
+			GraphEditorCommands->MapAction(FDataflowEditorCommands::Get().EvaluateNode,
+				FExecuteAction::CreateSP(this, &SDataflowGraphEditor::EvaluateNode)
+			);
+		}
+	}
+
+
+	SGraphEditor::FArguments Arguments;
+	Arguments._AdditionalCommands = GraphEditorCommands;
+	Arguments._Appearance = AppearanceInfo;
+	Arguments._GraphToEdit = InArgs._GraphToEdit;
+	Arguments._GraphEvents = InArgs._GraphEvents;
+
+	ensureMsgf(!Arguments._GraphEvents.OnSelectionChanged.IsBound(), TEXT("DataflowGraphEditor::OnSelectionChanged rebound during construction."));
+	Arguments._GraphEvents.OnSelectionChanged = SGraphEditor::FOnSelectionChanged::CreateSP(this, &SDataflowGraphEditor::OnSelectedNodesChanged);
+
+	SGraphEditor::Construct(Arguments);
+
+}
+
+
+void SDataflowGraphEditor::EvaluateNode()
+{
+	float EvalTime = FGameTime::GetTimeSinceAppStart().GetRealTimeSeconds();
+	if (EvaluateGraphCallback)
+	{
+		FDataflowEditorCommands::EvaluateNodes(GetSelectedNodes(), EvaluateGraphCallback);
+	}
+	else
+	{
+		FDataflowEditorCommands::FGraphEvaluationCallback LocalEvaluateCallback = [](Dataflow::FNode* Node, Dataflow::FConnection* Out)
+		{
+			float EvalTime = FGameTime::GetTimeSinceAppStart().GetRealTimeSeconds();
+			return Node->Evaluate(Dataflow::FContext(EvalTime), Out);
+		};
+
+		FDataflowEditorCommands::EvaluateNodes(GetSelectedNodes(), LocalEvaluateCallback);
+	}
+}
+
+void SDataflowGraphEditor::DeleteNode()
+{
+	if (DataflowAsset.Get())
+	{
+		FDataflowEditorCommands::DeleteNodes(DataflowAsset.Get(), GetSelectedNodes());
+	}
+}
+
+void SDataflowGraphEditor::OnSelectedNodesChanged(const TSet<UObject*>& NewSelection)
+{
+	if (DataflowAsset.Get() && DetailsView)
+	{
+		FDataflowEditorCommands::OnSelectedNodesChanged(DetailsView, AssetOwner.Get(), DataflowAsset.Get(), NewSelection);
+	}
+}
+
+FReply SDataflowGraphEditor::OnDrop(const FGeometry& MyGeometry, const FDragDropEvent& DragDropEvent)
+{
+	//TSharedPtr<FDataprepDragDropOp> DragNodeOp = DragDropEvent.GetOperationAs<FDataprepDragDropOp>();
+	//if (TrackGraphNodePtr.IsValid() && DragNodeOp.IsValid())
+	//{
+	//	return FReply::Handled().EndDragDrop();
+	//}
+
+	return SGraphEditor::OnDrop(MyGeometry, DragDropEvent);
+}
+
+#undef LOCTEXT_NAMESPACE
diff --git a/Engine/Plugins/Experimental/Dataflow/Source/DataflowEditor/Public/Dataflow/DataflowEditorActions.h b/Engine/Plugins/Experimental/Dataflow/Source/DataflowEditor/Public/Dataflow/DataflowEditorActions.h
--- a/Engine/Plugins/Experimental/Dataflow/Source/DataflowEditor/Public/Dataflow/DataflowEditorActions.h
+++ b/Engine/Plugins/Experimental/Dataflow/Source/DataflowEditor/Public/Dataflow/DataflowEditorActions.h
@@ -5,12 +5,22 @@
 #include "CoreMinimal.h"
 #include "Framework/Commands/Commands.h"
 #include "Styling/AppStyle.h"
-#include "Dataflow/DataflowNodeParameters.h"
 
+namespace Dataflow 
+{
+	class FConnection;
+	class FNode;
+}
+
 class IDetailsView;
+class UDataflow;
 class UDataflowEdNode;
 typedef TSet<class UObject*> FGraphPanelSelectionSet;
 
+/*
+* FDataflowEditorCommandsImpl
+* 
+*/
 class FDataflowEditorCommandsImpl : public TCommands<FDataflowEditorCommandsImpl>
 {
 public:
@@ -33,6 +43,8 @@
 class DATAFLOWEDITOR_API FDataflowEditorCommands
 {
 public:
+	typedef TFunction<void(Dataflow::FNode*, Dataflow::FConnection*)> FGraphEvaluationCallback;
+
 	static void Register();
 	static void Unregister();
 
@@ -41,7 +53,7 @@
 	/*
 	*  EvaluateNodes
 	*/
-	static void EvaluateNodes(const FGraphPanelSelectionSet& SelectedNodes, const Dataflow::FContext& InContext);
+	static void EvaluateNodes(const FGraphPanelSelectionSet& SelectedNodes, FGraphEvaluationCallback);
 	
 	/*
 	*  DeleteNodes
diff --git a/Engine/Plugins/Experimental/Dataflow/Source/DataflowEditor/Public/Dataflow/DataflowGraphEditor.h b/Engine/Plugins/Experimental/Dataflow/Source/DataflowEditor/Public/Dataflow/DataflowGraphEditor.h
new file mode 100644
--- /dev/null
+++ b/Engine/Plugins/Experimental/Dataflow/Source/DataflowEditor/Public/Dataflow/DataflowGraphEditor.h
@@ -0,0 +1,86 @@
+// Copyright Epic Games, Inc. All Rights Reserved.
+
+#pragma once
+
+#include "CoreMinimal.h"
+
+#include "Dataflow/DataflowEditorActions.h"
+#include "EdGraphUtilities.h"
+#include "Framework/Commands/GenericCommands.h"
+#include "Framework/Commands/UICommandList.h"
+#include "GraphEditor.h"
+#include "GraphEditorActions.h"
+#include "Layout/SlateRect.h"
+#include "NodeFactory.h"
+#include "SGraphNode.h"
+#include "Widgets/DeclarativeSyntaxSupport.h"
+#include "Widgets/Layout/SConstraintCanvas.h"
+#include "Widgets/SWidget.h"
+
+
+class FDataflowEditorToolkit;
+class FDataflowSNodeFactory;
+class UDataflow;
+namespace Dataflow {
+	class FContext;
+	class FConnection;
+}
+/**
+ * The SDataflowGraphEditpr class is a specialization of SGraphEditor
+ * to display and manipulate the actions of a Dataflow asset
+ * 
+ * see(SDataprepGraphEditor for reference)
+ */
+class DATAFLOWEDITOR_API SDataflowGraphEditor : public SGraphEditor
+{
+public:
+
+	SLATE_BEGIN_ARGS(SDataflowGraphEditor)
+		: _AdditionalCommands(static_cast<FUICommandList*>(nullptr))
+		, _GraphToEdit(nullptr)
+	{}
+
+	SLATE_ARGUMENT(TSharedPtr<FUICommandList>, AdditionalCommands)
+	SLATE_ATTRIBUTE(FGraphAppearanceInfo, Appearance)
+	SLATE_ARGUMENT(UEdGraph*, GraphToEdit)
+	SLATE_ARGUMENT(FGraphEditorEvents, GraphEvents)
+	SLATE_ARGUMENT(TSharedPtr<IDetailsView>, DetailsView)
+	SLATE_ARGUMENT(FDataflowEditorCommands::FGraphEvaluationCallback, EvaluateGraph)
+	SLATE_END_ARGS()
+
+	// SWidget overrides
+	virtual FReply OnDrop(const FGeometry& MyGeometry, const FDragDropEvent& DragDropEvent) override;
+
+	/** */
+	void Construct(const FArguments& InArgs, UObject* AssetOwner);
+
+	/** */
+	void EvaluateNode();
+
+	/** */
+	void DeleteNode();
+
+	/** */
+	void OnSelectedNodesChanged(const TSet<UObject*>& NewSelection);
+
+
+
+
+private:
+	FDataflowEditorCommands::FGraphEvaluationCallback EvaluateGraphCallback;
+
+	/** The asset that ownes this dataflow graph */
+	TWeakObjectPtr<UObject> AssetOwner;
+
+	/** The dataflow asset associated with this graph */
+	TWeakObjectPtr<UDataflow> DataflowAsset;
+
+	/** Command list associated with this graph editor */
+	TSharedPtr<FUICommandList> GraphEditorCommands;
+
+	/** Factory to create the associated SGraphNode classes for Dataflow graph's UEdGraph classes */
+	static TSharedPtr<FDataflowSNodeFactory> NodeFactory;
+
+	/** The details view that responds to this widget. */
+	TSharedPtr<IDetailsView> DetailsView;
+};
diff --git a/Engine/Plugins/Experimental/GeometryCollectionPlugin/Source/GeometryCollectionEditor/Private/GeometryCollection/Dataflow/GeometryCollectionEditorToolkit.cpp b/Engine/Plugins/Experimental/GeometryCollectionPlugin/Source/GeometryCollectionEditor/Private/GeometryCollection/Dataflow/GeometryCollectionEditorToolkit.cpp
--- a/Engine/Plugins/Experimental/GeometryCollectionPlugin/Source/GeometryCollectionEditor/Private/GeometryCollection/Dataflow/GeometryCollectionEditorToolkit.cpp
+++ b/Engine/Plugins/Experimental/GeometryCollectionPlugin/Source/GeometryCollectionEditor/Private/GeometryCollection/Dataflow/GeometryCollectionEditorToolkit.cpp
@@ -3,7 +3,9 @@
 #include "GeometryCollection/Dataflow/GeometryCollectionEditorToolkit.h"
 
 #include "Dataflow/DataflowEditorActions.h"
+#include "Dataflow/DataflowEditorPlugin.h"
 #include "Dataflow/DataflowEdNode.h"
+#include "Dataflow/DataflowGraphEditor.h"
 #include "Dataflow/DataflowNodeFactory.h"
 #include "Dataflow/DataflowObject.h"
 #include "Dataflow/DataflowSchema.h"
@@ -30,6 +32,7 @@ void FGeometryCollectionEditorToolkit::InitGeometryCollectionAssetEditor(const E
 {
 	Dataflow = nullptr;
 	GeometryCollection = CastChecked<UGeometryCollection>(ObjectToEdit);
+
 	if (GeometryCollection != nullptr)
 	{
 		if (GeometryCollection->Dataflow == nullptr)
@@ -40,8 +43,8 @@ void FGeometryCollectionEditorToolkit::InitGeometryCollectionAssetEditor(const E
 		Dataflow = GeometryCollection->Dataflow;
 		GeometryCollection->Dataflow->Schema = UDataflowSchema::StaticClass();
 
-		GraphEditor = CreateGraphEditorWidget(Dataflow);
 		PropertiesEditor = CreatePropertiesEditorWidget(ObjectToEdit);
+		GraphEditor = CreateGraphEditorWidget(Dataflow, PropertiesEditor);
 
 		const TSharedRef<FTabManager::FLayout> StandaloneDefaultLayout = FTabManager::NewLayout("GeometryCollectionDataflowEditor_Layout")
 			->AddArea
@@ -81,73 +84,22 @@ void FGeometryCollectionEditorToolkit::InitGeometryCollectionAssetEditor(const E
 	}
 }
 
-void FGeometryCollectionEditorToolkit::EvaluateNode()
-{
-	Dataflow::FGeometryCollectionContext Context(GeometryCollection, FGameTime::GetTimeSinceAppStart().GetRealTimeSeconds());
-	FDataflowEditorCommands::EvaluateNodes(GetSelectedNodes(), Context);
-}
-
-void FGeometryCollectionEditorToolkit::DeleteNode()
-{
-	if (UDataflow* Graph = dynamic_cast<UDataflow*>(GraphEditor->GetCurrentGraph()))
-	{
-		FDataflowEditorCommands::DeleteNodes(Graph, GetSelectedNodes());
-	}
-}
-
-FGraphPanelSelectionSet FGeometryCollectionEditorToolkit::GetSelectedNodes() const
-{
-	if (GraphEditor.IsValid())
-	{
-		return GraphEditor->GetSelectedNodes();
-	}
-	return FGraphPanelSelectionSet();
-}
-
-void FGeometryCollectionEditorToolkit::OnSelectedNodesChanged(const TSet<UObject*>& NewSelection)
-{
-	if (UDataflow* Graph = dynamic_cast<UDataflow*>(GraphEditor->GetCurrentGraph()))
-	{
-		FDataflowEditorCommands::OnSelectedNodesChanged(PropertiesEditor, Dataflow, Graph, NewSelection);
-	}
-}
-
-
-TSharedRef<SGraphEditor> FGeometryCollectionEditorToolkit::CreateGraphEditorWidget(UDataflow* DataflowToEdit)
+TSharedRef<SGraphEditor> FGeometryCollectionEditorToolkit::CreateGraphEditorWidget(UDataflow* DataflowToEdit, TSharedPtr<IDetailsView> InPropertiesEditor)
 {
 	ensure(DataflowToEdit);
+	using namespace Dataflow;
+	IDataflowEditorPlugin& DataflowEditorModule = FModuleManager::LoadModuleChecked<IDataflowEditorPlugin>(TEXT("DataflowEditor"));
 
-	FDataflowEditorCommands::Register();
-	FGraphEditorCommands::Register();
-
-	// No need to regenerate the commands.
-	if (!GraphEditorCommands.IsValid())
+	FDataflowEditorCommands::FGraphEvaluationCallback Evaluate = [&](Dataflow::FNode* Node, Dataflow::FConnection* Out)
 	{
-		GraphEditorCommands = MakeShareable(new FUICommandList);
-		{
-			GraphEditorCommands->MapAction(FGenericCommands::Get().Delete,
-				FExecuteAction::CreateSP(this, &FGeometryCollectionEditorToolkit::DeleteNode)
-			);
-			GraphEditorCommands->MapAction(FDataflowEditorCommands::Get().EvaluateNode,
-				FExecuteAction::CreateSP(this, &FGeometryCollectionEditorToolkit::EvaluateNode)
-			);
-		}
-	}
-
-
-	FGraphAppearanceInfo AppearanceInfo;
-	AppearanceInfo.CornerText = LOCTEXT("AppearanceCornerText_GeometryCollectionEditor", "Dataflow");
-
-	SGraphEditor::FGraphEditorEvents InEvents;
-	InEvents.OnSelectionChanged = SGraphEditor::FOnSelectionChanged::CreateSP(this, &FGeometryCollectionEditorToolkit::OnSelectedNodesChanged);
+		float EvalTime = FGameTime::GetTimeSinceAppStart().GetRealTimeSeconds();
+		return Node->Evaluate(FGeometryCollectionContext(GeometryCollection, EvalTime), Out);
+	};
 
-	return SNew(SGraphEditor)
-		.AdditionalCommands(GraphEditorCommands)
-		.IsEditable(true)
-		.Appearance(AppearanceInfo)
+	return SNew(SDataflowGraphEditor, GeometryCollection)
 		.GraphToEdit(DataflowToEdit)
-		.GraphEvents(InEvents)
-		.ShowGraphStateOverlay(false);
+		.DetailsView(InPropertiesEditor)
+		.EvaluateGraph(Evaluate);
 }
 
 TSharedPtr<IDetailsView> FGeometryCollectionEditorToolkit::CreatePropertiesEditorWidget(UObject* ObjectToEdit)
diff --git a/Engine/Plugins/Experimental/GeometryCollectionPlugin/Source/GeometryCollectionEditor/Public/GeometryCollection/Dataflow/GeometryCollectionEditorToolkit.h b/Engine/Plugins/Experimental/GeometryCollectionPlugin/Source/GeometryCollectionEditor/Public/GeometryCollection/Dataflow/GeometryCollectionEditorToolkit.h
--- a/Engine/Plugins/Experimental/GeometryCollectionPlugin/Source/GeometryCollectionEditor/Public/GeometryCollection/Dataflow/GeometryCollectionEditorToolkit.h
+++ b/Engine/Plugins/Experimental/GeometryCollectionPlugin/Source/GeometryCollectionEditor/Public/GeometryCollection/Dataflow/GeometryCollectionEditorToolkit.h
@@ -4,6 +4,7 @@
 
 #include "CoreMinimal.h"
 #include "Dataflow/DataflowNodeParameters.h"
+#include "Dataflow/DataflowGraphEditor.h"
 #include "GeometryCollection/GeometryCollectionObject.h"
 #include "GraphEditor.h"
 #include "Toolkits/SimpleAssetEditor.h"
@@ -49,30 +50,27 @@ public:
 	TSharedRef<SDockTab> SpawnTab_GraphCanvas(const FSpawnTabArgs& Args);
 	TSharedRef<SDockTab> SpawnTab_Properties(const FSpawnTabArgs& Args);
 
-	// Graph Editor Operations
-	void OnSelectedNodesChanged(const TSet<UObject*>& NewSelection);
-	void DeleteNode();
-	void EvaluateNode();
-	FGraphPanelSelectionSet GetSelectedNodes() const;
-
 	UDataflow* GetDataflow() { return Dataflow; }
 	const UDataflow* GetDataflow() const { return Dataflow; }
 
-private:
+	UGeometryCollection* GetGeometryCollection() { return GeometryCollection; }
+	const UGeometryCollection* GetGeometryCollection() const { return GeometryCollection; }
+
+	TSharedPtr<IDetailsView> GetPropertiesEditor() { return PropertiesEditor; }
+	const TSharedPtr<IDetailsView> GetPropertiesEditor() const { return PropertiesEditor; }
 
+	TSharedPtr<SGraphEditor> GetGraphEditor() { return GraphEditor; }
+	const TSharedPtr<SGraphEditor> GetGraphEditor() const { return GraphEditor; }
 
+private:
 	static const FName GraphCanvasTabId;
 	TSharedPtr<SGraphEditor> GraphEditor;
-	TSharedRef<SGraphEditor> CreateGraphEditorWidget(UDataflow* ObjectToEdit);
+	TSharedRef<SGraphEditor> CreateGraphEditorWidget(UDataflow* ObjectToEdit, TSharedPtr<IDetailsView> PropertiesEditor);
 
 	static const FName PropertiesTabId;
 	TSharedPtr<IDetailsView> PropertiesEditor;
 	TSharedPtr<IDetailsView> CreatePropertiesEditorWidget(UObject* ObjectToEdit);
 
-
 	UDataflow* Dataflow = nullptr;
 	UGeometryCollection* GeometryCollection = nullptr;
-
-	/** Command list for the graph editor */
-	TSharedPtr<FUICommandList> GraphEditorCommands;
 };
