Workaround so that we can use live coding with engine geometry collection module. This is a hack until explicit load is not needed

diff --git a/Engine/Plugins/Experimental/GeometryCollectionPlugin/Source/GeometryCollectionTracks/Private/GeometryCollectionTracksModule.cpp b/Engine/Plugins/Experimental/GeometryCollectionPlugin/Source/GeometryCollectionTracks/Private/GeometryCollectionTracksModule.cpp
--- a/Engine/Plugins/Experimental/GeometryCollectionPlugin/Source/GeometryCollectionTracks/Private/GeometryCollectionTracksModule.cpp
+++ b/Engine/Plugins/Experimental/GeometryCollectionPlugin/Source/GeometryCollectionTracks/Private/GeometryCollectionTracksModule.cpp
@@ -3,6 +3,7 @@
 #include "GeometryCollectionTracksModule.h"
 #if WITH_EDITOR
 #include "GeometryCollectionSequencerModule.h"
+#include "GeometryCollection/GeometryCollectionComponentPlugin.h"
 #endif // WITH_EDITOR
 
 IMPLEMENT_MODULE(FGeometryCollectionTracksModule, GeometryCollectionTracks)
@@ -11,6 +12,10 @@ void FGeometryCollectionTracksModule::StartupModule()
 {
 #if WITH_EDITOR
 	FGeometryCollectionSequencerModule& Module = FModuleManager::LoadModuleChecked<FGeometryCollectionSequencerModule>(TEXT("GeometryCollectionSequencer"));
+	
+	//HACK: The geometry collection engine module will not be added to live coding without this.
+	//We should probably move the engine module into this plugin or something, but for now this makes it so we can change the GC engine module and use live coding
+	IGeometryCollectionComponentPlugin::Get();
 #endif
 }
 
diff --git a/Engine/Source/Runtime/Experimental/GeometryCollectionEngine/Public/GeometryCollection/GeometryCollectionComponentPlugin.h b/Engine/Source/Runtime/Experimental/GeometryCollectionEngine/Public/GeometryCollection/GeometryCollectionComponentPlugin.h
--- a/Engine/Source/Runtime/Experimental/GeometryCollectionEngine/Public/GeometryCollection/GeometryCollectionComponentPlugin.h
+++ b/Engine/Source/Runtime/Experimental/GeometryCollectionEngine/Public/GeometryCollection/GeometryCollectionComponentPlugin.h
@@ -23,7 +23,7 @@ public:
 	 */
 	static inline IGeometryCollectionComponentPlugin& Get()
 	{
-		return FModuleManager::LoadModuleChecked< IGeometryCollectionComponentPlugin >( "GeometryCollectionComponentPlugin" );
+		return FModuleManager::LoadModuleChecked< IGeometryCollectionComponentPlugin >( "GeometryCollectionEngine" );
 	}
 
 	/**
@@ -33,7 +33,7 @@ public:
 	 */
 	static inline bool IsAvailable()
 	{
-		return FModuleManager::Get().IsModuleLoaded( "GeometryCollectionComponentPlugin" );
+		return FModuleManager::Get().IsModuleLoaded( "GeometryCollectionEngine" );
 	}
 };
 
