GeometryCollection : Structure change
- Moved AppendStaticMesh and AppendSkeletalMesh to the Engine level
- Moved the Node implementation into GeometryCollectionNodes within the plugin.

diff --git a/Engine/Plugins/Experimental/ChaosEditor/Source/FractureEditor/Private/FractureToolGenerators.cpp b/Engine/Plugins/Experimental/ChaosEditor/Source/FractureEditor/Private/FractureToolGenerators.cpp
--- a/Engine/Plugins/Experimental/ChaosEditor/Source/FractureEditor/Private/FractureToolGenerators.cpp
+++ b/Engine/Plugins/Experimental/ChaosEditor/Source/FractureEditor/Private/FractureToolGenerators.cpp
@@ -18,6 +18,7 @@
 #include "GeometryCollection/GeometryCollection.h"
 #include "GeometryCollection/GeometryCollectionClusteringUtility.h"
 #include "GeometryCollection/GeometryCollectionConversion.h"
+#include "GeometryCollection/GeometryCollectionEngineConversion.h"
 #include "GeometryCollection/GeometryCollectionUtility.h"
 #include "GeometryCollection/GeometryCollectionProximityUtility.h"
 #include "FractureToolContext.h"
@@ -204,7 +205,7 @@ AGeometryCollectionActor* UFractureToolGenerateAsset::ConvertActorsToGeometryCol
 				decltype(FGeometryCollectionSource::SourceMaterial) SourceMaterials(StaticMeshComponent->GetMaterials());
 				FracturedGeometryCollection->GeometrySource.Add({ SourceSoftObjectPath, ComponentTransform, SourceMaterials });
 
-				FGeometryCollectionConversion::AppendStaticMesh(ComponentStaticMesh, SourceMaterials, ComponentTransform, FracturedGeometryCollection, false);
+				FGeometryCollectionEngineConversion::AppendStaticMesh(ComponentStaticMesh, SourceMaterials, ComponentTransform, FracturedGeometryCollection, false);
 			}
 		}
 
@@ -238,7 +239,7 @@ AGeometryCollectionActor* UFractureToolGenerateAsset::ConvertActorsToGeometryCol
 				}
 				FracturedGeometryCollection->GeometrySource.Add({ SourceSoftObjectPath, ComponentTransform, SourceMaterials });
 
-				FGeometryCollectionConversion::AppendGeometryCollection(RestCollection, GeometryCollectionComponent, ComponentTransform, FracturedGeometryCollection, false);
+				FGeometryCollectionEngineConversion::AppendGeometryCollection(RestCollection, GeometryCollectionComponent, ComponentTransform, FracturedGeometryCollection, false);
 
 			}
 		}
@@ -413,7 +414,7 @@ void UFractureToolResetAsset::Execute(TWeakPtr<FFractureEditorModeToolkit> InToo
 					const UObject* SourceMesh = Source.SourceGeometryObject.TryLoad();
 					if (const UStaticMesh* SourceStaticMesh = Cast<UStaticMesh>(SourceMesh))
 					{
-						FGeometryCollectionConversion::AppendStaticMesh(SourceStaticMesh, Source.SourceMaterial, Source.LocalTransform, GeometryCollectionObject, false);
+						FGeometryCollectionEngineConversion::AppendStaticMesh(SourceStaticMesh, Source.SourceMaterial, Source.LocalTransform, GeometryCollectionObject, false);
 					}
 					else if (const USkeletalMesh* SourceSkeletalMesh = Cast<USkeletalMesh>(SourceMesh))
 					{
@@ -422,7 +423,7 @@ void UFractureToolResetAsset::Execute(TWeakPtr<FFractureEditorModeToolkit> InToo
 					}
 					else if (const UGeometryCollection* SourceGeometryCollection = Cast<UGeometryCollection>(SourceMesh))
 					{
-						FGeometryCollectionConversion::AppendGeometryCollection(SourceGeometryCollection, Source.SourceMaterial, Source.LocalTransform, GeometryCollectionObject, false);
+						FGeometryCollectionEngineConversion::AppendGeometryCollection(SourceGeometryCollection, Source.SourceMaterial, Source.LocalTransform, GeometryCollectionObject, false);
 					}
 				}
 
diff --git a/Engine/Plugins/Experimental/GeometryCollectionPlugin/GeometryCollectionPlugin.uplugin b/Engine/Plugins/Experimental/GeometryCollectionPlugin/GeometryCollectionPlugin.uplugin
--- a/Engine/Plugins/Experimental/GeometryCollectionPlugin/GeometryCollectionPlugin.uplugin
+++ b/Engine/Plugins/Experimental/GeometryCollectionPlugin/GeometryCollectionPlugin.uplugin
@@ -14,23 +14,28 @@
     "CanContainContent" : true,
     "IsBetaVersion" : true,
     "Installed" : false,
-    "Modules": [
-    {
-      "Name": "GeometryCollectionEditor",
-      "Type": "Editor",
-      "LoadingPhase": "Default"
-    },
-	{
-      "Name": "GeometryCollectionTracks",
-      "Type": "Runtime",
-      "LoadingPhase": "Default"
-    },
-    {
-      "Name": "GeometryCollectionSequencer",
-      "Type": "Editor",
-      "LoadingPhase": "Default"
-    },
-    ],
+	"Modules": [
+		{
+			"Name": "GeometryCollectionEditor",
+			"Type": "Editor",
+			"LoadingPhase": "Default"
+		},
+		{
+			"Name": "GeometryCollectionTracks",
+			"Type": "Runtime",
+			"LoadingPhase": "Default"
+		},
+		{
+			"Name": "GeometryCollectionSequencer",
+			"Type": "Editor",
+			"LoadingPhase": "Default"
+		},
+		{
+			"Name": "GeometryCollectionNodes",
+			"Type": "Runtime",
+			"LoadingPhase": "Default"
+		}
+	],
 	"Plugins": [
 		{
 			"Name": "ProceduralMeshComponent",
diff --git a/Engine/Plugins/Experimental/GeometryCollectionPlugin/Source/GeometryCollectionEditor/Private/GeometryCollection/AssetTypeActions_GeometryCollection.cpp b/Engine/Plugins/Experimental/GeometryCollectionPlugin/Source/GeometryCollectionEditor/Private/GeometryCollection/AssetTypeActions_GeometryCollection.cpp
--- a/Engine/Plugins/Experimental/GeometryCollectionPlugin/Source/GeometryCollectionEditor/Private/GeometryCollection/AssetTypeActions_GeometryCollection.cpp
+++ b/Engine/Plugins/Experimental/GeometryCollectionPlugin/Source/GeometryCollectionEditor/Private/GeometryCollection/AssetTypeActions_GeometryCollection.cpp
@@ -2,7 +2,7 @@
 
 #include "GeometryCollection/AssetTypeActions_GeometryCollection.h"
 
-#include "GeometryCollection/Dataflow/GeometryCollectionEditorToolkit.h"
+#include "GeometryCollection/GeometryCollectionEditorToolkit.h"
 #include "GeometryCollection/GeometryCollectionEditorPlugin.h"
 #include "ThumbnailRendering/SceneThumbnailInfo.h"
 #include "GeometryCollection/GeometryCollectionObject.h"
diff --git a/Engine/Plugins/Experimental/GeometryCollectionPlugin/Source/GeometryCollectionEditor/Private/GeometryCollection/Dataflow/GeometryCollectionEditorToolkit.cpp b/Engine/Plugins/Experimental/GeometryCollectionPlugin/Source/GeometryCollectionEditor/Private/GeometryCollection/Dataflow/GeometryCollectionEditorToolkit.cpp
deleted file mode 100644
diff --git a/Engine/Plugins/Experimental/GeometryCollectionPlugin/Source/GeometryCollectionEditor/Private/GeometryCollection/Dataflow/GeometryCollectionNodeConnectionTypes.h b/Engine/Plugins/Experimental/GeometryCollectionPlugin/Source/GeometryCollectionEditor/Private/GeometryCollection/Dataflow/GeometryCollectionNodeConnectionTypes.h
deleted file mode 100644
diff --git a/Engine/Plugins/Experimental/GeometryCollectionPlugin/Source/GeometryCollectionEditor/Private/GeometryCollection/Dataflow/GeometryCollectionNodes.cpp b/Engine/Plugins/Experimental/GeometryCollectionPlugin/Source/GeometryCollectionEditor/Private/GeometryCollection/Dataflow/GeometryCollectionNodes.cpp
deleted file mode 100644
diff --git a/Engine/Plugins/Experimental/GeometryCollectionPlugin/Source/GeometryCollectionEditor/Private/GeometryCollection/Dataflow/GeometryCollectionNodes.h b/Engine/Plugins/Experimental/GeometryCollectionPlugin/Source/GeometryCollectionEditor/Private/GeometryCollection/Dataflow/GeometryCollectionNodes.h
deleted file mode 100644
diff --git a/Engine/Plugins/Experimental/GeometryCollectionPlugin/Source/GeometryCollectionEditor/Private/GeometryCollection/GeometryCollectionConversion.cpp b/Engine/Plugins/Experimental/GeometryCollectionPlugin/Source/GeometryCollectionEditor/Private/GeometryCollection/GeometryCollectionConversion.cpp
--- a/Engine/Plugins/Experimental/GeometryCollectionPlugin/Source/GeometryCollectionEditor/Private/GeometryCollection/GeometryCollectionConversion.cpp
+++ b/Engine/Plugins/Experimental/GeometryCollectionPlugin/Source/GeometryCollectionEditor/Private/GeometryCollection/GeometryCollectionConversion.cpp
@@ -18,6 +18,7 @@
 #include "GeometryCollection/GeometryCollectionFactory.h"
 #include "GeometryCollection/GeometryCollectionComponent.h"
 #include "GeometryCollection/GeometryCollectionClusteringUtility.h"
+#include "GeometryCollection/GeometryCollectionEngineConversion.h"
 #include "GeometryCollectionProxyData.h"
 #include "GeometryCollection/GeometryCollectionUtility.h"
 #include "Logging/LogMacros.h"
@@ -31,855 +32,32 @@
 
 DEFINE_LOG_CATEGORY_STATIC(UGeometryCollectionConversionLogging, Log, All);
 
-struct FUniqueVertex
-{
-	FVector3f Normal;
-	FVector3f Tangent;
-	TArray<FVector2f> UVs;
 
-	bool operator==(const FUniqueVertex& Other) const
-	{
-		if (this->UVs.Num() != Other.UVs.Num())
-		{
-			return false;
-		}
-		
-		bool bEquality = true;
-		bEquality &= (this->Normal == Other.Normal);
-		bEquality &= (this->Tangent == Other.Tangent);
-		for (int32 UVLayerIdx = 0; UVLayerIdx < UVs.Num(); ++UVLayerIdx)
-		{
-			bEquality &= (this->UVs[UVLayerIdx] == Other.UVs[UVLayerIdx]);
-		}
-		
-		return bEquality;
-	}
-};
 
-FORCEINLINE uint32 GetTypeHash(const FUniqueVertex& UniqueVertex)
-{
-	uint32 VertexHash = GetTypeHash(UniqueVertex.Normal);
-	VertexHash = HashCombine(VertexHash, GetTypeHash(UniqueVertex.Tangent));
-	for (int32 UVLayerIdx = 0; UVLayerIdx < UniqueVertex.UVs.Num(); ++UVLayerIdx)
-	{
-		VertexHash = HashCombine(VertexHash, GetTypeHash(UniqueVertex.UVs[UVLayerIdx]));
-	}
-	
-	return VertexHash;
-}
 
 void FGeometryCollectionConversion::AppendStaticMesh(const UStaticMesh* StaticMesh, const TArray<UMaterialInterface*>& Materials, const FTransform& StaticMeshTransform, UGeometryCollection* GeometryCollectionObject, bool ReindexMaterials)
 {
-	if (StaticMesh == nullptr)
-	{
-		return;
-	}
-
-	// Prefer the HiRes description, although this isn't always available.
-	FMeshDescription* MeshDescription = nullptr;
-	if (StaticMesh->IsHiResMeshDescriptionValid())
-	{
-		MeshDescription = StaticMesh->GetHiResMeshDescription();
-	}
-	else
-	{
-		MeshDescription = StaticMesh->GetMeshDescription(0);
-	}
-
-	check(GeometryCollectionObject);
-	TSharedPtr<FGeometryCollection, ESPMode::ThreadSafe> GeometryCollectionPtr = GeometryCollectionObject->GetGeometryCollection();
-	FGeometryCollection* GeometryCollection = GeometryCollectionPtr.Get();
-	check(GeometryCollection);
-
-	if (MeshDescription)
-	{
-		FStaticMeshOperations::ComputeTriangleTangentsAndNormals(*MeshDescription);
-		FStaticMeshOperations::RecomputeNormalsAndTangentsIfNeeded(*MeshDescription, EComputeNTBsFlags::UseMikkTSpace);
-
-		// source vertex information
-		FStaticMeshAttributes Attributes(*MeshDescription);
-		TArrayView<const FVector3f> SourcePosition = Attributes.GetVertexPositions().GetRawArray();
-		TArrayView<const FVector3f> SourceTangent = Attributes.GetVertexInstanceTangents().GetRawArray();
-		TArrayView<const float> SourceBinormalSign = Attributes.GetVertexInstanceBinormalSigns().GetRawArray();
-		TArrayView<const FVector3f> SourceNormal = Attributes.GetVertexInstanceNormals().GetRawArray();
-		TArrayView<const FVector4f> SourceColor = Attributes.GetVertexInstanceColors().GetRawArray();
-
-		TVertexInstanceAttributesConstRef<FVector2f> InstanceUVs = Attributes.GetVertexInstanceUVs();
-		const int32 NumUVLayers = InstanceUVs.GetNumChannels();
-		TArray<TArrayView<const FVector2f>> SourceUVArrays;
-		SourceUVArrays.SetNum(NumUVLayers);
-		for (int32 UVLayerIdx = 0; UVLayerIdx < NumUVLayers; ++UVLayerIdx)
-		{
-			SourceUVArrays[UVLayerIdx] = InstanceUVs.GetRawArray(UVLayerIdx);
-		}
-		
-		// target vertex information
-		TManagedArray<FVector3f>& TargetVertex = GeometryCollection->Vertex;
-		TManagedArray<FVector3f>& TargetTangentU = GeometryCollection->TangentU;
-		TManagedArray<FVector3f>& TargetTangentV = GeometryCollection->TangentV;
-		TManagedArray<FVector3f>& TargetNormal = GeometryCollection->Normal;
-		TManagedArray<TArray<FVector2f>>& TargetUVs = GeometryCollection->UVs;
-		TManagedArray<FLinearColor>& TargetColor = GeometryCollection->Color;
-		TManagedArray<int32>& TargetBoneMap = GeometryCollection->BoneMap;
-		TManagedArray<FLinearColor>& TargetBoneColor = GeometryCollection->BoneColor;
-		TManagedArray<FString>& TargetBoneName = GeometryCollection->BoneName;
-
-		const int32 VertexStart = GeometryCollection->NumElements(FGeometryCollection::VerticesGroup);
-		int32 VertexCount = 0;
-		
-		FVector Scale = StaticMeshTransform.GetScale3D();
-		
-		// We'll need to re-introduce UV seams, etc. by splitting vertices.
-		// A new mapping of MeshDescription vertex instances to the split vertices is maintained.
-		TMap<FVertexInstanceID, int32> VertexInstanceToGeometryCollectionVertex;
-		VertexInstanceToGeometryCollectionVertex.Reserve(Attributes.GetVertexInstanceNormals().GetNumElements());
-		
-		for (const FVertexID VertexIndex : MeshDescription->Vertices().GetElementIDs())
-		{		
-			TArrayView<const FVertexInstanceID> ReferencingVertexInstances = MeshDescription->GetVertexVertexInstanceIDs(VertexIndex);
-
-			// Generate per instance hash of splittable attributes.
-			TMap<FUniqueVertex, TArray<FVertexInstanceID>> SplitVertices;
-			for (const FVertexInstanceID& InstanceID : ReferencingVertexInstances)
-			{
-				TArray<FVector2f> SourceUVs;
-				SourceUVs.SetNum(NumUVLayers);
-				for (int32 UVLayerIdx = 0; UVLayerIdx < NumUVLayers; ++UVLayerIdx)
-				{
-					SourceUVs[UVLayerIdx] = SourceUVArrays[UVLayerIdx][InstanceID];
-				}
-				
-				FUniqueVertex UniqueVertex{ SourceNormal[InstanceID], SourceTangent[InstanceID], SourceUVs };
-				TArray<FVertexInstanceID>& SplitVertex = SplitVertices.FindOrAdd(UniqueVertex);
-				SplitVertex.Add(InstanceID);
-			}
-
-			int32 CurrentVertex = GeometryCollection->AddElements(SplitVertices.Num(), FGeometryCollection::VerticesGroup);
-			
-			// Create a new vertex for each split vertex and map the mesh description instance to it.
-			for (const TTuple<FUniqueVertex,TArray<FVertexInstanceID>>& SplitVertex : SplitVertices)
-			{
-				const TArray<FVertexInstanceID>& InstanceIDs = SplitVertex.Value;
-				const FVertexInstanceID& ExemplarInstanceID = InstanceIDs[0];
-
-				TargetVertex[CurrentVertex] = SourcePosition[VertexIndex] * (FVector3f)Scale;
-				TargetBoneMap[CurrentVertex] = GeometryCollection->NumElements(FGeometryCollection::TransformGroup);
-
-				TargetNormal[CurrentVertex] = SourceNormal[ExemplarInstanceID];
-				TargetTangentU[CurrentVertex] = SourceTangent[ExemplarInstanceID];
-				TargetTangentV[CurrentVertex] = (FVector3f)SourceBinormalSign[ExemplarInstanceID] * FVector3f::CrossProduct(TargetNormal[CurrentVertex], TargetTangentU[CurrentVertex]);
-
-				TargetUVs[CurrentVertex] = SplitVertex.Key.UVs;
-
-				if (SourceColor.Num() > 0)
-				{
-					TargetColor[CurrentVertex] = FLinearColor(SourceColor[ExemplarInstanceID]);
-				}
-				else
-				{
-					TargetColor[CurrentVertex] = FLinearColor::White;
-				}
-
-				for (const FVertexInstanceID& InstanceID : InstanceIDs)
-				{
-					VertexInstanceToGeometryCollectionVertex.Add(InstanceID, CurrentVertex);
-				}
-
-				++CurrentVertex;
-				++VertexCount;
-			}
-		}
-
-
-		// for each material, add a reference in our GeometryCollectionObject
-		const int32 MaterialStart = GeometryCollectionObject->Materials.Num();
-		const int32 NumMeshMaterials = Materials.Num();
-		GeometryCollectionObject->Materials.Reserve(MaterialStart + NumMeshMaterials);
-
-		for (int32 Index = 0; Index < NumMeshMaterials; ++Index)
-		{
-			UMaterialInterface* CurrMaterial = Materials[Index];
-
-			// Possible we have a null entry - replace with default
-			if (CurrMaterial == nullptr)
-			{
-				CurrMaterial = UMaterial::GetDefaultMaterial(MD_Surface);
-			}
-
-			// We add the material twice, once for interior and again for exterior.
-			GeometryCollectionObject->Materials.Add(CurrMaterial);
-			GeometryCollectionObject->Materials.Add(CurrMaterial);
-		}
-
-		// target triangle indices
-		TManagedArray<FIntVector>& TargetIndices = GeometryCollection->Indices;
-		TManagedArray<bool>& TargetVisible = GeometryCollection->Visible;
-		TManagedArray<int32>& TargetMaterialID = GeometryCollection->MaterialID;
-		TManagedArray<int32>& TargetMaterialIndex = GeometryCollection->MaterialIndex;
-
-		const int32 IndicesCount = MeshDescription->Triangles().Num();
-		const int32 InitialNumIndices = GeometryCollection->NumElements(FGeometryCollection::FacesGroup);
-		const int32 IndicesStart = GeometryCollection->AddElements(IndicesCount, FGeometryCollection::FacesGroup);
-		int32 TargetIndex = IndicesStart;
-		for (const int32 TriangleIndex : MeshDescription->Triangles().GetElementIDs())
-		{
-			TArrayView<const FVertexInstanceID> TriangleVertices = MeshDescription->GetTriangleVertexInstances(TriangleIndex);
-
-			TargetIndices[TargetIndex] = FIntVector(
-				VertexInstanceToGeometryCollectionVertex[TriangleVertices[0]],
-				VertexInstanceToGeometryCollectionVertex[TriangleVertices[1]],
-				VertexInstanceToGeometryCollectionVertex[TriangleVertices[2]]
-			);
-
-			TargetVisible[TargetIndex] = true;
-
-			// Materials are ganged in pairs and we want the id to associate with the first of each pair.
-			TargetMaterialID[TargetIndex] = MaterialStart + (MeshDescription->GetTrianglePolygonGroup(TriangleIndex) * 2);
-
-			// Is this right?
-			TargetMaterialIndex[TargetIndex] = TargetIndex;
-
-			++TargetIndex;
-		}
-
-		// Geometry transform
-		TManagedArray<FTransform>& Transform = GeometryCollection->Transform;
-
-		int32 TransformIndex1 = GeometryCollection->AddElements(1, FGeometryCollection::TransformGroup);
-		Transform[TransformIndex1] = StaticMeshTransform;
-		Transform[TransformIndex1].SetScale3D(FVector(1.f, 1.f, 1.f));
-
-		// collisions
-		if (UBodySetup* BodySetup = StaticMesh->GetBodySetup())
-		{
-			TArray<TUniquePtr<Chaos::FImplicitObject>> Geoms;
-			Chaos::FShapesArray Shapes;
-
-			FGeometryAddParams CreateGeometryParams;
-			CreateGeometryParams.bDoubleSided = false;
-			CreateGeometryParams.CollisionData.CollisionFlags.bEnableQueryCollision = true;
-			CreateGeometryParams.CollisionData.CollisionFlags.bEnableSimCollisionComplex = false; // no support for trimesh in destruction
-			CreateGeometryParams.CollisionData.CollisionFlags.bEnableSimCollisionSimple = true;
-			CreateGeometryParams.CollisionTraceType = ECollisionTraceFlag::CTF_UseSimpleAsComplex;
-			CreateGeometryParams.Scale = Scale;
-			CreateGeometryParams.LocalTransform = Chaos::FRigidTransform3::Identity;
-			CreateGeometryParams.WorldTransform = Chaos::FRigidTransform3::Identity;
-			CreateGeometryParams.Geometry = &BodySetup->AggGeom;
-			CreateGeometryParams.ChaosTriMeshes = MakeArrayView(BodySetup->ChaosTriMeshes);
-
-			// todo(chaos) : this currently also create the shape array which is unnecessary ,this could be optimized by having a common function to create only the implicits 
-			ChaosInterface::CreateGeometry(CreateGeometryParams, Geoms, Shapes);
-
-			using FCollisionType = FGeometryDynamicCollection::FSharedImplicit;
-			TManagedArray<FCollisionType>& ExternaCollisions = GeometryCollection->AddAttribute<FCollisionType>("ExternalCollisions", FGeometryCollection::TransformGroup);
-			if (Geoms.Num() == 0)
-			{
-				ExternaCollisions[TransformIndex1] = nullptr;
-			}
-			else if (Geoms.Num() == 1)
-			{
-				ExternaCollisions[TransformIndex1] = TSharedPtr<Chaos::FImplicitObject>(Geoms[0].Release());  
-			}
-			else
-			{
-				ExternaCollisions[TransformIndex1] = MakeShared<Chaos::FImplicitObjectUnion>(MoveTemp(Geoms)); 
-			}
-		}
-		
-		// Bone Hierarchy - Added at root with no common parent
-		TManagedArray<int32>& Parent = GeometryCollection->Parent;
-		TManagedArray<int32>& SimulationType = GeometryCollection->SimulationType;
-		Parent[TransformIndex1] = FGeometryCollection::Invalid;
-		SimulationType[TransformIndex1] = FGeometryCollection::ESimulationTypes::FST_Rigid;
-
-		const FColor RandBoneColor(FMath::Rand() % 100 + 5, FMath::Rand() % 100 + 5, FMath::Rand() % 100 + 5, 255);
-		TargetBoneColor[TransformIndex1] = FLinearColor(RandBoneColor);
-		TargetBoneName[TransformIndex1] = StaticMesh->GetName();
-
-		// GeometryGroup
-		int GeometryIndex = GeometryCollection->AddElements(1, FGeometryCollection::GeometryGroup);
-
-		TManagedArray<int32>& TransformIndex = GeometryCollection->TransformIndex;
-		TManagedArray<FBox>& BoundingBox = GeometryCollection->BoundingBox;
-		TManagedArray<float>& InnerRadius = GeometryCollection->InnerRadius;
-		TManagedArray<float>& OuterRadius = GeometryCollection->OuterRadius;
-		TManagedArray<int32>& VertexStartArray = GeometryCollection->VertexStart;
-		TManagedArray<int32>& VertexCountArray = GeometryCollection->VertexCount;
-		TManagedArray<int32>& FaceStartArray = GeometryCollection->FaceStart;
-		TManagedArray<int32>& FaceCountArray = GeometryCollection->FaceCount;
-
-		TransformIndex[GeometryIndex] = TargetBoneMap[VertexStart];
-		VertexStartArray[GeometryIndex] = VertexStart;
-		VertexCountArray[GeometryIndex] = VertexCount;
-		FaceStartArray[GeometryIndex] = InitialNumIndices;
-		FaceCountArray[GeometryIndex] = IndicesCount;
-
-		// TransformGroup
-		TManagedArray<int32>& TransformToGeometryIndexArray = GeometryCollection->TransformToGeometryIndex;
-		TransformToGeometryIndexArray[TransformIndex1] = GeometryIndex;
-
-		FVector Center(0);
-		for (int32 VertexIndex = VertexStart; VertexIndex < VertexStart + VertexCount; VertexIndex++)
-		{
-			Center += (FVector)TargetVertex[VertexIndex];
-		}
-		if (VertexCount) Center /= VertexCount;
-
-		// Inner/Outer edges, bounding box
-		BoundingBox[GeometryIndex] = FBox(ForceInitToZero);
-		InnerRadius[GeometryIndex] = FLT_MAX;
-		OuterRadius[GeometryIndex] = -FLT_MAX;
-		for (int32 VertexIndex = VertexStart; VertexIndex < VertexStart + VertexCount; VertexIndex++)
-		{
-			BoundingBox[GeometryIndex] += (FVector)TargetVertex[VertexIndex];
-
-			float Delta = (Center - (FVector)TargetVertex[VertexIndex]).Size();
-			InnerRadius[GeometryIndex] = FMath::Min(InnerRadius[GeometryIndex], Delta);
-			OuterRadius[GeometryIndex] = FMath::Max(OuterRadius[GeometryIndex], Delta);
-		}
-
-		// Inner/Outer centroid
-		for (int fdx = IndicesStart; fdx < IndicesStart + IndicesCount; fdx++)
-		{
-			FVector Centroid(0);
-			for (int e = 0; e < 3; e++)
-			{
-				Centroid += (FVector)TargetVertex[TargetIndices[fdx][e]];
-			}
-			Centroid /= 3;
-
-			float Delta = (Center - Centroid).Size();
-			InnerRadius[GeometryIndex] = FMath::Min(InnerRadius[GeometryIndex], Delta);
-			OuterRadius[GeometryIndex] = FMath::Max(OuterRadius[GeometryIndex], Delta);
-		}
-
-		// Inner/Outer edges
-		for (int fdx = IndicesStart; fdx < IndicesStart + IndicesCount; fdx++)
-		{
-			for (int e = 0; e < 3; e++)
-			{
-				int i = e, j = (e + 1) % 3;
-				FVector Edge = (FVector)TargetVertex[TargetIndices[fdx][i]] + 0.5 * FVector(TargetVertex[TargetIndices[fdx][j]] - TargetVertex[TargetIndices[fdx][i]]);
-				float Delta = (Center - Edge).Size();
-				InnerRadius[GeometryIndex] = FMath::Min(InnerRadius[GeometryIndex], Delta);
-				OuterRadius[GeometryIndex] = FMath::Max(OuterRadius[GeometryIndex], Delta);
-			}
-		}
-
-		if (ReindexMaterials) {
-			GeometryCollection->ReindexMaterials();
-		}
-	}
+	FGeometryCollectionEngineConversion::AppendStaticMesh(StaticMesh, Materials, StaticMeshTransform, GeometryCollectionObject, ReindexMaterials);
 }
 
-void FGeometryCollectionConversion::AppendGeometryCollection(const UGeometryCollection* SourceGeometryCollection, const TArray<UMaterialInterface*>& Materials, const FTransform& GeometryCollectionTransform, UGeometryCollection* TargetGeometryCollectionObject, bool ReindexMaterials)
+void FGeometryCollectionConversion::AppendStaticMesh(const UStaticMesh* StaticMesh, const UStaticMeshComponent* StaticMeshComponent, const FTransform& StaticMeshTransform, UGeometryCollection* GeometryCollection, bool ReindexMaterials)
 {
-	if (SourceGeometryCollection == nullptr)
-	{
-		return;
-	}
-	const TSharedPtr<FGeometryCollection, ESPMode::ThreadSafe> SourceGeometryCollectionPtr = SourceGeometryCollection->GetGeometryCollection();
-
-	check(TargetGeometryCollectionObject);
-	TSharedPtr<FGeometryCollection, ESPMode::ThreadSafe> GeometryCollectionPtr = TargetGeometryCollectionObject->GetGeometryCollection();
-	FGeometryCollection* GeometryCollection = GeometryCollectionPtr.Get();
-	check(GeometryCollection);
-
-	// Assemble offsets and add elements
-	const int32 VertexCount = SourceGeometryCollectionPtr->Vertex.Num();
-	const int32 FaceCount = SourceGeometryCollectionPtr->Indices.Num();
-	const int32 TransformCount = SourceGeometryCollectionPtr->Transform.Num();
-	const int32 GeometryCount = SourceGeometryCollectionPtr->TransformIndex.Num();
-	const int32 SectionCount = SourceGeometryCollectionPtr->Sections.Num();
-
-	const int32 AssetMaterialStart = TargetGeometryCollectionObject->Materials.Num();
-
-	FVector Scale = GeometryCollectionTransform.GetScale3D();
-	FTransform AppliedTransform = GeometryCollectionTransform;
-	AppliedTransform.RemoveScaling();
-
-	const int32 VertexStart = GeometryCollection->AddElements(VertexCount, FGeometryCollection::VerticesGroup);
-	const int32 FaceStart = GeometryCollection->AddElements(FaceCount, FGeometryCollection::FacesGroup);
-	const int32 TransformStart = GeometryCollection->AddElements(TransformCount, FGeometryCollection::TransformGroup);
-	const int32 GeometryStart = GeometryCollection->AddElements(GeometryCount, FGeometryCollection::GeometryGroup);
-	const int32 SectionStart = GeometryCollection->AddElements(SectionCount, FGeometryCollection::MaterialGroup);
-
-	// source vertex information
-	const TManagedArray<FVector3f>& SourceVertex = SourceGeometryCollectionPtr->Vertex;
-	const TManagedArray<FVector3f>& SourceTangentU = SourceGeometryCollectionPtr->TangentU;
-	const TManagedArray<FVector3f>& SourceTangentV = SourceGeometryCollectionPtr->TangentV;
-	const TManagedArray<FVector3f>& SourceNormal = SourceGeometryCollectionPtr->Normal;
-	const TManagedArray<TArray<FVector2f>>& SourceUVs = SourceGeometryCollectionPtr->UVs;
-	const TManagedArray<FLinearColor>& SourceColor = SourceGeometryCollectionPtr->Color;
-	const TManagedArray<int32>& SourceBoneMap = SourceGeometryCollectionPtr->BoneMap;
-
-	// target vertex information
-	TManagedArray<FVector3f>& TargetVertex = GeometryCollection->Vertex;
-	TManagedArray<FVector3f>& TargetTangentU = GeometryCollection->TangentU;
-	TManagedArray<FVector3f>& TargetTangentV = GeometryCollection->TangentV;
-	TManagedArray<FVector3f>& TargetNormal = GeometryCollection->Normal;
-	TManagedArray<TArray<FVector2f>>& TargetUVs = GeometryCollection->UVs;
-	TManagedArray<FLinearColor>& TargetColor = GeometryCollection->Color;
-	TManagedArray<int32>& TargetBoneMap = GeometryCollection->BoneMap;
-
-	// append vertices
-	for (int32 VertexIndex = 0; VertexIndex < VertexCount; VertexIndex++)
-	{
-		const int32 VertexOffset = VertexStart + VertexIndex;
-		TargetVertex[VertexOffset] = SourceVertex[VertexIndex] * (FVector3f)Scale;
-		
-		TargetTangentU[VertexOffset] = SourceTangentU[VertexIndex];
-		TargetTangentV[VertexOffset] = SourceTangentV[VertexIndex];
-		TargetNormal[VertexOffset] = SourceNormal[VertexIndex];
-		TargetUVs[VertexOffset] = SourceUVs[VertexIndex];
-		TargetColor[VertexOffset] = SourceColor[VertexIndex];
-
-		TargetBoneMap[VertexOffset] = SourceBoneMap[VertexIndex] + TransformStart;
-	}
-
-	// source face information
-	const TManagedArray<FIntVector>& SourceIndices = SourceGeometryCollectionPtr->Indices;
-	const TManagedArray<bool>& SourceVisible = SourceGeometryCollectionPtr->Visible;
-	const TManagedArray<int32>& SourceMaterialID = SourceGeometryCollectionPtr->MaterialID;
-	const TManagedArray<int32>& SourceMaterialIndex = SourceGeometryCollectionPtr->MaterialIndex;
-	
-	// target face information
-	TManagedArray<FIntVector>& TargetIndices = GeometryCollection->Indices;
-	TManagedArray<bool>& TargetVisible = GeometryCollection->Visible;
-	TManagedArray<int32>& TargetMaterialID = GeometryCollection->MaterialID;
-	TManagedArray<int32>& TargetMaterialIndex = GeometryCollection->MaterialIndex;
-
-	// append faces
-	for (int32 FaceIndex = 0; FaceIndex < FaceCount; ++FaceIndex)
-	{
-		const FIntVector& SourceFace = SourceIndices[FaceIndex];
-		const int32 FaceOffset = FaceStart + FaceIndex;
-		TargetIndices[FaceOffset] = FIntVector(
-			SourceFace[0] + VertexStart,
-			SourceFace[1] + VertexStart,
-			SourceFace[2] + VertexStart);
-		TargetVisible[FaceOffset] = SourceVisible[FaceIndex];
-		
-		TargetMaterialID[FaceOffset] = AssetMaterialStart + SourceMaterialID[FaceIndex];
-		TargetMaterialIndex[FaceOffset] = FaceOffset;
-	}
-
-	// source transform information
-	const TManagedArray<FTransform>& SourceTransform = SourceGeometryCollectionPtr->Transform;
-	const TManagedArray<FString>& SourceBoneName = SourceGeometryCollectionPtr->BoneName;
-	const TManagedArray<FLinearColor>& SourceBoneColor = SourceGeometryCollectionPtr->BoneColor;
-	const TManagedArray<int32>& SourceParent = SourceGeometryCollectionPtr->Parent;
-	const TManagedArray<TSet<int32>>& SourceChildren = SourceGeometryCollectionPtr->Children;
-	const TManagedArray<int32>& SourceTransformToGeometryIndex = SourceGeometryCollectionPtr->TransformToGeometryIndex;
-	const TManagedArray<int32>& SourceSimulationType = SourceGeometryCollectionPtr->SimulationType;
-	const TManagedArray<int32>& SourceStatusFlags = SourceGeometryCollectionPtr->StatusFlags;
-	const TManagedArray<int32>& SourceInitialDynamicState = SourceGeometryCollectionPtr->InitialDynamicState;
-	
-	// target transform information
-	TManagedArray<FTransform>& TargetTransform = GeometryCollection->Transform;
-	TManagedArray<FString>& TargetBoneName = GeometryCollection->BoneName;
-	TManagedArray<FLinearColor>& TargetBoneColor = GeometryCollection->BoneColor;
-	TManagedArray<int32>& TargetParent = GeometryCollection->Parent;
-	TManagedArray<TSet<int32>>& TargetChildren = GeometryCollection->Children;
-	TManagedArray<int32>& TargetTransformToGeometryIndex = GeometryCollection->TransformToGeometryIndex;
-	TManagedArray<int32>& TargetSimulationType = GeometryCollection->SimulationType;
-	TManagedArray<int32>& TargetStatusFlags = GeometryCollection->StatusFlags;
-	TManagedArray<int32>& TargetInitialDynamicState = GeometryCollection->InitialDynamicState;
-
-	// append transform hierarchy
-	for (int32 TransformIndex = 0; TransformIndex < TransformCount; ++TransformIndex)
-	{
-		const int32 TransformOffset = TransformStart + TransformIndex;
-
-		// Only apply the transform to the parent node. Child nodes only need scaling applied to translation offsets.
-		if (SourceParent[TransformIndex] == INDEX_NONE)
-		{
-			TargetTransform[TransformOffset] = SourceTransform[TransformIndex] * AppliedTransform;
-		}
-		else
-		{
-			FTransform ScaledTranslation = SourceTransform[TransformIndex];
-			ScaledTranslation.ScaleTranslation(Scale);
-			TargetTransform[TransformOffset] = ScaledTranslation;
-		}
-	
-		// #todo Get this Bone name to be unique
-		TargetBoneName[TransformOffset] = SourceBoneName[TransformIndex];
-
-		const FColor RandBoneColor(FMath::Rand() % 100 + 5, FMath::Rand() % 100 + 5, FMath::Rand() % 100 + 5, 255);
-		TargetBoneColor[TransformOffset] = FLinearColor(RandBoneColor);
-
-		TargetParent[TransformOffset] = (SourceParent[TransformIndex]==INDEX_NONE) ? INDEX_NONE: SourceParent[TransformIndex] + TransformStart;
-
-		const TSet<int32>& SourceChildrenSet = SourceChildren[TransformIndex];
-		for (int32 ChildIndex : SourceChildrenSet)
-		{
-			TargetChildren[TransformOffset].Add(ChildIndex + TransformStart);
-		}
-
-		TargetTransformToGeometryIndex[TransformOffset] = SourceTransformToGeometryIndex[TransformIndex] + GeometryStart;
-		TargetSimulationType[TransformOffset] = SourceSimulationType[TransformIndex];
-		TargetStatusFlags[TransformOffset] = SourceStatusFlags[TransformIndex];
-		TargetInitialDynamicState[TransformOffset] = SourceInitialDynamicState[TransformIndex];
-	}
-
-	// source geometry information
-	const TManagedArray<int32>& SourceTransformIndex = SourceGeometryCollectionPtr->TransformIndex;
-	const TManagedArray<int32>& SourceVertexStart = SourceGeometryCollectionPtr->VertexStart;
-	const TManagedArray<int32>& SourceVertexCount = SourceGeometryCollectionPtr->VertexCount;
-	const TManagedArray<int32>& SourceFaceStart = SourceGeometryCollectionPtr->FaceStart;
-	const TManagedArray<int32>& SourceFaceCount = SourceGeometryCollectionPtr->FaceCount;
-	
-	// target geometry information
-	TManagedArray<int32>& TargetTransformIndex = GeometryCollection->TransformIndex;
-	TManagedArray<FBox>& TargetBoundingBox = GeometryCollection->BoundingBox;
-	TManagedArray<float>& TargetInnerRadius = GeometryCollection->InnerRadius;
-	TManagedArray<float>& TargetOuterRadius = GeometryCollection->OuterRadius;
-	TManagedArray<int32>& TargetVertexStart = GeometryCollection->VertexStart;
-	TManagedArray<int32>& TargetVertexCount = GeometryCollection->VertexCount;
-	TManagedArray<int32>& TargetFaceStart = GeometryCollection->FaceStart;
-	TManagedArray<int32>& TargetFaceCount = GeometryCollection->FaceCount;
-
-	// append geometry
-	for (int32 GeometryIndex = 0; GeometryIndex < GeometryCount; ++GeometryIndex)
-	{
-		const int32 GeometryOffset = GeometryStart + GeometryIndex;
-
-		TargetTransformIndex[GeometryOffset] = SourceTransformIndex[GeometryIndex] + TransformStart;
-		
-		TargetVertexStart[GeometryOffset] = SourceVertexStart[GeometryIndex] + VertexStart;
-		TargetVertexCount[GeometryOffset] = SourceVertexCount[GeometryIndex];
-		TargetFaceStart[GeometryOffset] = SourceFaceStart[GeometryIndex] + FaceStart;
-		TargetFaceCount[GeometryOffset] = SourceFaceCount[GeometryIndex];
-
-		// Find centroid of geometry for inner/outer radius calculations
-		FVector Center(0);
-		for (int32 VertexIndex = TargetVertexStart[GeometryOffset]; VertexIndex < TargetVertexStart[GeometryOffset] + TargetVertexCount[GeometryOffset]; ++VertexIndex)
-		{
-			Center += (FVector)TargetVertex[VertexIndex];
-		}
-		if (TargetVertexCount[GeometryOffset]) Center /= TargetVertexCount[GeometryOffset];
-		
-		TargetBoundingBox[GeometryOffset] = FBox(ForceInitToZero);
-		TargetInnerRadius[GeometryOffset] = FLT_MAX;
-		TargetOuterRadius[GeometryOffset] = -FLT_MAX;
-		for (int32 VertexIndex = TargetVertexStart[GeometryOffset]; VertexIndex < TargetVertexStart[GeometryOffset] + TargetVertexCount[GeometryOffset]; ++VertexIndex)
-		{
-			TargetBoundingBox[GeometryOffset] += (FVector)TargetVertex[VertexIndex];
-
-			float Delta = (Center - (FVector)TargetVertex[VertexIndex]).Size();
-			TargetInnerRadius[GeometryOffset] = FMath::Min(TargetInnerRadius[GeometryOffset], Delta);
-			TargetOuterRadius[GeometryOffset] = FMath::Max(TargetOuterRadius[GeometryOffset], Delta);
-		}
-	}
-
-	// for each material, add a reference in our GeometryCollectionObject
-	const int32 NumMeshMaterials = Materials.Num();
-	TargetGeometryCollectionObject->Materials.Reserve(AssetMaterialStart + NumMeshMaterials);
-
-	for (int32 Index = 0; Index < NumMeshMaterials; ++Index)
-	{
-		UMaterialInterface* CurrMaterial = Materials[Index];
-
-		// Possible we have a null entry - replace with default
-		if (CurrMaterial == nullptr)
-		{
-			CurrMaterial = UMaterial::GetDefaultMaterial(MD_Surface);
-		}
-
-		TargetGeometryCollectionObject->Materials.Add(CurrMaterial);
-	}
-
-	// source material information
-	const TManagedArray<FGeometryCollectionSection>& SourceSections = SourceGeometryCollectionPtr->Sections;
-
-	// target material information
-	TManagedArray<FGeometryCollectionSection>& TargetSections = GeometryCollection->Sections;
-
-	// append sections
-	for (int32 SectionIndex = 0; SectionIndex < SectionCount; ++SectionIndex)
-	{
-		int32 SectionOffset = SectionStart + SectionIndex;
-
-		TargetSections[SectionOffset].MaterialID = AssetMaterialStart + SourceSections[SectionIndex].MaterialID;
-
-		TargetSections[SectionOffset].FirstIndex = SourceSections[SectionIndex].FirstIndex + FaceStart * 3;
-		TargetSections[SectionOffset].MinVertexIndex = VertexStart + SourceSections[SectionIndex].MinVertexIndex;
-
-		TargetSections[SectionOffset].NumTriangles = SourceSections[SectionIndex].NumTriangles;
-		TargetSections[SectionOffset].MaxVertexIndex = VertexStart + SourceSections[SectionIndex].MaxVertexIndex;
-	}
-
-	if (ReindexMaterials) {
-		GeometryCollection->ReindexMaterials();
-	}
-
+	FGeometryCollectionEngineConversion::AppendStaticMesh(StaticMesh, StaticMeshComponent, StaticMeshTransform, GeometryCollection, ReindexMaterials);
 }
 
-
-void FGeometryCollectionConversion::AppendStaticMesh(const UStaticMesh* StaticMesh, const UStaticMeshComponent* StaticMeshComponent, const FTransform& StaticMeshTransform, UGeometryCollection* GeometryCollectionObject, bool ReindexMaterials)
+void FGeometryCollectionConversion::AppendSkeletalMesh(const USkeletalMesh* SkeletalMesh, const USkeletalMeshComponent* SkeletalMeshComponent, const FTransform& SkeletalMeshTransform, UGeometryCollection* GeometryCollection, bool ReindexMaterials)
 {
-	if (StaticMesh == nullptr)
-	{
-		return;
-	}
-
-	TArray<UMaterialInterface*> Materials;
-	Materials.Reserve(StaticMesh->GetStaticMaterials().Num());
-
-	for (int32 Index = 0; Index < StaticMesh->GetStaticMaterials().Num(); ++Index)
-	{
-		UMaterialInterface* CurrMaterial = StaticMeshComponent ? StaticMeshComponent->GetMaterial(Index) : StaticMesh->GetMaterial(Index);
-		Materials.Add(CurrMaterial);
-	}
-
-	// Geometry collections usually carry the selection material, which we'll delete before appending
-	UMaterialInterface* BoneSelectedMaterial = LoadObject<UMaterialInterface>(nullptr, UGeometryCollection::GetSelectedMaterialPath(), nullptr, LOAD_None, nullptr);
-	GeometryCollectionObject->Materials.Remove(BoneSelectedMaterial);
-	Materials.Remove(BoneSelectedMaterial);
-	
-	AppendStaticMesh(StaticMesh, Materials, StaticMeshTransform, GeometryCollectionObject, ReindexMaterials);
+	FGeometryCollectionEngineConversion::AppendSkeletalMesh(SkeletalMesh, SkeletalMeshComponent, SkeletalMeshTransform, GeometryCollection, ReindexMaterials);
 }
 
-
-void FGeometryCollectionConversion::AppendGeometryCollection(const UGeometryCollection* SourceGeometryCollection, const UGeometryCollectionComponent* GeometryCollectionComponent, const FTransform& GeometryCollectionTransform, UGeometryCollection* TargetGeometryCollectionObject, bool ReindexMaterials)
+void FGeometryCollectionConversion::AppendGeometryCollection(const UGeometryCollection* SourceGeometryCollection, const TArray<UMaterialInterface*>& Materials, const FTransform& GeometryCollectionTransform, UGeometryCollection* TargetGeometryCollectionObject, bool ReindexMaterials)
 {
-	if (SourceGeometryCollection == nullptr)
-	{
-		return;
-	}
-
-	TArray<UMaterialInterface*> Materials;
-	Materials.Reserve(SourceGeometryCollection->Materials.Num());
-
-	for (int32 Index = 0; Index < SourceGeometryCollection->Materials.Num(); ++Index)
-	{
-		UMaterialInterface* CurrMaterial = GeometryCollectionComponent ? GeometryCollectionComponent->GetMaterial(Index) : SourceGeometryCollection->Materials[Index].Get();
-		Materials.Add(CurrMaterial);
-	}
-
-	// Geometry collections usually carry the selection material, which we'll delete before appending
-	UMaterialInterface* BoneSelectedMaterial = LoadObject<UMaterialInterface>(nullptr, UGeometryCollection::GetSelectedMaterialPath(), nullptr, LOAD_None, nullptr);
-	TargetGeometryCollectionObject->Materials.Remove(BoneSelectedMaterial);
-	Materials.Remove(BoneSelectedMaterial);
-
-	AppendGeometryCollection(SourceGeometryCollection, Materials, GeometryCollectionTransform, TargetGeometryCollectionObject, ReindexMaterials);
+	FGeometryCollectionEngineConversion::AppendGeometryCollection(SourceGeometryCollection, Materials, GeometryCollectionTransform, TargetGeometryCollectionObject, ReindexMaterials);
 }
 
-
-void FGeometryCollectionConversion::AppendSkeletalMesh(const USkeletalMesh* SkeletalMesh, const USkeletalMeshComponent *SkeletalMeshComponent, const FTransform & SkeletalMeshTransform, UGeometryCollection * GeometryCollectionObject, bool ReindexMaterials)
+void FGeometryCollectionConversion::AppendGeometryCollection(const UGeometryCollection* SourceGeometryCollection, const UGeometryCollectionComponent* GeometryCollectionComponent, const FTransform& GeometryCollectionTransform, UGeometryCollection* TargetGeometryCollectionObject, bool ReindexMaterials)
 {
-	//UE_LOG(UGeometryCollectionConversionLogging, Log, TEXT("FGeometryCollectionConversion::AppendSkeletalMesh()"));
-	check(SkeletalMesh);
-	if (GeometryCollectionObject)
-	{
-		TSharedPtr<FGeometryCollection, ESPMode::ThreadSafe> GeometryCollectionPtr = GeometryCollectionObject->GetGeometryCollection();
-		if (FGeometryCollection* GeometryCollection = GeometryCollectionPtr.Get())
-		{
-
-			if (const USkeleton* Skeleton = SkeletalMesh->GetSkeleton())
-			{
-				if (const FSkeletalMeshRenderData * SkelMeshRenderData = SkeletalMesh->GetResourceForRendering())
-				{
-					if (SkelMeshRenderData->LODRenderData.Num())
-					{
-						const FSkeletalMeshLODRenderData& SkeletalMeshLODRenderData = SkelMeshRenderData->LODRenderData[0];
-						const FSkinWeightVertexBuffer& SkinWeightVertexBuffer = *SkeletalMeshLODRenderData.GetSkinWeightVertexBuffer();
-						const FStaticMeshVertexBuffers& VertexBuffers = SkeletalMeshLODRenderData.StaticVertexBuffers;
-						const FPositionVertexBuffer& PositionVertexBuffer = VertexBuffers.PositionVertexBuffer;
-						const int32 VertexCount = PositionVertexBuffer.GetNumVertices();
-						// Check that all vertex weightings are rigid. 
-						for (int32 VertexIndex = 0; VertexIndex < VertexCount; VertexIndex++)
-						{
-							int32 SkeletalBoneIndex = -1;
-							if (!SkinWeightVertexBuffer.GetRigidWeightBone(VertexIndex, SkeletalBoneIndex))
-							{
-								UE_LOG(UGeometryCollectionConversionLogging, Error, TEXT("Non-rigid weighting found on vertex %d: Cannot convert to GeometryCollection."), VertexIndex);
-								return;
-							}
-						}
-						const FSkelMeshRenderSection& RenderSection = SkeletalMeshLODRenderData.RenderSections[0];
-						const TArray<FBoneIndexType>& SkeletalBoneMap = RenderSection.BoneMap;
-
-						//
-						// The Component transform for each Mesh will become the FTransform that drives
-						// its associated VerticesGroup. The Skeleton will contain a nested transform hierarchy
-						// that is evaluated using the GetComponentSpaceTransformRefPose. The resulting
-						// Transforms array stored in the GeometryCollection will be the same size as
-						// the SkeletalBoneMap. Note the @todo: the SkeletalBoneMap is pulled from only
-						// the first render section, this will need to be expanded to include all render
-						// sections.
-						//
-						TManagedArray<FTransform>& Transform = GeometryCollection->Transform;
-						int32 TransformBaseIndex = GeometryCollection->AddElements(SkeletalBoneMap.Num(), FGeometryCollection::TransformGroup);
-						const FReferenceSkeleton & ReferenceSkeletion = Skeleton->GetReferenceSkeleton();
-						const TArray<FTransform> & RestArray = Skeleton->GetRefLocalPoses();
-						for (int32 BoneIndex = 0; BoneIndex < SkeletalBoneMap.Num(); BoneIndex++)
-						{
-							FTransform BoneTransform = FAnimationRuntime::GetComponentSpaceTransformRefPose(ReferenceSkeletion, SkeletalBoneMap[BoneIndex]);
-							Transform[TransformBaseIndex + BoneIndex] = BoneTransform;
-						}
-
-
-						//
-						// The Triangle Indices
-						//
-						TManagedArray<FIntVector>& Indices = GeometryCollection->Indices;
-						TManagedArray<bool>& Visible = GeometryCollection->Visible;
-						TManagedArray<int32>& MaterialID = GeometryCollection->MaterialID;
-						TManagedArray<int32>& MaterialIndex = GeometryCollection->MaterialIndex;
-
-						TArray<uint32> IndexBuffer;
-						SkeletalMeshLODRenderData.MultiSizeIndexContainer.GetIndexBuffer(IndexBuffer);
-
-						const int32 IndicesCount = IndexBuffer.Num() / 3;
-						int NumVertices = GeometryCollection->NumElements(FGeometryCollection::VerticesGroup);
-						int InitialNumIndices = GeometryCollection->NumElements(FGeometryCollection::FacesGroup);
-						int IndicesBaseIndex = GeometryCollection->AddElements(IndicesCount, FGeometryCollection::FacesGroup);
-						for (int32 IndicesIndex = 0, StaticIndex = 0; IndicesIndex < IndicesCount; IndicesIndex++, StaticIndex += 3)
-						{
-							int32 IndicesOffset = IndicesBaseIndex + IndicesIndex;
-							Indices[IndicesOffset] = FIntVector(
-								IndexBuffer[StaticIndex] + NumVertices,
-								IndexBuffer[StaticIndex + 1] + NumVertices,
-								IndexBuffer[StaticIndex + 2] + NumVertices);
-							Visible[IndicesOffset] = true;
-							MaterialID[IndicesOffset] = 0;
-							MaterialIndex[IndicesOffset] = IndicesOffset;
-						}
-
-						//
-						// Vertex Attributes
-						//
-						TManagedArray<FVector3f>& Vertex = GeometryCollection->Vertex;
-						TManagedArray<FVector3f>& TangentU = GeometryCollection->TangentU;
-						TManagedArray<FVector3f>& TangentV = GeometryCollection->TangentV;
-						TManagedArray<FVector3f>& Normal = GeometryCollection->Normal;
-						TManagedArray<TArray<FVector2f>>& UVs = GeometryCollection->UVs;
-						TManagedArray<FLinearColor>& Color = GeometryCollection->Color;
-						TManagedArray<int32>& BoneMap = GeometryCollection->BoneMap;
-						TManagedArray<FLinearColor>& BoneColor = GeometryCollection->BoneColor;
-						TManagedArray<FString>& BoneName = GeometryCollection->BoneName;
-
-						// 
-						// Transform Attributes 
-						// 
-						TManagedArray<int32>& Parent = GeometryCollection->Parent;
-						TManagedArray<int32>& SimulationType = GeometryCollection->SimulationType;
-						int InitialNumVertices = GeometryCollection->NumElements(FGeometryCollection::VerticesGroup);
-						int VertexBaseIndex = GeometryCollection->AddElements(VertexCount, FGeometryCollection::VerticesGroup);
-						const int32 NumUVLayers = VertexBuffers.StaticMeshVertexBuffer.GetNumTexCoords();
-						for (int32 VertexIndex = 0; VertexIndex < VertexCount; VertexIndex++)
-						{
-							int VertexOffset = VertexBaseIndex + VertexIndex;
-							BoneMap[VertexOffset] = -1;
-							int32 SkeletalBoneIndex = -1;
-							SkinWeightVertexBuffer.GetRigidWeightBone(VertexIndex, SkeletalBoneIndex);
-							if (SkeletalBoneIndex > -1)
-							{
-								BoneMap[VertexOffset] = SkeletalBoneIndex + TransformBaseIndex;
-								Vertex[VertexOffset] = (FVector4f)Transform[BoneMap[VertexOffset]].ToInverseMatrixWithScale().TransformPosition((FVector)PositionVertexBuffer.VertexPosition(VertexIndex));
-							}
-							check(BoneMap[VertexOffset] != -1);
-							TangentU[VertexOffset] = VertexBuffers.StaticMeshVertexBuffer.VertexTangentX(VertexIndex);
-							TangentV[VertexOffset] = VertexBuffers.StaticMeshVertexBuffer.VertexTangentY(VertexIndex);
-							Normal[VertexOffset] = VertexBuffers.StaticMeshVertexBuffer.VertexTangentZ(VertexIndex);
-							
-							UVs[VertexOffset].SetNum(NumUVLayers);
-							for (int32 UVLayerIdx = 0; UVLayerIdx < NumUVLayers; ++UVLayerIdx)
-							{
-								UVs[VertexOffset][UVLayerIdx] = VertexBuffers.StaticMeshVertexBuffer.GetVertexUV(VertexIndex, UVLayerIdx);
-							}
-							
-							if (VertexBuffers.ColorVertexBuffer.GetNumVertices() == VertexCount)
-								Color[VertexOffset] = VertexBuffers.ColorVertexBuffer.VertexColor(VertexIndex);
-							else
-								Color[VertexOffset] = FLinearColor::White;
-						}
-
-						int32 InitialIndex = -1;
-						int32 LastParentIndex = -1;
-						int32 CurrentLevel = 0;
-						for (int32 BoneIndex = 0; BoneIndex < SkeletalBoneMap.Num(); BoneIndex++)
-						{
-							// transform based on position of the actor. 
-							Transform[TransformBaseIndex + BoneIndex] = SkeletalMeshTransform * Transform[TransformBaseIndex + BoneIndex];
-
-							// bone attributes
-							BoneName[TransformBaseIndex + BoneIndex] = ReferenceSkeletion.GetBoneName(SkeletalBoneMap[BoneIndex]).ToString();
-							const FColor RandBoneColor(FMath::Rand() % 100 + 5, FMath::Rand() % 100 + 5, FMath::Rand() % 100 + 5, 255);
-							BoneColor[TransformBaseIndex + BoneIndex] = FLinearColor(RandBoneColor);
-
-							// Bone Hierarchy - Added at root with no common parent
-							int32 ParentIndex = ReferenceSkeletion.GetParentIndex(SkeletalBoneMap[BoneIndex]);
-							int32 UseParentIndex = ParentIndex + InitialIndex;
-							if (LastParentIndex != UseParentIndex)
-							{
-								LastParentIndex = UseParentIndex;
-							}
-							Parent[TransformBaseIndex + BoneIndex] = UseParentIndex;
-							SimulationType[TransformBaseIndex + BoneIndex] = FGeometryCollection::ESimulationTypes::FST_Rigid;
-						}
-
-						// Geometry Group
-						TArray<int32> GeometryIndices;
-						GeometryCollectionAlgo::ContiguousArray(GeometryIndices, GeometryCollection->NumElements(FGeometryCollection::GeometryGroup));
-						GeometryCollection->RemoveDependencyFor(FGeometryCollection::GeometryGroup);
-						GeometryCollection->RemoveElements(FGeometryCollection::GeometryGroup, GeometryIndices);
-						::GeometryCollection::AddGeometryProperties(GeometryCollection);
-
-						// for each material, add a reference in our GeometryCollectionObject
-						int CurrIdx = 0;
-						
-						const TArray<FSkeletalMaterial>& SkeletalMeshMaterials = SkeletalMesh->GetMaterials();
-
-						UMaterialInterface* CurrMaterial = SkeletalMeshComponent ? SkeletalMeshComponent->GetMaterial(CurrIdx) : ToRawPtr(SkeletalMeshMaterials[CurrIdx].MaterialInterface);
-
-
-						int MaterialStart = GeometryCollectionObject->Materials.Num();
-						while (CurrMaterial)
-						{
-							GeometryCollectionObject->Materials.Add(CurrMaterial);
-							CurrMaterial = SkeletalMeshComponent ? SkeletalMeshComponent->GetMaterial(++CurrIdx) : ToRawPtr(SkeletalMeshMaterials[++CurrIdx].MaterialInterface);
-
-						}
-
-						const TArray<FSkelMeshRenderSection> &StaticMeshSections = SkeletalMesh->GetResourceForRendering()->LODRenderData[0].RenderSections;
-
-						TManagedArray<FGeometryCollectionSection> & Sections = GeometryCollection->Sections;
-
-						for (const FSkelMeshRenderSection &CurrSection : StaticMeshSections)
-						{
-							// create new section
-							int32 SectionIndex = GeometryCollection->AddElements(1, FGeometryCollection::MaterialGroup);
-						
-							Sections[SectionIndex].MaterialID = MaterialStart + CurrSection.MaterialIndex;
-
-							Sections[SectionIndex].FirstIndex = IndicesBaseIndex * 3 + CurrSection.BaseIndex;
-							Sections[SectionIndex].MinVertexIndex = VertexBaseIndex + CurrSection.BaseVertexIndex;
-
-							Sections[SectionIndex].NumTriangles = CurrSection.NumTriangles;
-
-							// #todo(dmp): what should we set this to?  SkeletalMesh sections are different
-							// but we are resetting this when the re indexing happens
-							Sections[SectionIndex].MaxVertexIndex = VertexBaseIndex + CurrSection.NumVertices;
-
-							// set the materialid for all of the faces
-							for (int32 i = Sections[SectionIndex].FirstIndex / 3; i < Sections[SectionIndex].FirstIndex / 3 + Sections[SectionIndex].NumTriangles; ++i)
-							{
-								MaterialID[i] = SectionIndex;
-							}
-						}
-
-					}
-				}
-			}
-
-			if (ReindexMaterials) {
-				GeometryCollection->ReindexMaterials();
-			}
-		}
-	}
+	FGeometryCollectionEngineConversion::AppendGeometryCollection(SourceGeometryCollection, GeometryCollectionComponent, GeometryCollectionTransform, TargetGeometryCollectionObject, ReindexMaterials);
 }
 
 struct FGeometryCollectionAsset
@@ -953,7 +131,7 @@ void FGeometryCollectionConversion::CreateFromSelectedActorsCommand(UWorld * Wor
 						{
 							GeometryCollectionAsset.Create();
 						}
-						FGeometryCollectionConversion::AppendStaticMesh(
+						FGeometryCollectionEngineConversion::AppendStaticMesh(
 							StaticMeshComponents[Index]->GetStaticMesh(),
 							StaticMeshComponents[Index],
 							Actor->GetTransform(),
@@ -974,7 +152,7 @@ void FGeometryCollectionConversion::CreateFromSelectedActorsCommand(UWorld * Wor
 						{
 							GeometryCollectionAsset.Create();
 						}
-						FGeometryCollectionConversion::AppendSkeletalMesh(
+						FGeometryCollectionEngineConversion::AppendSkeletalMesh(
 							SkeletalMeshComponents[Index]->SkeletalMesh,
 							SkeletalMeshComponents[Index],
 							Actor->GetTransform(),
@@ -1008,7 +186,7 @@ void FGeometryCollectionConversion::CreateFromSelectedAssetsCommand(UWorld * Wor
 			{
 				GeometryCollectionAsset.Create();
 			}
-			FGeometryCollectionConversion::AppendStaticMesh(
+			FGeometryCollectionEngineConversion::AppendStaticMesh(
 				static_cast<const UStaticMesh*>(AssetData.GetAsset()),
 				nullptr,
 				FTransform(), 
@@ -1021,7 +199,7 @@ void FGeometryCollectionConversion::CreateFromSelectedAssetsCommand(UWorld * Wor
 			{
 				GeometryCollectionAsset.Create();
 			}
-			FGeometryCollectionConversion::AppendSkeletalMesh(
+			FGeometryCollectionEngineConversion::AppendSkeletalMesh(
 				static_cast<const USkeletalMesh *>(AssetData.GetAsset()), 
 				nullptr,
 				FTransform(),
diff --git a/Engine/Plugins/Experimental/GeometryCollectionPlugin/Source/GeometryCollectionEditor/Private/GeometryCollection/GeometryCollectionEditorPlugin.cpp b/Engine/Plugins/Experimental/GeometryCollectionPlugin/Source/GeometryCollectionEditor/Private/GeometryCollection/GeometryCollectionEditorPlugin.cpp
--- a/Engine/Plugins/Experimental/GeometryCollectionPlugin/Source/GeometryCollectionEditor/Private/GeometryCollection/GeometryCollectionEditorPlugin.cpp
+++ b/Engine/Plugins/Experimental/GeometryCollectionPlugin/Source/GeometryCollectionEditor/Private/GeometryCollection/GeometryCollectionEditorPlugin.cpp
@@ -5,11 +5,11 @@
 
 #include "AssetToolsModule.h"
 #include "CoreMinimal.h"
-#include "Dataflow/GeometryCollectionNodes.h"
 #include "LevelEditor.h"
 #include "Modules/ModuleManager.h"
 #include "PropertyEditorModule.h"
 #include "GeometryCollection/GeometryCollectionComponent.h"
+#include "GeometryCollection/GeometryCollectionConversion.h"
 #include "GeometryCollection/GeometryCollectionCommands.h"
 #include "GeometryCollection/AssetTypeActions_GeometryCollection.h"
 #include "GeometryCollection/AssetTypeActions_GeometryCollectionCache.h"
@@ -19,7 +19,8 @@
 #include "GeometryCollection/GeometryCollectionThumbnailRenderer.h"
 #include "GeometryCollection/GeometryCollectionSelectRigidBodyEdMode.h"
 #include "GeometryCollection/GeometryCollectionSelectionCommands.h"
-#include "GeometryCollection/Dataflow/GeometryCollectionEditorToolkit.h"
+#include "GeometryCollection/GeometryCollectionEditorToolkit.h"
+#include "GeometryCollection/GeometryCollectionFactory.h"
 #include "HAL/ConsoleManager.h"
 #include "Features/IModularFeatures.h"
 #include "GeometryCollection/DetailCustomizations/GeomComponentCacheCustomization.h"
@@ -261,10 +262,6 @@ void IGeometryCollectionEditorPlugin::StartupModule()
 	// Register tool menus - delayed until ToolMenus module is loaded
 	UToolMenus::RegisterStartupCallback(
 		FSimpleMulticastDelegate::FDelegate::CreateRaw(this, &IGeometryCollectionEditorPlugin::RegisterMenus));
-
-	// Dataflow 
-	Dataflow::GeometryCollectionEngineAssetNodes();
-
 }
 
 
diff --git a/Engine/Plugins/Experimental/GeometryCollectionPlugin/Source/GeometryCollectionEditor/Private/GeometryCollection/GeometryCollectionEditorToolkit.cpp b/Engine/Plugins/Experimental/GeometryCollectionPlugin/Source/GeometryCollectionEditor/Private/GeometryCollection/GeometryCollectionEditorToolkit.cpp
new file mode 100644
--- /dev/null
+++ b/Engine/Plugins/Experimental/GeometryCollectionPlugin/Source/GeometryCollectionEditor/Private/GeometryCollection/GeometryCollectionEditorToolkit.cpp
@@ -0,0 +1,204 @@
+// Copyright Epic Games, Inc. All Rights Reserved.
+
+#include "GeometryCollection/GeometryCollectionEditorToolkit.h"
+
+#include "Dataflow/DataflowEditorActions.h"
+#include "Dataflow/DataflowEditorPlugin.h"
+#include "Dataflow/DataflowEdNode.h"
+#include "Dataflow/DataflowGraphEditor.h"
+#include "Dataflow/DataflowNodeFactory.h"
+#include "Dataflow/DataflowObject.h"
+#include "Dataflow/DataflowObjectInterface.h"
+#include "Dataflow/DataflowSchema.h"
+#include "Dataflow/DataflowCore.h"
+#include "EditorStyleSet.h"
+#include "EditorViewportTabContent.h"
+#include "EditorViewportLayout.h"
+#include "EditorViewportCommands.h"
+#include "GraphEditorActions.h"
+#include "Modules/ModuleManager.h"
+#include "PropertyEditorModule.h"
+#include "Styling/SlateStyleRegistry.h"
+#include "Framework/Commands/GenericCommands.h"
+
+#define LOCTEXT_NAMESPACE "GeometryCollectionEditorToolkit"
+
+//DEFINE_LOG_CATEGORY_STATIC(FGeometryCollectionEditorToolkitLog, Log, All);
+
+
+const FName FGeometryCollectionEditorToolkit::GraphCanvasTabId(TEXT("GeometryCollectionEditor_GraphCanvas"));
+const FName FGeometryCollectionEditorToolkit::PropertiesTabId(TEXT("GeometryCollectionEditor_Properties"));
+
+void FGeometryCollectionEditorToolkit::InitGeometryCollectionAssetEditor(const EToolkitMode::Type Mode, const TSharedPtr<IToolkitHost>& InitToolkitHost, UObject* ObjectToEdit)
+{
+	Dataflow = nullptr;
+	GeometryCollection = CastChecked<UGeometryCollection>(ObjectToEdit);
+
+	if (GeometryCollection != nullptr)
+	{
+		if (GeometryCollection->Dataflow == nullptr)
+		{
+			const FName NodeName = MakeUniqueObjectName(GeometryCollection, UDataflow::StaticClass(), FName("GeometryCollectionDataflowAsset"));
+			GeometryCollection->Dataflow = NewObject<UDataflow>(GeometryCollection, NodeName);
+		}
+		Dataflow = GeometryCollection->Dataflow;
+		GeometryCollection->Dataflow->Schema = UDataflowSchema::StaticClass();
+
+		PropertiesEditor = CreatePropertiesEditorWidget(ObjectToEdit);
+		GraphEditor = CreateGraphEditorWidget(Dataflow, PropertiesEditor);
+
+		const TSharedRef<FTabManager::FLayout> StandaloneDefaultLayout = FTabManager::NewLayout("GeometryCollectionDataflowEditor_Layout")
+			->AddArea
+			(
+				FTabManager::NewPrimaryArea()->SetOrientation(Orient_Vertical)
+				->Split
+				(
+					FTabManager::NewSplitter()->SetOrientation(Orient_Vertical)
+					->Split
+					(
+						FTabManager::NewSplitter()->SetOrientation(Orient_Horizontal)
+						->SetSizeCoefficient(0.9f)
+						->Split
+						(
+							FTabManager::NewStack()
+							->SetSizeCoefficient(0.6f)
+							->AddTab(GraphCanvasTabId, ETabState::OpenedTab)
+						)
+						->Split
+						(
+							FTabManager::NewSplitter()->SetOrientation(Orient_Vertical)
+							->SetSizeCoefficient(0.2f)
+							->Split
+							(
+								FTabManager::NewStack()
+								->SetSizeCoefficient(0.7f)
+								->AddTab(PropertiesTabId, ETabState::OpenedTab)
+							)
+						)
+					)
+				)
+			);
+
+		const bool bCreateDefaultStandaloneMenu = true;
+		const bool bCreateDefaultToolbar = true;
+		FAssetEditorToolkit::InitAssetEditor(Mode, InitToolkitHost, FName(TEXT("GeometryCollectionEditorApp")), StandaloneDefaultLayout, bCreateDefaultToolbar, bCreateDefaultStandaloneMenu, ObjectToEdit);
+	}
+}
+
+TSharedRef<SGraphEditor> FGeometryCollectionEditorToolkit::CreateGraphEditorWidget(UDataflow* DataflowToEdit, TSharedPtr<IDetailsView> InPropertiesEditor)
+{
+	ensure(DataflowToEdit);
+	using namespace Dataflow;
+	IDataflowEditorPlugin& DataflowEditorModule = FModuleManager::LoadModuleChecked<IDataflowEditorPlugin>(TEXT("DataflowEditor"));
+
+	FDataflowEditorCommands::FGraphEvaluationCallback Evaluate = [&](Dataflow::FNode* Node, Dataflow::FConnection* Out)
+	{
+		float EvalTime = FGameTime::GetTimeSinceAppStart().GetRealTimeSeconds();
+		return Node->Evaluate(FEngineContext(GeometryCollection, Dataflow, EvalTime, FName("UGeometryCollection")), Out);
+	};
+
+	return SNew(SDataflowGraphEditor, GeometryCollection)
+		.GraphToEdit(DataflowToEdit)
+		.DetailsView(InPropertiesEditor)
+		.EvaluateGraph(Evaluate);
+}
+
+TSharedPtr<IDetailsView> FGeometryCollectionEditorToolkit::CreatePropertiesEditorWidget(UObject* ObjectToEdit)
+{
+	ensure(ObjectToEdit);
+	FPropertyEditorModule& PropertyEditorModule = FModuleManager::LoadModuleChecked<FPropertyEditorModule>(TEXT("PropertyEditor"));
+
+	FDetailsViewArgs DetailsViewArgs;
+	DetailsViewArgs.bAllowSearch = true;
+	DetailsViewArgs.bLockable = false;
+	DetailsViewArgs.bUpdatesFromSelection = false;
+	DetailsViewArgs.NameAreaSettings = FDetailsViewArgs::HideNameArea;
+	DetailsViewArgs.NotifyHook = this;
+
+	TSharedPtr<IDetailsView> DetailsView = PropertyEditorModule.CreateDetailView(DetailsViewArgs);
+	DetailsView->SetObject(ObjectToEdit);
+	return DetailsView;
+
+}
+
+
+TSharedRef<SDockTab> FGeometryCollectionEditorToolkit::SpawnTab_GraphCanvas(const FSpawnTabArgs& Args)
+{
+	check(Args.GetTabId() == GraphCanvasTabId);
+
+	TSharedRef<SDockTab> SpawnedTab = SNew(SDockTab)
+		.Label(LOCTEXT("GeometryCollectionEditor_Dataflow_TabTitle", "Graph"))
+		[
+			GraphEditor.ToSharedRef()
+		];
+
+	return SpawnedTab;
+}
+
+TSharedRef<SDockTab> FGeometryCollectionEditorToolkit::SpawnTab_Properties(const FSpawnTabArgs& Args)
+{
+	check(Args.GetTabId() == PropertiesTabId);
+
+	return SNew(SDockTab)
+		.Label(LOCTEXT("GeometryCollectionEditor_Properties_TabTitle", "Details"))
+		[
+			PropertiesEditor.ToSharedRef()
+		];
+}
+
+
+void FGeometryCollectionEditorToolkit::RegisterTabSpawners(const TSharedRef<FTabManager>& InTabManager)
+{
+	TSharedRef<FWorkspaceItem> WorkspaceMenuCategoryRef = InTabManager->AddLocalWorkspaceMenuCategory(LOCTEXT("WorkspaceMenu_GeometryCollectionEditor", "Dataflow Editor"));
+
+	InTabManager->RegisterTabSpawner(GraphCanvasTabId, FOnSpawnTab::CreateSP(this, &FGeometryCollectionEditorToolkit::SpawnTab_GraphCanvas))
+		.SetDisplayName(LOCTEXT("DataflowTab", "Graph"))
+		.SetGroup(WorkspaceMenuCategoryRef)
+		.SetIcon(FSlateIcon(FAppStyle::GetAppStyleSetName(), "GraphEditor.EventGraph_16x"));
+
+	InTabManager->RegisterTabSpawner(PropertiesTabId, FOnSpawnTab::CreateSP(this, &FGeometryCollectionEditorToolkit::SpawnTab_Properties))
+		.SetDisplayName(LOCTEXT("PropertiesTab", "Details"))
+		.SetGroup(WorkspaceMenuCategoryRef)
+		.SetIcon(FSlateIcon(FAppStyle::GetAppStyleSetName(), "LevelEditor.Tabs.Details"));
+
+	FAssetEditorToolkit::RegisterTabSpawners(InTabManager);
+}
+
+
+FName FGeometryCollectionEditorToolkit::GetToolkitFName() const
+{
+	return FName("GeometryCollectionEditor");
+}
+
+FText FGeometryCollectionEditorToolkit::GetBaseToolkitName() const
+{
+	return LOCTEXT("AppLabel", "Geometry Collection Editor");
+}
+
+FString FGeometryCollectionEditorToolkit::GetWorldCentricTabPrefix() const
+{
+	return LOCTEXT("WorldCentricTabPrefix", "GeometryCollection").ToString();
+}
+
+FLinearColor FGeometryCollectionEditorToolkit::GetWorldCentricTabColorScale() const
+{
+	return FLinearColor(0.3f, 0.2f, 0.5f, 0.5f);
+}
+
+FString FGeometryCollectionEditorToolkit::GetReferencerName() const
+{
+	return TEXT("GeometryCollectionEditorToolkit");
+}
+
+void FGeometryCollectionEditorToolkit::AddReferencedObjects(FReferenceCollector& Collector)
+{
+	if (Dataflow)
+	{
+		Collector.AddReferencedObject(Dataflow);
+	}
+	if (GeometryCollection)
+	{
+		Collector.AddReferencedObject(GeometryCollection);
+	}
+}
+#undef LOCTEXT_NAMESPACE
diff --git a/Engine/Plugins/Experimental/GeometryCollectionPlugin/Source/GeometryCollectionEditor/Private/GeometryCollection/GeometryCollectionFactory.cpp b/Engine/Plugins/Experimental/GeometryCollectionPlugin/Source/GeometryCollectionEditor/Private/GeometryCollection/GeometryCollectionFactory.cpp
--- a/Engine/Plugins/Experimental/GeometryCollectionPlugin/Source/GeometryCollectionEditor/Private/GeometryCollection/GeometryCollectionFactory.cpp
+++ b/Engine/Plugins/Experimental/GeometryCollectionPlugin/Source/GeometryCollectionEditor/Private/GeometryCollection/GeometryCollectionFactory.cpp
@@ -7,6 +7,7 @@
 #include "GeometryCollection/GeometryCollectionObject.h"
 #include "GeometryCollection/GeometryCollection.h"
 #include "GeometryCollection/GeometryCollectionComponent.h"
+#include "GeometryCollection/GeometryCollectionEngineConversion.h"
 #include "Editor.h"
 #include "Editor/EditorEngine.h"
 #include "Engine/Selection.h"
@@ -154,12 +155,12 @@ UObject* UGeometryCollectionFactory::FactoryCreateNew(UClass* Class, UObject* In
 
 	for (GeometryCollectionStaticMeshConversionTuple & StaticMeshData : StaticMeshList)
 	{
-		FGeometryCollectionConversion::AppendStaticMesh(StaticMeshData.Get<0>(), StaticMeshData.Get<1>(), StaticMeshData.Get<2>(), NewGeometryCollection, false);
+		FGeometryCollectionEngineConversion::AppendStaticMesh(StaticMeshData.Get<0>(), StaticMeshData.Get<1>(), StaticMeshData.Get<2>(), NewGeometryCollection, false);
 	}
 
 	for (GeometryCollectionSkeletalMeshConversionTuple & SkeletalMeshData : SkeletalMeshList)
 	{
-		FGeometryCollectionConversion::AppendSkeletalMesh(SkeletalMeshData.Get<0>(), SkeletalMeshData.Get<1>(), SkeletalMeshData.Get<2>(), NewGeometryCollection, false);
+		FGeometryCollectionEngineConversion::AppendSkeletalMesh(SkeletalMeshData.Get<0>(), SkeletalMeshData.Get<1>(), SkeletalMeshData.Get<2>(), NewGeometryCollection, false);
 	}
 
 	for (GeometryCollectionTuple & GeometryCollectionData : GeometryCollectionList)
diff --git a/Engine/Plugins/Experimental/GeometryCollectionPlugin/Source/GeometryCollectionEditor/Public/GeometryCollection/Dataflow/GeometryCollectionEditorToolkit.h b/Engine/Plugins/Experimental/GeometryCollectionPlugin/Source/GeometryCollectionEditor/Public/GeometryCollection/Dataflow/GeometryCollectionEditorToolkit.h
deleted file mode 100644
diff --git a/Engine/Plugins/Experimental/GeometryCollectionPlugin/Source/GeometryCollectionEditor/Public/GeometryCollection/GeometryCollectionEditorToolkit.h b/Engine/Plugins/Experimental/GeometryCollectionPlugin/Source/GeometryCollectionEditor/Public/GeometryCollection/GeometryCollectionEditorToolkit.h
new file mode 100644
--- /dev/null
+++ b/Engine/Plugins/Experimental/GeometryCollectionPlugin/Source/GeometryCollectionEditor/Public/GeometryCollection/GeometryCollectionEditorToolkit.h
@@ -0,0 +1,62 @@
+// Copyright Epic Games, Inc. All Rights Reserved.
+
+#pragma once
+
+#include "CoreMinimal.h"
+#include "Dataflow/DataflowNodeParameters.h"
+#include "Dataflow/DataflowGraphEditor.h"
+#include "GeometryCollection/GeometryCollectionObject.h"
+#include "GraphEditor.h"
+#include "Toolkits/SimpleAssetEditor.h"
+#include "Misc/NotifyHook.h"
+
+class IDetailsView;
+class FTabManager;
+class IToolkitHost;
+class UDataflow;
+
+
+class FGeometryCollectionEditorToolkit : public FAssetEditorToolkit, public FNotifyHook, public FGCObject
+{
+public:
+	void InitGeometryCollectionAssetEditor(const EToolkitMode::Type Mode, const TSharedPtr<IToolkitHost>& InitToolkitHost, UObject* ObjectToEdit);
+
+	// IToolkit Interface
+	virtual FName GetToolkitFName() const override;
+	virtual FText GetBaseToolkitName() const override;
+	virtual FString GetWorldCentricTabPrefix() const override;
+	virtual FLinearColor GetWorldCentricTabColorScale() const override;
+
+	// FGCObject Interface
+	virtual void AddReferencedObjects(FReferenceCollector& Collector) override;
+	virtual FString GetReferencerName() const override; 
+
+	// Tab spawners 
+	virtual void RegisterTabSpawners(const TSharedRef<FTabManager>& TabManager) override;
+	TSharedRef<SDockTab> SpawnTab_GraphCanvas(const FSpawnTabArgs& Args);
+	TSharedRef<SDockTab> SpawnTab_Properties(const FSpawnTabArgs& Args);
+
+	UDataflow* GetDataflow() { return Dataflow; }
+	const UDataflow* GetDataflow() const { return Dataflow; }
+
+	UGeometryCollection* GetGeometryCollection() { return GeometryCollection; }
+	const UGeometryCollection* GetGeometryCollection() const { return GeometryCollection; }
+
+	TSharedPtr<IDetailsView> GetPropertiesEditor() { return PropertiesEditor; }
+	const TSharedPtr<IDetailsView> GetPropertiesEditor() const { return PropertiesEditor; }
+
+	TSharedPtr<SGraphEditor> GetGraphEditor() { return GraphEditor; }
+	const TSharedPtr<SGraphEditor> GetGraphEditor() const { return GraphEditor; }
+
+private:
+	static const FName GraphCanvasTabId;
+	TSharedPtr<SGraphEditor> GraphEditor;
+	TSharedRef<SGraphEditor> CreateGraphEditorWidget(UDataflow* ObjectToEdit, TSharedPtr<IDetailsView> PropertiesEditor);
+
+	static const FName PropertiesTabId;
+	TSharedPtr<IDetailsView> PropertiesEditor;
+	TSharedPtr<IDetailsView> CreatePropertiesEditorWidget(UObject* ObjectToEdit);
+
+	UDataflow* Dataflow = nullptr;
+	UGeometryCollection* GeometryCollection = nullptr;
+};
diff --git a/Engine/Plugins/Experimental/GeometryCollectionPlugin/Source/GeometryCollectionNodes/GeometryCollectionNodes.Build.cs b/Engine/Plugins/Experimental/GeometryCollectionPlugin/Source/GeometryCollectionNodes/GeometryCollectionNodes.Build.cs
new file mode 100644
--- /dev/null
+++ b/Engine/Plugins/Experimental/GeometryCollectionPlugin/Source/GeometryCollectionNodes/GeometryCollectionNodes.Build.cs
@@ -0,0 +1,27 @@
+// Copyright Epic Games, Inc. All Rights Reserved.
+
+namespace UnrealBuildTool.Rules
+{
+	public class GeometryCollectionNodes : ModuleRules
+	{
+        public GeometryCollectionNodes(ReadOnlyTargetRules Target) : base(Target)
+		{
+			PrivateIncludePaths.Add("GeometryCollectionNodes/Private");
+            PublicIncludePaths.Add(ModuleDirectory + "/Public");
+
+			PublicDependencyModuleNames.AddRange(
+				new string[]
+				{
+					"Core",
+					"CoreUObject",
+					"ChaosCore",
+					"Chaos",
+					"DataflowCore",
+					"DataflowEngine",
+					"GeometryCollectionEngine",
+					"Engine"
+				}
+			);
+		}
+	}
+}
diff --git a/Engine/Plugins/Experimental/GeometryCollectionPlugin/Source/GeometryCollectionNodes/Private/Dataflow/GeometryCollectionNodes.cpp b/Engine/Plugins/Experimental/GeometryCollectionPlugin/Source/GeometryCollectionNodes/Private/Dataflow/GeometryCollectionNodes.cpp
new file mode 100644
--- /dev/null
+++ b/Engine/Plugins/Experimental/GeometryCollectionPlugin/Source/GeometryCollectionNodes/Private/Dataflow/GeometryCollectionNodes.cpp
@@ -0,0 +1,14 @@
+// Copyright Epic Games, Inc. All Rights Reserved.
+
+#include "Dataflow/GeometryCollectionNodes.h"
+
+namespace Dataflow
+{
+	void GeometryCollectionEngineAssetNodes()
+	{
+		DATAFLOW_NODE_REGISTER_CREATION_FACTORY(GetCollectionAssetNode);
+		DATAFLOW_NODE_REGISTER_CREATION_FACTORY(ExampleCollectionEditNode);
+		DATAFLOW_NODE_REGISTER_CREATION_FACTORY(SetCollectionAssetNode);
+	}
+}
+
diff --git a/Engine/Plugins/Experimental/GeometryCollectionPlugin/Source/GeometryCollectionNodes/Private/Dataflow/GeometryCollectionNodesPlugin.cpp b/Engine/Plugins/Experimental/GeometryCollectionPlugin/Source/GeometryCollectionNodes/Private/Dataflow/GeometryCollectionNodesPlugin.cpp
new file mode 100644
--- /dev/null
+++ b/Engine/Plugins/Experimental/GeometryCollectionPlugin/Source/GeometryCollectionNodes/Private/Dataflow/GeometryCollectionNodesPlugin.cpp
@@ -0,0 +1,26 @@
+// Copyright Epic Games, Inc. All Rights Reserved.
+
+#include "Dataflow/GeometryCollectionNodesPlugin.h"
+
+#include "CoreMinimal.h"
+
+#include "Dataflow/DataflowNodeFactory.h"
+#include "Dataflow/GeometryCollectionNodes.h"
+
+#define LOCTEXT_NAMESPACE "DataflowNodes"
+
+
+void IGeometryCollectionNodesPlugin::StartupModule()
+{
+	Dataflow::GeometryCollectionEngineAssetNodes();
+}
+
+void IGeometryCollectionNodesPlugin::ShutdownModule()
+{
+}
+
+
+IMPLEMENT_MODULE(IGeometryCollectionNodesPlugin, GeometryCollectionNodes)
+
+
+#undef LOCTEXT_NAMESPACE
diff --git a/Engine/Plugins/Experimental/GeometryCollectionPlugin/Source/GeometryCollectionNodes/Public/Dataflow/GeometryCollectionNodeConnectionTypes.h b/Engine/Plugins/Experimental/GeometryCollectionPlugin/Source/GeometryCollectionNodes/Public/Dataflow/GeometryCollectionNodeConnectionTypes.h
new file mode 100644
--- /dev/null
+++ b/Engine/Plugins/Experimental/GeometryCollectionPlugin/Source/GeometryCollectionNodes/Public/Dataflow/GeometryCollectionNodeConnectionTypes.h
@@ -0,0 +1,17 @@
+// Copyright Epic Games, Inc. All Rights Reserved.
+
+#pragma once
+
+#include "CoreMinimal.h"
+
+#include "Dataflow/DataflowCore.h"
+#include "GeometryCollection/ManagedArrayCollection.h"
+
+
+namespace Dataflow
+{
+	typedef TSharedPtr<FManagedArrayCollection> FManagedArrayCollectionSharedPtr;
+	 
+	DATAFLOW_CONNECTION_TYPE_SHARED_POINTER(GEOMETRYCOLLECTIONNODES_API, FManagedArrayCollectionSharedPtr, ManagedArrayCollectionSharedPtr)
+	
+}
diff --git a/Engine/Plugins/Experimental/GeometryCollectionPlugin/Source/GeometryCollectionNodes/Public/Dataflow/GeometryCollectionNodes.h b/Engine/Plugins/Experimental/GeometryCollectionPlugin/Source/GeometryCollectionNodes/Public/Dataflow/GeometryCollectionNodes.h
new file mode 100644
--- /dev/null
+++ b/Engine/Plugins/Experimental/GeometryCollectionPlugin/Source/GeometryCollectionNodes/Public/Dataflow/GeometryCollectionNodes.h
@@ -0,0 +1,126 @@
+// Copyright Epic Games, Inc. All Rights Reserved.
+
+#pragma once 
+
+#include "GeometryCollectionNodeConnectionTypes.h"
+
+#include "CoreMinimal.h"
+#include "Dataflow/DataflowObjectInterface.h"
+#include "Dataflow/DataflowEngine.h"
+#include "Dataflow/DataflowProperty.h"
+#include "GeometryCollection/GeometryCollectionObject.h"
+#include "GeometryCollection/ManagedArrayCollection.h"
+#include "GeometryCollection/GeometryCollection.h"
+
+
+namespace Dataflow
+{
+
+	class GetCollectionAssetNode : public FNode
+	{
+		DATAFLOW_NODE_DEFINE_INTERNAL(GetCollectionAssetNode)
+
+	public:
+		typedef FManagedArrayCollectionSharedPtr DataType;
+
+		TSharedPtr< class TOutput<DataType> > Output;
+
+		GetCollectionAssetNode(const FNodeParameters& InParam, FGuid InGuid = FGuid::NewGuid())
+			: FNode(InParam, InGuid)
+			, Output(new TOutput<DataType>(TOutputParameters<DataType>({ FName("CollectionOut"), this })))
+		{}
+
+		virtual void Evaluate(const FContext& Context, FConnection* Out) const override
+		{
+			if (Output.Get() == Out)
+			{
+				if (const FEngineContext* EngineContext = Context.AsType<FEngineContext>(FName("UGeometryCollection")))
+				{
+					if (UGeometryCollection* CollectionAsset = Cast<UGeometryCollection>(EngineContext->Owner))
+					{
+						if (const TSharedPtr<FGeometryCollection, ESPMode::ThreadSafe> AssetCollection = CollectionAsset->GetGeometryCollection())
+						{
+							FManagedArrayCollection* NewCollection = AssetCollection->NewCopy<FManagedArrayCollection>();
+							Output->SetValue(FManagedArrayCollectionSharedPtr(NewCollection), Context);
+						}
+					}
+				}
+			}
+		}
+	};
+
+
+	class ExampleCollectionEditNode : public FNode
+	{
+		DATAFLOW_NODE_DEFINE_INTERNAL(ExampleCollectionEditNode)
+
+	public:
+		typedef FManagedArrayCollectionSharedPtr DataType;
+		TSharedPtr< class TInput<DataType> > Input;
+		TSharedPtr< class TOutput<DataType> > Output;
+		TProperty<bool> Active;
+		TProperty<float> Scale;
+
+		ExampleCollectionEditNode(const FNodeParameters& InParam, FGuid InGuid = FGuid::NewGuid())
+			: FNode(InParam, InGuid)
+			, Input(new TInput<DataType>(TInputParameters<DataType>({ FName("CollectionIn"), this })))
+			, Output(new TOutput<DataType>(TOutputParameters<DataType>({ FName("CollectionOut"), this })))
+			, Active({ FName("Active"), true, this })
+			, Scale({ FName("Scale"), 10.f, this })
+		{}
+
+		virtual void Evaluate(const FContext& Context, FConnection* Out) const override
+		{
+			if (Output.Get() == Out)
+			{
+				DataType Collection = Input->GetValue(Context);
+				if (Active.GetValue())
+				{
+					TManagedArray<FVector3f>* Vertex = Collection->FindAttribute<FVector3f>("Vertex", "Vertices");
+					for (int i = 0; i < Vertex->Num(); i++)
+					{
+						(*Vertex)[i][1] *= Scale.GetValue();
+					}
+				}
+				Output->SetValue(Collection,Context);
+			}
+		}
+	};
+
+	class SetCollectionAssetNode : public FNode
+	{
+		DATAFLOW_NODE_DEFINE_INTERNAL(SetCollectionAssetNode)
+
+	public:
+		typedef FManagedArrayCollectionSharedPtr DataType;
+		TSharedPtr< class TInput<DataType> > Input;
+
+		SetCollectionAssetNode(const FNodeParameters& InParam, FGuid InGuid = FGuid::NewGuid())
+			: FNode(InParam, InGuid)
+			, Input(new TInput<DataType>(TInputParameters<DataType>({ FName("CollectionIn"), this })))
+		{}
+
+		virtual void Evaluate(const FContext& Context, FConnection* Out) const override
+		{
+			if (Out == nullptr)
+			{
+				DataType Collection = Input->GetValue(Context);
+				if (const FEngineContext* EngineContext = Context.AsType<FEngineContext>(FName("UGeometryCollection")))
+				{
+					if (UGeometryCollection* CollectionAsset = Cast<UGeometryCollection>(EngineContext->Owner))
+					{
+						TSharedPtr<FGeometryCollection, ESPMode::ThreadSafe> NewCollection(Collection->NewCopy<FGeometryCollection>());
+						CollectionAsset->SetGeometryCollection(NewCollection);
+						CollectionAsset->InvalidateCollection();
+					}
+				}
+			}
+		}
+	};
+
+
+	void GeometryCollectionEngineAssetNodes();
+
+
+}
+
diff --git a/Engine/Plugins/Experimental/GeometryCollectionPlugin/Source/GeometryCollectionNodes/Public/Dataflow/GeometryCollectionNodesPlugin.h b/Engine/Plugins/Experimental/GeometryCollectionPlugin/Source/GeometryCollectionNodes/Public/Dataflow/GeometryCollectionNodesPlugin.h
new file mode 100644
--- /dev/null
+++ b/Engine/Plugins/Experimental/GeometryCollectionPlugin/Source/GeometryCollectionNodes/Public/Dataflow/GeometryCollectionNodesPlugin.h
@@ -0,0 +1,18 @@
+// Copyright Epic Games, Inc. All Rights Reserved.
+
+#pragma once
+
+#include "CoreMinimal.h"
+#include "Modules/ModuleInterface.h"
+#include "Modules/ModuleManager.h"
+
+/**
+ * The public interface to this module
+ */
+class IGeometryCollectionNodesPlugin : public IModuleInterface
+{
+public:
+	virtual void StartupModule();
+	virtual void ShutdownModule();
+};
+
diff --git a/Engine/Source/Runtime/Experimental/GeometryCollectionEngine/GeometryCollectionEngine.Build.cs b/Engine/Source/Runtime/Experimental/GeometryCollectionEngine/GeometryCollectionEngine.Build.cs
--- a/Engine/Source/Runtime/Experimental/GeometryCollectionEngine/GeometryCollectionEngine.Build.cs
+++ b/Engine/Source/Runtime/Experimental/GeometryCollectionEngine/GeometryCollectionEngine.Build.cs
@@ -25,7 +25,9 @@ namespace UnrealBuildTool.Rules
 					"NetCore",
 					"IntelISPC",
 					"DataflowCore",
-					"DataflowEngine"
+					"DataflowEngine",
+					"MeshDescription",
+					"StaticMeshDescription",
 				}
 				);
 
diff --git a/Engine/Source/Runtime/Experimental/GeometryCollectionEngine/Private/GeometryCollection/GeometryCollectionEngineConversion.cpp b/Engine/Source/Runtime/Experimental/GeometryCollectionEngine/Private/GeometryCollection/GeometryCollectionEngineConversion.cpp
new file mode 100644
--- /dev/null
+++ b/Engine/Source/Runtime/Experimental/GeometryCollectionEngine/Private/GeometryCollection/GeometryCollectionEngineConversion.cpp
@@ -0,0 +1,885 @@
+// Copyright Epic Games, Inc. All Rights Reserved.
+
+#include "GeometryCollection/GeometryCollectionEngineConversion.h"
+
+#include "AssetRegistry/AssetRegistryModule.h"
+#include "AnimationRuntime.h"
+#include "Async/ParallelFor.h"
+#include "Components/SkeletalMeshComponent.h"
+#include "Components/StaticMeshComponent.h"
+#include "Engine/Selection.h"
+#include "Engine/StaticMesh.h"
+#include "GameFramework/Actor.h"
+#include "GeometryCollection/GeometryCollection.h"
+#include "GeometryCollection/GeometryCollectionActor.h"
+#include "GeometryCollection/GeometryCollectionAlgo.h"
+#include "GeometryCollection/GeometryCollectionComponent.h"
+#include "GeometryCollection/GeometryCollectionClusteringUtility.h"
+#include "GeometryCollectionProxyData.h"
+#include "GeometryCollection/GeometryCollectionUtility.h"
+#include "Logging/LogMacros.h"
+#include "Materials/Material.h"
+#include "Rendering/SkeletalMeshRenderData.h"
+#include "MeshDescription.h"
+#include "StaticMeshAttributes.h"
+#include "StaticMeshOperations.h"
+#include "Physics/Experimental/ChaosInterfaceUtils.h"
+
+
+DEFINE_LOG_CATEGORY_STATIC(UGeometryCollectionConversionLogging, Log, All);
+
+struct FUniqueVertex
+{
+	FVector3f Normal;
+	FVector3f Tangent;
+	TArray<FVector2f> UVs;
+
+	bool operator==(const FUniqueVertex& Other) const
+	{
+		if (this->UVs.Num() != Other.UVs.Num())
+		{
+			return false;
+		}
+		
+		bool bEquality = true;
+		bEquality &= (this->Normal == Other.Normal);
+		bEquality &= (this->Tangent == Other.Tangent);
+		for (int32 UVLayerIdx = 0; UVLayerIdx < UVs.Num(); ++UVLayerIdx)
+		{
+			bEquality &= (this->UVs[UVLayerIdx] == Other.UVs[UVLayerIdx]);
+		}
+		
+		return bEquality;
+	}
+};
+
+FORCEINLINE uint32 GetTypeHash(const FUniqueVertex& UniqueVertex)
+{
+	uint32 VertexHash = GetTypeHash(UniqueVertex.Normal);
+	VertexHash = HashCombine(VertexHash, GetTypeHash(UniqueVertex.Tangent));
+	for (int32 UVLayerIdx = 0; UVLayerIdx < UniqueVertex.UVs.Num(); ++UVLayerIdx)
+	{
+		VertexHash = HashCombine(VertexHash, GetTypeHash(UniqueVertex.UVs[UVLayerIdx]));
+	}
+	
+	return VertexHash;
+}
+
+void FGeometryCollectionEngineConversion::AppendStaticMesh(const UStaticMesh* StaticMesh, const TArray<UMaterialInterface*>& Materials, const FTransform& StaticMeshTransform, UGeometryCollection* GeometryCollectionObject, bool ReindexMaterials)
+{
+#if WITH_EDITORONLY_DATA
+
+	if (StaticMesh == nullptr)
+	{
+		return;
+	}
+
+	// Prefer the HiRes description, although this isn't always available.
+	FMeshDescription* MeshDescription = nullptr;
+	if (StaticMesh->IsHiResMeshDescriptionValid())
+	{
+		MeshDescription = StaticMesh->GetHiResMeshDescription();
+	}
+	else
+	{
+		MeshDescription = StaticMesh->GetMeshDescription(0);
+	}
+
+	check(GeometryCollectionObject);
+	TSharedPtr<FGeometryCollection, ESPMode::ThreadSafe> GeometryCollectionPtr = GeometryCollectionObject->GetGeometryCollection();
+	FGeometryCollection* GeometryCollection = GeometryCollectionPtr.Get();
+	check(GeometryCollection);
+
+	if (MeshDescription)
+	{
+		FStaticMeshOperations::ComputeTriangleTangentsAndNormals(*MeshDescription);
+		FStaticMeshOperations::RecomputeNormalsAndTangentsIfNeeded(*MeshDescription, EComputeNTBsFlags::UseMikkTSpace);
+
+		// source vertex information
+		FStaticMeshAttributes Attributes(*MeshDescription);
+		TArrayView<const FVector3f> SourcePosition = Attributes.GetVertexPositions().GetRawArray();
+		TArrayView<const FVector3f> SourceTangent = Attributes.GetVertexInstanceTangents().GetRawArray();
+		TArrayView<const float> SourceBinormalSign = Attributes.GetVertexInstanceBinormalSigns().GetRawArray();
+		TArrayView<const FVector3f> SourceNormal = Attributes.GetVertexInstanceNormals().GetRawArray();
+		TArrayView<const FVector4f> SourceColor = Attributes.GetVertexInstanceColors().GetRawArray();
+
+		TVertexInstanceAttributesConstRef<FVector2f> InstanceUVs = Attributes.GetVertexInstanceUVs();
+		const int32 NumUVLayers = InstanceUVs.GetNumChannels();
+		TArray<TArrayView<const FVector2f>> SourceUVArrays;
+		SourceUVArrays.SetNum(NumUVLayers);
+		for (int32 UVLayerIdx = 0; UVLayerIdx < NumUVLayers; ++UVLayerIdx)
+		{
+			SourceUVArrays[UVLayerIdx] = InstanceUVs.GetRawArray(UVLayerIdx);
+		}
+		
+		// target vertex information
+		TManagedArray<FVector3f>& TargetVertex = GeometryCollection->Vertex;
+		TManagedArray<FVector3f>& TargetTangentU = GeometryCollection->TangentU;
+		TManagedArray<FVector3f>& TargetTangentV = GeometryCollection->TangentV;
+		TManagedArray<FVector3f>& TargetNormal = GeometryCollection->Normal;
+		TManagedArray<TArray<FVector2f>>& TargetUVs = GeometryCollection->UVs;
+		TManagedArray<FLinearColor>& TargetColor = GeometryCollection->Color;
+		TManagedArray<int32>& TargetBoneMap = GeometryCollection->BoneMap;
+		TManagedArray<FLinearColor>& TargetBoneColor = GeometryCollection->BoneColor;
+		TManagedArray<FString>& TargetBoneName = GeometryCollection->BoneName;
+
+		const int32 VertexStart = GeometryCollection->NumElements(FGeometryCollection::VerticesGroup);
+		int32 VertexCount = 0;
+		
+		FVector Scale = StaticMeshTransform.GetScale3D();
+		
+		// We'll need to re-introduce UV seams, etc. by splitting vertices.
+		// A new mapping of MeshDescription vertex instances to the split vertices is maintained.
+		TMap<FVertexInstanceID, int32> VertexInstanceToGeometryCollectionVertex;
+		VertexInstanceToGeometryCollectionVertex.Reserve(Attributes.GetVertexInstanceNormals().GetNumElements());
+		
+		for (const FVertexID VertexIndex : MeshDescription->Vertices().GetElementIDs())
+		{		
+			TArrayView<const FVertexInstanceID> ReferencingVertexInstances = MeshDescription->GetVertexVertexInstanceIDs(VertexIndex);
+
+			// Generate per instance hash of splittable attributes.
+			TMap<FUniqueVertex, TArray<FVertexInstanceID>> SplitVertices;
+			for (const FVertexInstanceID& InstanceID : ReferencingVertexInstances)
+			{
+				TArray<FVector2f> SourceUVs;
+				SourceUVs.SetNum(NumUVLayers);
+				for (int32 UVLayerIdx = 0; UVLayerIdx < NumUVLayers; ++UVLayerIdx)
+				{
+					SourceUVs[UVLayerIdx] = SourceUVArrays[UVLayerIdx][InstanceID];
+				}
+				
+				FUniqueVertex UniqueVertex{ SourceNormal[InstanceID], SourceTangent[InstanceID], SourceUVs };
+				TArray<FVertexInstanceID>& SplitVertex = SplitVertices.FindOrAdd(UniqueVertex);
+				SplitVertex.Add(InstanceID);
+			}
+
+			int32 CurrentVertex = GeometryCollection->AddElements(SplitVertices.Num(), FGeometryCollection::VerticesGroup);
+			
+			// Create a new vertex for each split vertex and map the mesh description instance to it.
+			for (const TTuple<FUniqueVertex,TArray<FVertexInstanceID>>& SplitVertex : SplitVertices)
+			{
+				const TArray<FVertexInstanceID>& InstanceIDs = SplitVertex.Value;
+				const FVertexInstanceID& ExemplarInstanceID = InstanceIDs[0];
+
+				TargetVertex[CurrentVertex] = SourcePosition[VertexIndex] * (FVector3f)Scale;
+				TargetBoneMap[CurrentVertex] = GeometryCollection->NumElements(FGeometryCollection::TransformGroup);
+
+				TargetNormal[CurrentVertex] = SourceNormal[ExemplarInstanceID];
+				TargetTangentU[CurrentVertex] = SourceTangent[ExemplarInstanceID];
+				TargetTangentV[CurrentVertex] = (FVector3f)SourceBinormalSign[ExemplarInstanceID] * FVector3f::CrossProduct(TargetNormal[CurrentVertex], TargetTangentU[CurrentVertex]);
+
+				TargetUVs[CurrentVertex] = SplitVertex.Key.UVs;
+
+				if (SourceColor.Num() > 0)
+				{
+					TargetColor[CurrentVertex] = FLinearColor(SourceColor[ExemplarInstanceID]);
+				}
+				else
+				{
+					TargetColor[CurrentVertex] = FLinearColor::White;
+				}
+
+				for (const FVertexInstanceID& InstanceID : InstanceIDs)
+				{
+					VertexInstanceToGeometryCollectionVertex.Add(InstanceID, CurrentVertex);
+				}
+
+				++CurrentVertex;
+				++VertexCount;
+			}
+		}
+
+
+		// for each material, add a reference in our GeometryCollectionObject
+		const int32 MaterialStart = GeometryCollectionObject->Materials.Num();
+		const int32 NumMeshMaterials = Materials.Num();
+		GeometryCollectionObject->Materials.Reserve(MaterialStart + NumMeshMaterials);
+
+		for (int32 Index = 0; Index < NumMeshMaterials; ++Index)
+		{
+			UMaterialInterface* CurrMaterial = Materials[Index];
+
+			// Possible we have a null entry - replace with default
+			if (CurrMaterial == nullptr)
+			{
+				CurrMaterial = UMaterial::GetDefaultMaterial(MD_Surface);
+			}
+
+			// We add the material twice, once for interior and again for exterior.
+			GeometryCollectionObject->Materials.Add(CurrMaterial);
+			GeometryCollectionObject->Materials.Add(CurrMaterial);
+		}
+
+		// target triangle indices
+		TManagedArray<FIntVector>& TargetIndices = GeometryCollection->Indices;
+		TManagedArray<bool>& TargetVisible = GeometryCollection->Visible;
+		TManagedArray<int32>& TargetMaterialID = GeometryCollection->MaterialID;
+		TManagedArray<int32>& TargetMaterialIndex = GeometryCollection->MaterialIndex;
+
+		const int32 IndicesCount = MeshDescription->Triangles().Num();
+		const int32 InitialNumIndices = GeometryCollection->NumElements(FGeometryCollection::FacesGroup);
+		const int32 IndicesStart = GeometryCollection->AddElements(IndicesCount, FGeometryCollection::FacesGroup);
+		int32 TargetIndex = IndicesStart;
+		for (const int32 TriangleIndex : MeshDescription->Triangles().GetElementIDs())
+		{
+			TArrayView<const FVertexInstanceID> TriangleVertices = MeshDescription->GetTriangleVertexInstances(TriangleIndex);
+
+			TargetIndices[TargetIndex] = FIntVector(
+				VertexInstanceToGeometryCollectionVertex[TriangleVertices[0]],
+				VertexInstanceToGeometryCollectionVertex[TriangleVertices[1]],
+				VertexInstanceToGeometryCollectionVertex[TriangleVertices[2]]
+			);
+
+			TargetVisible[TargetIndex] = true;
+
+			// Materials are ganged in pairs and we want the id to associate with the first of each pair.
+			TargetMaterialID[TargetIndex] = MaterialStart + (MeshDescription->GetTrianglePolygonGroup(TriangleIndex) * 2);
+
+			// Is this right?
+			TargetMaterialIndex[TargetIndex] = TargetIndex;
+
+			++TargetIndex;
+		}
+
+		// Geometry transform
+		TManagedArray<FTransform>& Transform = GeometryCollection->Transform;
+
+		int32 TransformIndex1 = GeometryCollection->AddElements(1, FGeometryCollection::TransformGroup);
+		Transform[TransformIndex1] = StaticMeshTransform;
+		Transform[TransformIndex1].SetScale3D(FVector(1.f, 1.f, 1.f));
+
+		// collisions
+		if (UBodySetup* BodySetup = StaticMesh->GetBodySetup())
+		{
+			TArray<TUniquePtr<Chaos::FImplicitObject>> Geoms;
+			Chaos::FShapesArray Shapes;
+
+			FGeometryAddParams CreateGeometryParams;
+			CreateGeometryParams.bDoubleSided = false;
+			CreateGeometryParams.CollisionData.CollisionFlags.bEnableQueryCollision = true;
+			CreateGeometryParams.CollisionData.CollisionFlags.bEnableSimCollisionComplex = false; // no support for trimesh in destruction
+			CreateGeometryParams.CollisionData.CollisionFlags.bEnableSimCollisionSimple = true;
+			CreateGeometryParams.CollisionTraceType = ECollisionTraceFlag::CTF_UseSimpleAsComplex;
+			CreateGeometryParams.Scale = Scale;
+			CreateGeometryParams.LocalTransform = Chaos::FRigidTransform3::Identity;
+			CreateGeometryParams.WorldTransform = Chaos::FRigidTransform3::Identity;
+			CreateGeometryParams.Geometry = &BodySetup->AggGeom;
+			CreateGeometryParams.ChaosTriMeshes = MakeArrayView(BodySetup->ChaosTriMeshes);
+
+			// todo(chaos) : this currently also create the shape array which is unnecessary ,this could be optimized by having a common function to create only the implicits 
+			ChaosInterface::CreateGeometry(CreateGeometryParams, Geoms, Shapes);
+
+			using FCollisionType = FGeometryDynamicCollection::FSharedImplicit;
+			TManagedArray<FCollisionType>& ExternaCollisions = GeometryCollection->AddAttribute<FCollisionType>("ExternalCollisions", FGeometryCollection::TransformGroup);
+			if (Geoms.Num() == 0)
+			{
+				ExternaCollisions[TransformIndex1] = nullptr;
+			}
+			else if (Geoms.Num() == 1)
+			{
+				ExternaCollisions[TransformIndex1] = TSharedPtr<Chaos::FImplicitObject>(Geoms[0].Release());  
+			}
+			else
+			{
+				ExternaCollisions[TransformIndex1] = MakeShared<Chaos::FImplicitObjectUnion>(MoveTemp(Geoms)); 
+			}
+		}
+		
+		// Bone Hierarchy - Added at root with no common parent
+		TManagedArray<int32>& Parent = GeometryCollection->Parent;
+		TManagedArray<int32>& SimulationType = GeometryCollection->SimulationType;
+		Parent[TransformIndex1] = FGeometryCollection::Invalid;
+		SimulationType[TransformIndex1] = FGeometryCollection::ESimulationTypes::FST_Rigid;
+
+		const FColor RandBoneColor(FMath::Rand() % 100 + 5, FMath::Rand() % 100 + 5, FMath::Rand() % 100 + 5, 255);
+		TargetBoneColor[TransformIndex1] = FLinearColor(RandBoneColor);
+		TargetBoneName[TransformIndex1] = StaticMesh->GetName();
+
+		// GeometryGroup
+		int GeometryIndex = GeometryCollection->AddElements(1, FGeometryCollection::GeometryGroup);
+
+		TManagedArray<int32>& TransformIndex = GeometryCollection->TransformIndex;
+		TManagedArray<FBox>& BoundingBox = GeometryCollection->BoundingBox;
+		TManagedArray<float>& InnerRadius = GeometryCollection->InnerRadius;
+		TManagedArray<float>& OuterRadius = GeometryCollection->OuterRadius;
+		TManagedArray<int32>& VertexStartArray = GeometryCollection->VertexStart;
+		TManagedArray<int32>& VertexCountArray = GeometryCollection->VertexCount;
+		TManagedArray<int32>& FaceStartArray = GeometryCollection->FaceStart;
+		TManagedArray<int32>& FaceCountArray = GeometryCollection->FaceCount;
+
+		TransformIndex[GeometryIndex] = TargetBoneMap[VertexStart];
+		VertexStartArray[GeometryIndex] = VertexStart;
+		VertexCountArray[GeometryIndex] = VertexCount;
+		FaceStartArray[GeometryIndex] = InitialNumIndices;
+		FaceCountArray[GeometryIndex] = IndicesCount;
+
+		// TransformGroup
+		TManagedArray<int32>& TransformToGeometryIndexArray = GeometryCollection->TransformToGeometryIndex;
+		TransformToGeometryIndexArray[TransformIndex1] = GeometryIndex;
+
+		FVector Center(0);
+		for (int32 VertexIndex = VertexStart; VertexIndex < VertexStart + VertexCount; VertexIndex++)
+		{
+			Center += (FVector)TargetVertex[VertexIndex];
+		}
+		if (VertexCount) Center /= VertexCount;
+
+		// Inner/Outer edges, bounding box
+		BoundingBox[GeometryIndex] = FBox(ForceInitToZero);
+		InnerRadius[GeometryIndex] = FLT_MAX;
+		OuterRadius[GeometryIndex] = -FLT_MAX;
+		for (int32 VertexIndex = VertexStart; VertexIndex < VertexStart + VertexCount; VertexIndex++)
+		{
+			BoundingBox[GeometryIndex] += (FVector)TargetVertex[VertexIndex];
+
+			float Delta = (Center - (FVector)TargetVertex[VertexIndex]).Size();
+			InnerRadius[GeometryIndex] = FMath::Min(InnerRadius[GeometryIndex], Delta);
+			OuterRadius[GeometryIndex] = FMath::Max(OuterRadius[GeometryIndex], Delta);
+		}
+
+		// Inner/Outer centroid
+		for (int fdx = IndicesStart; fdx < IndicesStart + IndicesCount; fdx++)
+		{
+			FVector Centroid(0);
+			for (int e = 0; e < 3; e++)
+			{
+				Centroid += (FVector)TargetVertex[TargetIndices[fdx][e]];
+			}
+			Centroid /= 3;
+
+			float Delta = (Center - Centroid).Size();
+			InnerRadius[GeometryIndex] = FMath::Min(InnerRadius[GeometryIndex], Delta);
+			OuterRadius[GeometryIndex] = FMath::Max(OuterRadius[GeometryIndex], Delta);
+		}
+
+		// Inner/Outer edges
+		for (int fdx = IndicesStart; fdx < IndicesStart + IndicesCount; fdx++)
+		{
+			for (int e = 0; e < 3; e++)
+			{
+				int i = e, j = (e + 1) % 3;
+				FVector Edge = (FVector)TargetVertex[TargetIndices[fdx][i]] + 0.5 * FVector(TargetVertex[TargetIndices[fdx][j]] - TargetVertex[TargetIndices[fdx][i]]);
+				float Delta = (Center - Edge).Size();
+				InnerRadius[GeometryIndex] = FMath::Min(InnerRadius[GeometryIndex], Delta);
+				OuterRadius[GeometryIndex] = FMath::Max(OuterRadius[GeometryIndex], Delta);
+			}
+		}
+
+		if (ReindexMaterials) {
+			GeometryCollection->ReindexMaterials();
+		}
+	}
+#endif //WITH_EDITORONLY_DATA
+
+}
+
+void FGeometryCollectionEngineConversion::AppendGeometryCollection(const UGeometryCollection* SourceGeometryCollection, const TArray<UMaterialInterface*>& Materials, const FTransform& GeometryCollectionTransform, UGeometryCollection* TargetGeometryCollectionObject, bool ReindexMaterials)
+{
+	if (SourceGeometryCollection == nullptr)
+	{
+		return;
+	}
+	const TSharedPtr<FGeometryCollection, ESPMode::ThreadSafe> SourceGeometryCollectionPtr = SourceGeometryCollection->GetGeometryCollection();
+
+	check(TargetGeometryCollectionObject);
+	TSharedPtr<FGeometryCollection, ESPMode::ThreadSafe> GeometryCollectionPtr = TargetGeometryCollectionObject->GetGeometryCollection();
+	FGeometryCollection* GeometryCollection = GeometryCollectionPtr.Get();
+	check(GeometryCollection);
+
+	// Assemble offsets and add elements
+	const int32 VertexCount = SourceGeometryCollectionPtr->Vertex.Num();
+	const int32 FaceCount = SourceGeometryCollectionPtr->Indices.Num();
+	const int32 TransformCount = SourceGeometryCollectionPtr->Transform.Num();
+	const int32 GeometryCount = SourceGeometryCollectionPtr->TransformIndex.Num();
+	const int32 SectionCount = SourceGeometryCollectionPtr->Sections.Num();
+
+	const int32 AssetMaterialStart = TargetGeometryCollectionObject->Materials.Num();
+
+	FVector Scale = GeometryCollectionTransform.GetScale3D();
+	FTransform AppliedTransform = GeometryCollectionTransform;
+	AppliedTransform.RemoveScaling();
+
+	const int32 VertexStart = GeometryCollection->AddElements(VertexCount, FGeometryCollection::VerticesGroup);
+	const int32 FaceStart = GeometryCollection->AddElements(FaceCount, FGeometryCollection::FacesGroup);
+	const int32 TransformStart = GeometryCollection->AddElements(TransformCount, FGeometryCollection::TransformGroup);
+	const int32 GeometryStart = GeometryCollection->AddElements(GeometryCount, FGeometryCollection::GeometryGroup);
+	const int32 SectionStart = GeometryCollection->AddElements(SectionCount, FGeometryCollection::MaterialGroup);
+
+	// source vertex information
+	const TManagedArray<FVector3f>& SourceVertex = SourceGeometryCollectionPtr->Vertex;
+	const TManagedArray<FVector3f>& SourceTangentU = SourceGeometryCollectionPtr->TangentU;
+	const TManagedArray<FVector3f>& SourceTangentV = SourceGeometryCollectionPtr->TangentV;
+	const TManagedArray<FVector3f>& SourceNormal = SourceGeometryCollectionPtr->Normal;
+	const TManagedArray<TArray<FVector2f>>& SourceUVs = SourceGeometryCollectionPtr->UVs;
+	const TManagedArray<FLinearColor>& SourceColor = SourceGeometryCollectionPtr->Color;
+	const TManagedArray<int32>& SourceBoneMap = SourceGeometryCollectionPtr->BoneMap;
+
+	// target vertex information
+	TManagedArray<FVector3f>& TargetVertex = GeometryCollection->Vertex;
+	TManagedArray<FVector3f>& TargetTangentU = GeometryCollection->TangentU;
+	TManagedArray<FVector3f>& TargetTangentV = GeometryCollection->TangentV;
+	TManagedArray<FVector3f>& TargetNormal = GeometryCollection->Normal;
+	TManagedArray<TArray<FVector2f>>& TargetUVs = GeometryCollection->UVs;
+	TManagedArray<FLinearColor>& TargetColor = GeometryCollection->Color;
+	TManagedArray<int32>& TargetBoneMap = GeometryCollection->BoneMap;
+
+	// append vertices
+	for (int32 VertexIndex = 0; VertexIndex < VertexCount; VertexIndex++)
+	{
+		const int32 VertexOffset = VertexStart + VertexIndex;
+		TargetVertex[VertexOffset] = SourceVertex[VertexIndex] * (FVector3f)Scale;
+		
+		TargetTangentU[VertexOffset] = SourceTangentU[VertexIndex];
+		TargetTangentV[VertexOffset] = SourceTangentV[VertexIndex];
+		TargetNormal[VertexOffset] = SourceNormal[VertexIndex];
+		TargetUVs[VertexOffset] = SourceUVs[VertexIndex];
+		TargetColor[VertexOffset] = SourceColor[VertexIndex];
+
+		TargetBoneMap[VertexOffset] = SourceBoneMap[VertexIndex] + TransformStart;
+	}
+
+	// source face information
+	const TManagedArray<FIntVector>& SourceIndices = SourceGeometryCollectionPtr->Indices;
+	const TManagedArray<bool>& SourceVisible = SourceGeometryCollectionPtr->Visible;
+	const TManagedArray<int32>& SourceMaterialID = SourceGeometryCollectionPtr->MaterialID;
+	const TManagedArray<int32>& SourceMaterialIndex = SourceGeometryCollectionPtr->MaterialIndex;
+	
+	// target face information
+	TManagedArray<FIntVector>& TargetIndices = GeometryCollection->Indices;
+	TManagedArray<bool>& TargetVisible = GeometryCollection->Visible;
+	TManagedArray<int32>& TargetMaterialID = GeometryCollection->MaterialID;
+	TManagedArray<int32>& TargetMaterialIndex = GeometryCollection->MaterialIndex;
+
+	// append faces
+	for (int32 FaceIndex = 0; FaceIndex < FaceCount; ++FaceIndex)
+	{
+		const FIntVector& SourceFace = SourceIndices[FaceIndex];
+		const int32 FaceOffset = FaceStart + FaceIndex;
+		TargetIndices[FaceOffset] = FIntVector(
+			SourceFace[0] + VertexStart,
+			SourceFace[1] + VertexStart,
+			SourceFace[2] + VertexStart);
+		TargetVisible[FaceOffset] = SourceVisible[FaceIndex];
+		
+		TargetMaterialID[FaceOffset] = AssetMaterialStart + SourceMaterialID[FaceIndex];
+		TargetMaterialIndex[FaceOffset] = FaceOffset;
+	}
+
+	// source transform information
+	const TManagedArray<FTransform>& SourceTransform = SourceGeometryCollectionPtr->Transform;
+	const TManagedArray<FString>& SourceBoneName = SourceGeometryCollectionPtr->BoneName;
+	const TManagedArray<FLinearColor>& SourceBoneColor = SourceGeometryCollectionPtr->BoneColor;
+	const TManagedArray<int32>& SourceParent = SourceGeometryCollectionPtr->Parent;
+	const TManagedArray<TSet<int32>>& SourceChildren = SourceGeometryCollectionPtr->Children;
+	const TManagedArray<int32>& SourceTransformToGeometryIndex = SourceGeometryCollectionPtr->TransformToGeometryIndex;
+	const TManagedArray<int32>& SourceSimulationType = SourceGeometryCollectionPtr->SimulationType;
+	const TManagedArray<int32>& SourceStatusFlags = SourceGeometryCollectionPtr->StatusFlags;
+	const TManagedArray<int32>& SourceInitialDynamicState = SourceGeometryCollectionPtr->InitialDynamicState;
+	
+	// target transform information
+	TManagedArray<FTransform>& TargetTransform = GeometryCollection->Transform;
+	TManagedArray<FString>& TargetBoneName = GeometryCollection->BoneName;
+	TManagedArray<FLinearColor>& TargetBoneColor = GeometryCollection->BoneColor;
+	TManagedArray<int32>& TargetParent = GeometryCollection->Parent;
+	TManagedArray<TSet<int32>>& TargetChildren = GeometryCollection->Children;
+	TManagedArray<int32>& TargetTransformToGeometryIndex = GeometryCollection->TransformToGeometryIndex;
+	TManagedArray<int32>& TargetSimulationType = GeometryCollection->SimulationType;
+	TManagedArray<int32>& TargetStatusFlags = GeometryCollection->StatusFlags;
+	TManagedArray<int32>& TargetInitialDynamicState = GeometryCollection->InitialDynamicState;
+
+	// append transform hierarchy
+	for (int32 TransformIndex = 0; TransformIndex < TransformCount; ++TransformIndex)
+	{
+		const int32 TransformOffset = TransformStart + TransformIndex;
+
+		// Only apply the transform to the parent node. Child nodes only need scaling applied to translation offsets.
+		if (SourceParent[TransformIndex] == INDEX_NONE)
+		{
+			TargetTransform[TransformOffset] = SourceTransform[TransformIndex] * AppliedTransform;
+		}
+		else
+		{
+			FTransform ScaledTranslation = SourceTransform[TransformIndex];
+			ScaledTranslation.ScaleTranslation(Scale);
+			TargetTransform[TransformOffset] = ScaledTranslation;
+		}
+	
+		// #todo Get this Bone name to be unique
+		TargetBoneName[TransformOffset] = SourceBoneName[TransformIndex];
+
+		const FColor RandBoneColor(FMath::Rand() % 100 + 5, FMath::Rand() % 100 + 5, FMath::Rand() % 100 + 5, 255);
+		TargetBoneColor[TransformOffset] = FLinearColor(RandBoneColor);
+
+		TargetParent[TransformOffset] = (SourceParent[TransformIndex]==INDEX_NONE) ? INDEX_NONE: SourceParent[TransformIndex] + TransformStart;
+
+		const TSet<int32>& SourceChildrenSet = SourceChildren[TransformIndex];
+		for (int32 ChildIndex : SourceChildrenSet)
+		{
+			TargetChildren[TransformOffset].Add(ChildIndex + TransformStart);
+		}
+
+		TargetTransformToGeometryIndex[TransformOffset] = SourceTransformToGeometryIndex[TransformIndex] + GeometryStart;
+		TargetSimulationType[TransformOffset] = SourceSimulationType[TransformIndex];
+		TargetStatusFlags[TransformOffset] = SourceStatusFlags[TransformIndex];
+		TargetInitialDynamicState[TransformOffset] = SourceInitialDynamicState[TransformIndex];
+	}
+
+	// source geometry information
+	const TManagedArray<int32>& SourceTransformIndex = SourceGeometryCollectionPtr->TransformIndex;
+	const TManagedArray<int32>& SourceVertexStart = SourceGeometryCollectionPtr->VertexStart;
+	const TManagedArray<int32>& SourceVertexCount = SourceGeometryCollectionPtr->VertexCount;
+	const TManagedArray<int32>& SourceFaceStart = SourceGeometryCollectionPtr->FaceStart;
+	const TManagedArray<int32>& SourceFaceCount = SourceGeometryCollectionPtr->FaceCount;
+	
+	// target geometry information
+	TManagedArray<int32>& TargetTransformIndex = GeometryCollection->TransformIndex;
+	TManagedArray<FBox>& TargetBoundingBox = GeometryCollection->BoundingBox;
+	TManagedArray<float>& TargetInnerRadius = GeometryCollection->InnerRadius;
+	TManagedArray<float>& TargetOuterRadius = GeometryCollection->OuterRadius;
+	TManagedArray<int32>& TargetVertexStart = GeometryCollection->VertexStart;
+	TManagedArray<int32>& TargetVertexCount = GeometryCollection->VertexCount;
+	TManagedArray<int32>& TargetFaceStart = GeometryCollection->FaceStart;
+	TManagedArray<int32>& TargetFaceCount = GeometryCollection->FaceCount;
+
+	// append geometry
+	for (int32 GeometryIndex = 0; GeometryIndex < GeometryCount; ++GeometryIndex)
+	{
+		const int32 GeometryOffset = GeometryStart + GeometryIndex;
+
+		TargetTransformIndex[GeometryOffset] = SourceTransformIndex[GeometryIndex] + TransformStart;
+		
+		TargetVertexStart[GeometryOffset] = SourceVertexStart[GeometryIndex] + VertexStart;
+		TargetVertexCount[GeometryOffset] = SourceVertexCount[GeometryIndex];
+		TargetFaceStart[GeometryOffset] = SourceFaceStart[GeometryIndex] + FaceStart;
+		TargetFaceCount[GeometryOffset] = SourceFaceCount[GeometryIndex];
+
+		// Find centroid of geometry for inner/outer radius calculations
+		FVector Center(0);
+		for (int32 VertexIndex = TargetVertexStart[GeometryOffset]; VertexIndex < TargetVertexStart[GeometryOffset] + TargetVertexCount[GeometryOffset]; ++VertexIndex)
+		{
+			Center += (FVector)TargetVertex[VertexIndex];
+		}
+		if (TargetVertexCount[GeometryOffset]) Center /= TargetVertexCount[GeometryOffset];
+		
+		TargetBoundingBox[GeometryOffset] = FBox(ForceInitToZero);
+		TargetInnerRadius[GeometryOffset] = FLT_MAX;
+		TargetOuterRadius[GeometryOffset] = -FLT_MAX;
+		for (int32 VertexIndex = TargetVertexStart[GeometryOffset]; VertexIndex < TargetVertexStart[GeometryOffset] + TargetVertexCount[GeometryOffset]; ++VertexIndex)
+		{
+			TargetBoundingBox[GeometryOffset] += (FVector)TargetVertex[VertexIndex];
+
+			float Delta = (Center - (FVector)TargetVertex[VertexIndex]).Size();
+			TargetInnerRadius[GeometryOffset] = FMath::Min(TargetInnerRadius[GeometryOffset], Delta);
+			TargetOuterRadius[GeometryOffset] = FMath::Max(TargetOuterRadius[GeometryOffset], Delta);
+		}
+	}
+
+	// for each material, add a reference in our GeometryCollectionObject
+	const int32 NumMeshMaterials = Materials.Num();
+	TargetGeometryCollectionObject->Materials.Reserve(AssetMaterialStart + NumMeshMaterials);
+
+	for (int32 Index = 0; Index < NumMeshMaterials; ++Index)
+	{
+		UMaterialInterface* CurrMaterial = Materials[Index];
+
+		// Possible we have a null entry - replace with default
+		if (CurrMaterial == nullptr)
+		{
+			CurrMaterial = UMaterial::GetDefaultMaterial(MD_Surface);
+		}
+
+		TargetGeometryCollectionObject->Materials.Add(CurrMaterial);
+	}
+
+	// source material information
+	const TManagedArray<FGeometryCollectionSection>& SourceSections = SourceGeometryCollectionPtr->Sections;
+
+	// target material information
+	TManagedArray<FGeometryCollectionSection>& TargetSections = GeometryCollection->Sections;
+
+	// append sections
+	for (int32 SectionIndex = 0; SectionIndex < SectionCount; ++SectionIndex)
+	{
+		int32 SectionOffset = SectionStart + SectionIndex;
+
+		TargetSections[SectionOffset].MaterialID = AssetMaterialStart + SourceSections[SectionIndex].MaterialID;
+
+		TargetSections[SectionOffset].FirstIndex = SourceSections[SectionIndex].FirstIndex + FaceStart * 3;
+		TargetSections[SectionOffset].MinVertexIndex = VertexStart + SourceSections[SectionIndex].MinVertexIndex;
+
+		TargetSections[SectionOffset].NumTriangles = SourceSections[SectionIndex].NumTriangles;
+		TargetSections[SectionOffset].MaxVertexIndex = VertexStart + SourceSections[SectionIndex].MaxVertexIndex;
+	}
+
+	if (ReindexMaterials) {
+		GeometryCollection->ReindexMaterials();
+	}
+
+}
+
+
+void FGeometryCollectionEngineConversion::AppendStaticMesh(const UStaticMesh* StaticMesh, const UStaticMeshComponent* StaticMeshComponent, const FTransform& StaticMeshTransform, UGeometryCollection* GeometryCollectionObject, bool ReindexMaterials)
+{
+	if (StaticMesh == nullptr)
+	{
+		return;
+	}
+
+	TArray<UMaterialInterface*> Materials;
+	Materials.Reserve(StaticMesh->GetStaticMaterials().Num());
+
+	for (int32 Index = 0; Index < StaticMesh->GetStaticMaterials().Num(); ++Index)
+	{
+		UMaterialInterface* CurrMaterial = StaticMeshComponent ? StaticMeshComponent->GetMaterial(Index) : StaticMesh->GetMaterial(Index);
+		Materials.Add(CurrMaterial);
+	}
+
+	// Geometry collections usually carry the selection material, which we'll delete before appending
+	UMaterialInterface* BoneSelectedMaterial = LoadObject<UMaterialInterface>(nullptr, UGeometryCollection::GetSelectedMaterialPath(), nullptr, LOAD_None, nullptr);
+	GeometryCollectionObject->Materials.Remove(BoneSelectedMaterial);
+	Materials.Remove(BoneSelectedMaterial);
+	
+	AppendStaticMesh(StaticMesh, Materials, StaticMeshTransform, GeometryCollectionObject, ReindexMaterials);
+}
+
+
+void FGeometryCollectionEngineConversion::AppendGeometryCollection(const UGeometryCollection* SourceGeometryCollection, const UGeometryCollectionComponent* GeometryCollectionComponent, const FTransform& GeometryCollectionTransform, UGeometryCollection* TargetGeometryCollectionObject, bool ReindexMaterials)
+{
+	if (SourceGeometryCollection == nullptr)
+	{
+		return;
+	}
+
+	TArray<UMaterialInterface*> Materials;
+	Materials.Reserve(SourceGeometryCollection->Materials.Num());
+
+	for (int32 Index = 0; Index < SourceGeometryCollection->Materials.Num(); ++Index)
+	{
+		UMaterialInterface* CurrMaterial = GeometryCollectionComponent ? GeometryCollectionComponent->GetMaterial(Index) : SourceGeometryCollection->Materials[Index].Get();
+		Materials.Add(CurrMaterial);
+	}
+
+	// Geometry collections usually carry the selection material, which we'll delete before appending
+	UMaterialInterface* BoneSelectedMaterial = LoadObject<UMaterialInterface>(nullptr, UGeometryCollection::GetSelectedMaterialPath(), nullptr, LOAD_None, nullptr);
+	TargetGeometryCollectionObject->Materials.Remove(BoneSelectedMaterial);
+	Materials.Remove(BoneSelectedMaterial);
+
+	AppendGeometryCollection(SourceGeometryCollection, Materials, GeometryCollectionTransform, TargetGeometryCollectionObject, ReindexMaterials);
+}
+
+
+void FGeometryCollectionEngineConversion::AppendSkeletalMesh(const USkeletalMesh* SkeletalMesh, const USkeletalMeshComponent *SkeletalMeshComponent, const FTransform & SkeletalMeshTransform, UGeometryCollection * GeometryCollectionObject, bool ReindexMaterials)
+{
+	//UE_LOG(UGeometryCollectionConversionLogging, Log, TEXT("FGeometryCollectionEngineConversion::AppendSkeletalMesh()"));
+	check(SkeletalMesh);
+	if (GeometryCollectionObject)
+	{
+		TSharedPtr<FGeometryCollection, ESPMode::ThreadSafe> GeometryCollectionPtr = GeometryCollectionObject->GetGeometryCollection();
+		if (FGeometryCollection* GeometryCollection = GeometryCollectionPtr.Get())
+		{
+
+			if (const USkeleton* Skeleton = SkeletalMesh->GetSkeleton())
+			{
+				if (const FSkeletalMeshRenderData * SkelMeshRenderData = SkeletalMesh->GetResourceForRendering())
+				{
+					if (SkelMeshRenderData->LODRenderData.Num())
+					{
+						const FSkeletalMeshLODRenderData& SkeletalMeshLODRenderData = SkelMeshRenderData->LODRenderData[0];
+						const FSkinWeightVertexBuffer& SkinWeightVertexBuffer = *SkeletalMeshLODRenderData.GetSkinWeightVertexBuffer();
+						const FStaticMeshVertexBuffers& VertexBuffers = SkeletalMeshLODRenderData.StaticVertexBuffers;
+						const FPositionVertexBuffer& PositionVertexBuffer = VertexBuffers.PositionVertexBuffer;
+						const int32 VertexCount = PositionVertexBuffer.GetNumVertices();
+						// Check that all vertex weightings are rigid. 
+						for (int32 VertexIndex = 0; VertexIndex < VertexCount; VertexIndex++)
+						{
+							int32 SkeletalBoneIndex = -1;
+							if (!SkinWeightVertexBuffer.GetRigidWeightBone(VertexIndex, SkeletalBoneIndex))
+							{
+								UE_LOG(UGeometryCollectionConversionLogging, Error, TEXT("Non-rigid weighting found on vertex %d: Cannot convert to GeometryCollection."), VertexIndex);
+								return;
+							}
+						}
+						const FSkelMeshRenderSection& RenderSection = SkeletalMeshLODRenderData.RenderSections[0];
+						const TArray<FBoneIndexType>& SkeletalBoneMap = RenderSection.BoneMap;
+
+						//
+						// The Component transform for each Mesh will become the FTransform that drives
+						// its associated VerticesGroup. The Skeleton will contain a nested transform hierarchy
+						// that is evaluated using the GetComponentSpaceTransformRefPose. The resulting
+						// Transforms array stored in the GeometryCollection will be the same size as
+						// the SkeletalBoneMap. Note the @todo: the SkeletalBoneMap is pulled from only
+						// the first render section, this will need to be expanded to include all render
+						// sections.
+						//
+						TManagedArray<FTransform>& Transform = GeometryCollection->Transform;
+						int32 TransformBaseIndex = GeometryCollection->AddElements(SkeletalBoneMap.Num(), FGeometryCollection::TransformGroup);
+						const FReferenceSkeleton & ReferenceSkeletion = Skeleton->GetReferenceSkeleton();
+						const TArray<FTransform> & RestArray = Skeleton->GetRefLocalPoses();
+						for (int32 BoneIndex = 0; BoneIndex < SkeletalBoneMap.Num(); BoneIndex++)
+						{
+							FTransform BoneTransform = FAnimationRuntime::GetComponentSpaceTransformRefPose(ReferenceSkeletion, SkeletalBoneMap[BoneIndex]);
+							Transform[TransformBaseIndex + BoneIndex] = BoneTransform;
+						}
+
+
+						//
+						// The Triangle Indices
+						//
+						TManagedArray<FIntVector>& Indices = GeometryCollection->Indices;
+						TManagedArray<bool>& Visible = GeometryCollection->Visible;
+						TManagedArray<int32>& MaterialID = GeometryCollection->MaterialID;
+						TManagedArray<int32>& MaterialIndex = GeometryCollection->MaterialIndex;
+
+						TArray<uint32> IndexBuffer;
+						SkeletalMeshLODRenderData.MultiSizeIndexContainer.GetIndexBuffer(IndexBuffer);
+
+						const int32 IndicesCount = IndexBuffer.Num() / 3;
+						int NumVertices = GeometryCollection->NumElements(FGeometryCollection::VerticesGroup);
+						int InitialNumIndices = GeometryCollection->NumElements(FGeometryCollection::FacesGroup);
+						int IndicesBaseIndex = GeometryCollection->AddElements(IndicesCount, FGeometryCollection::FacesGroup);
+						for (int32 IndicesIndex = 0, StaticIndex = 0; IndicesIndex < IndicesCount; IndicesIndex++, StaticIndex += 3)
+						{
+							int32 IndicesOffset = IndicesBaseIndex + IndicesIndex;
+							Indices[IndicesOffset] = FIntVector(
+								IndexBuffer[StaticIndex] + NumVertices,
+								IndexBuffer[StaticIndex + 1] + NumVertices,
+								IndexBuffer[StaticIndex + 2] + NumVertices);
+							Visible[IndicesOffset] = true;
+							MaterialID[IndicesOffset] = 0;
+							MaterialIndex[IndicesOffset] = IndicesOffset;
+						}
+
+						//
+						// Vertex Attributes
+						//
+						TManagedArray<FVector3f>& Vertex = GeometryCollection->Vertex;
+						TManagedArray<FVector3f>& TangentU = GeometryCollection->TangentU;
+						TManagedArray<FVector3f>& TangentV = GeometryCollection->TangentV;
+						TManagedArray<FVector3f>& Normal = GeometryCollection->Normal;
+						TManagedArray<TArray<FVector2f>>& UVs = GeometryCollection->UVs;
+						TManagedArray<FLinearColor>& Color = GeometryCollection->Color;
+						TManagedArray<int32>& BoneMap = GeometryCollection->BoneMap;
+						TManagedArray<FLinearColor>& BoneColor = GeometryCollection->BoneColor;
+						TManagedArray<FString>& BoneName = GeometryCollection->BoneName;
+
+						// 
+						// Transform Attributes 
+						// 
+						TManagedArray<int32>& Parent = GeometryCollection->Parent;
+						TManagedArray<int32>& SimulationType = GeometryCollection->SimulationType;
+						int InitialNumVertices = GeometryCollection->NumElements(FGeometryCollection::VerticesGroup);
+						int VertexBaseIndex = GeometryCollection->AddElements(VertexCount, FGeometryCollection::VerticesGroup);
+						const int32 NumUVLayers = VertexBuffers.StaticMeshVertexBuffer.GetNumTexCoords();
+						for (int32 VertexIndex = 0; VertexIndex < VertexCount; VertexIndex++)
+						{
+							int VertexOffset = VertexBaseIndex + VertexIndex;
+							BoneMap[VertexOffset] = -1;
+							int32 SkeletalBoneIndex = -1;
+							SkinWeightVertexBuffer.GetRigidWeightBone(VertexIndex, SkeletalBoneIndex);
+							if (SkeletalBoneIndex > -1)
+							{
+								BoneMap[VertexOffset] = SkeletalBoneIndex + TransformBaseIndex;
+								Vertex[VertexOffset] = (FVector4f)Transform[BoneMap[VertexOffset]].ToInverseMatrixWithScale().TransformPosition((FVector)PositionVertexBuffer.VertexPosition(VertexIndex));
+							}
+							check(BoneMap[VertexOffset] != -1);
+							TangentU[VertexOffset] = VertexBuffers.StaticMeshVertexBuffer.VertexTangentX(VertexIndex);
+							TangentV[VertexOffset] = VertexBuffers.StaticMeshVertexBuffer.VertexTangentY(VertexIndex);
+							Normal[VertexOffset] = VertexBuffers.StaticMeshVertexBuffer.VertexTangentZ(VertexIndex);
+							
+							UVs[VertexOffset].SetNum(NumUVLayers);
+							for (int32 UVLayerIdx = 0; UVLayerIdx < NumUVLayers; ++UVLayerIdx)
+							{
+								UVs[VertexOffset][UVLayerIdx] = VertexBuffers.StaticMeshVertexBuffer.GetVertexUV(VertexIndex, UVLayerIdx);
+							}
+							
+							if (VertexBuffers.ColorVertexBuffer.GetNumVertices() == VertexCount)
+								Color[VertexOffset] = VertexBuffers.ColorVertexBuffer.VertexColor(VertexIndex);
+							else
+								Color[VertexOffset] = FLinearColor::White;
+						}
+
+						int32 InitialIndex = -1;
+						int32 LastParentIndex = -1;
+						int32 CurrentLevel = 0;
+						for (int32 BoneIndex = 0; BoneIndex < SkeletalBoneMap.Num(); BoneIndex++)
+						{
+							// transform based on position of the actor. 
+							Transform[TransformBaseIndex + BoneIndex] = SkeletalMeshTransform * Transform[TransformBaseIndex + BoneIndex];
+
+							// bone attributes
+							BoneName[TransformBaseIndex + BoneIndex] = ReferenceSkeletion.GetBoneName(SkeletalBoneMap[BoneIndex]).ToString();
+							const FColor RandBoneColor(FMath::Rand() % 100 + 5, FMath::Rand() % 100 + 5, FMath::Rand() % 100 + 5, 255);
+							BoneColor[TransformBaseIndex + BoneIndex] = FLinearColor(RandBoneColor);
+
+							// Bone Hierarchy - Added at root with no common parent
+							int32 ParentIndex = ReferenceSkeletion.GetParentIndex(SkeletalBoneMap[BoneIndex]);
+							int32 UseParentIndex = ParentIndex + InitialIndex;
+							if (LastParentIndex != UseParentIndex)
+							{
+								LastParentIndex = UseParentIndex;
+							}
+							Parent[TransformBaseIndex + BoneIndex] = UseParentIndex;
+							SimulationType[TransformBaseIndex + BoneIndex] = FGeometryCollection::ESimulationTypes::FST_Rigid;
+						}
+
+						// Geometry Group
+						TArray<int32> GeometryIndices;
+						GeometryCollectionAlgo::ContiguousArray(GeometryIndices, GeometryCollection->NumElements(FGeometryCollection::GeometryGroup));
+						GeometryCollection->RemoveDependencyFor(FGeometryCollection::GeometryGroup);
+						GeometryCollection->RemoveElements(FGeometryCollection::GeometryGroup, GeometryIndices);
+						::GeometryCollection::AddGeometryProperties(GeometryCollection);
+
+						// for each material, add a reference in our GeometryCollectionObject
+						int CurrIdx = 0;
+						
+						const TArray<FSkeletalMaterial>& SkeletalMeshMaterials = SkeletalMesh->GetMaterials();
+
+						UMaterialInterface* CurrMaterial = SkeletalMeshComponent ? SkeletalMeshComponent->GetMaterial(CurrIdx) : ToRawPtr(SkeletalMeshMaterials[CurrIdx].MaterialInterface);
+
+
+						int MaterialStart = GeometryCollectionObject->Materials.Num();
+						while (CurrMaterial)
+						{
+							GeometryCollectionObject->Materials.Add(CurrMaterial);
+							CurrMaterial = SkeletalMeshComponent ? SkeletalMeshComponent->GetMaterial(++CurrIdx) : ToRawPtr(SkeletalMeshMaterials[++CurrIdx].MaterialInterface);
+
+						}
+
+						const TArray<FSkelMeshRenderSection> &StaticMeshSections = SkeletalMesh->GetResourceForRendering()->LODRenderData[0].RenderSections;
+
+						TManagedArray<FGeometryCollectionSection> & Sections = GeometryCollection->Sections;
+
+						for (const FSkelMeshRenderSection &CurrSection : StaticMeshSections)
+						{
+							// create new section
+							int32 SectionIndex = GeometryCollection->AddElements(1, FGeometryCollection::MaterialGroup);
+						
+							Sections[SectionIndex].MaterialID = MaterialStart + CurrSection.MaterialIndex;
+
+							Sections[SectionIndex].FirstIndex = IndicesBaseIndex * 3 + CurrSection.BaseIndex;
+							Sections[SectionIndex].MinVertexIndex = VertexBaseIndex + CurrSection.BaseVertexIndex;
+
+							Sections[SectionIndex].NumTriangles = CurrSection.NumTriangles;
+
+							// #todo(dmp): what should we set this to?  SkeletalMesh sections are different
+							// but we are resetting this when the re indexing happens
+							Sections[SectionIndex].MaxVertexIndex = VertexBaseIndex + CurrSection.NumVertices;
+
+							// set the materialid for all of the faces
+							for (int32 i = Sections[SectionIndex].FirstIndex / 3; i < Sections[SectionIndex].FirstIndex / 3 + Sections[SectionIndex].NumTriangles; ++i)
+							{
+								MaterialID[i] = SectionIndex;
+							}
+						}
+
+					}
+				}
+			}
+
+			if (ReindexMaterials) {
+				GeometryCollection->ReindexMaterials();
+			}
+		}
+	}
+}
+
diff --git a/Engine/Source/Runtime/Experimental/GeometryCollectionEngine/Public/GeometryCollection/GeometryCollectionEngineConversion.h b/Engine/Source/Runtime/Experimental/GeometryCollectionEngine/Public/GeometryCollection/GeometryCollectionEngineConversion.h
new file mode 100644
--- /dev/null
+++ b/Engine/Source/Runtime/Experimental/GeometryCollectionEngine/Public/GeometryCollection/GeometryCollectionEngineConversion.h
@@ -0,0 +1,64 @@
+// Copyright Epic Games, Inc. All Rights Reserved.
+
+#pragma once
+#include "Engine/World.h"
+
+class UStaticMesh;
+class USkeletalMesh;
+class UGeometryCollection;
+class UMaterialInterface;
+class UGeometryCollectionComponent;
+
+typedef TTuple<const UStaticMesh *, const UStaticMeshComponent *, FTransform> GeometryCollectionStaticMeshConversionTuple;
+typedef TTuple<const USkeletalMesh *, const USkeletalMeshComponent *, FTransform> GeometryCollectionSkeletalMeshConversionTuple;
+
+/**
+* The public interface to this module
+*/
+class GEOMETRYCOLLECTIONENGINE_API FGeometryCollectionEngineConversion
+{
+public:
+
+	/**
+	*  Appends a static mesh to a GeometryCollectionComponent.
+	*  @param StaticMesh : Const mesh to read vertex/normals/index data from
+	*  @param Materials : Materials fetched from the StaticMeshComponent used to configure this geometry
+	*  @param StaticMeshTransform : Mesh transform.
+	*  @param GeometryCollection  : Collection to append the mesh into.
+	*/
+	static void AppendStaticMesh(const UStaticMesh* StaticMesh, const TArray<UMaterialInterface*>& Materials, const FTransform& StaticMeshTransform, UGeometryCollection* GeometryCollectionObject, bool ReindexMaterials = true);
+
+	/**
+	*  Appends a static mesh to a GeometryCollectionComponent.
+	*  @param StaticMeshComponent : Const mesh to read vertex/normals/index data from
+	*  @param StaticMeshTransform : Mesh transform.
+	*  @param GeometryCollection  : Collection to append the mesh into.
+	*/
+	static void AppendStaticMesh(const UStaticMesh * StaticMesh, const UStaticMeshComponent *StaticMeshComponent, const FTransform & StaticMeshTransform, UGeometryCollection * GeometryCollection, bool ReindexMaterials = true);
+
+	/**
+	*  Appends a skeletal mesh to a GeometryCollectionComponent.
+	*  @param SkeletalMeshComponent : Const mesh to read vertex/normals/index data from
+	*  @param SkeletalMeshTransform : Mesh transform.
+	*  @param GeometryCollection    : Collection to append the mesh into.
+	*/
+	static void AppendSkeletalMesh(const USkeletalMesh* SkeletalMesh, const USkeletalMeshComponent *SkeletalMeshComponent, const FTransform & SkeletalMeshTransform, UGeometryCollection * GeometryCollection, bool ReindexMaterials = true);
+
+	/**
+	*  Appends a GeometryCollection to a GeometryCollectionComponent.
+	*  @param SourceGeometryCollection : Const GeometryCollection to read vertex/normals/index data from
+	*  @param Materials : Materials fetched from the GeometryCollectionComponent used to configure this geometry
+	*  @param GeometryCollectionTransform : GeometryCollection transform.
+	*  @param TargetGeometryCollection  : Collection to append the GeometryCollection into.
+	*/
+	static void AppendGeometryCollection(const UGeometryCollection* SourceGeometryCollection, const TArray<UMaterialInterface*>& Materials, const FTransform& GeometryCollectionTransform, UGeometryCollection* TargetGeometryCollectionObject, bool ReindexMaterials = true);
+
+	/**
+	*  Appends a GeometryCollection to a GeometryCollectionComponent.
+	*  @param GeometryCollectionComponent : Const GeometryCollection to read vertex/normals/index data from
+	*  @param GeometryCollectionTransform : GeometryCollection transform.
+	*  @param TargetGeometryCollection  : Collection to append the GeometryCollection into.
+	*/
+	static void AppendGeometryCollection(const UGeometryCollection* SourceGeometryCollection, const UGeometryCollectionComponent* GeometryCollectionComponent, const FTransform& GeometryCollectionTransform, UGeometryCollection* TargetGeometryCollectionObject, bool ReindexMaterials = true);
+
+};
\ No newline at end of file
