GeometryCollection : ResetNode for UGeometryCollection

diff --git a/Engine/Plugins/Experimental/GeometryCollectionPlugin/Source/GeometryCollectionNodes/Private/Dataflow/GeometryCollectionNodes.cpp b/Engine/Plugins/Experimental/GeometryCollectionPlugin/Source/GeometryCollectionNodes/Private/Dataflow/GeometryCollectionNodes.cpp
--- a/Engine/Plugins/Experimental/GeometryCollectionPlugin/Source/GeometryCollectionNodes/Private/Dataflow/GeometryCollectionNodes.cpp
+++ b/Engine/Plugins/Experimental/GeometryCollectionPlugin/Source/GeometryCollectionNodes/Private/Dataflow/GeometryCollectionNodes.cpp
@@ -9,6 +9,7 @@ namespace Dataflow
 		DATAFLOW_NODE_REGISTER_CREATION_FACTORY(GetCollectionAssetNode);
 		DATAFLOW_NODE_REGISTER_CREATION_FACTORY(ExampleCollectionEditNode);
 		DATAFLOW_NODE_REGISTER_CREATION_FACTORY(SetCollectionAssetNode);
+		DATAFLOW_NODE_REGISTER_CREATION_FACTORY(ResetGeometryCollection);
 	}
 }
 
diff --git a/Engine/Plugins/Experimental/GeometryCollectionPlugin/Source/GeometryCollectionNodes/Public/Dataflow/GeometryCollectionNodes.h b/Engine/Plugins/Experimental/GeometryCollectionPlugin/Source/GeometryCollectionNodes/Public/Dataflow/GeometryCollectionNodes.h
--- a/Engine/Plugins/Experimental/GeometryCollectionPlugin/Source/GeometryCollectionNodes/Public/Dataflow/GeometryCollectionNodes.h
+++ b/Engine/Plugins/Experimental/GeometryCollectionPlugin/Source/GeometryCollectionNodes/Public/Dataflow/GeometryCollectionNodes.h
@@ -7,10 +7,17 @@
 #include "CoreMinimal.h"
 #include "Dataflow/DataflowObjectInterface.h"
 #include "Dataflow/DataflowEngine.h"
+#include "Dataflow/DataflowEngineUtil.h"
 #include "Dataflow/DataflowProperty.h"
+#include "Engine/StaticMesh.h"
 #include "GeometryCollection/GeometryCollectionObject.h"
 #include "GeometryCollection/ManagedArrayCollection.h"
 #include "GeometryCollection/GeometryCollection.h"
+#include "GeometryCollection/GeometryCollectionEngineUtility.h"
+#include "GeometryCollection/GeometryCollectionEngineConversion.h"
+#include "Logging/LogMacros.h"
+#include "Templates/SharedPointer.h"
+#include "UObject/UnrealTypePrivate.h"
 
 
 namespace Dataflow
@@ -119,6 +126,83 @@ namespace Dataflow
 	};
 
 
+	class ResetGeometryCollection : public FNode
+	{
+		DATAFLOW_NODE_DEFINE_INTERNAL(ResetGeometryCollection)
+
+	public:
+		typedef FManagedArrayCollectionSharedPtr DataType;
+
+		TSharedPtr< class TOutput<DataType> > ManagedArrayOut;
+
+		ResetGeometryCollection(const FNodeParameters& InParam, FGuid InGuid = FGuid::NewGuid())
+			: FNode(InParam, InGuid)
+			, ManagedArrayOut(new TOutput<DataType>(TOutputParameters<DataType>({ FName("ManagedArrayOut"), this })))
+		{}
+
+
+		virtual void Evaluate(const FContext& Context, FConnection* Out) const override
+		{
+			ManagedArrayOut->SetValue(TSharedPtr<FManagedArrayCollection>(nullptr), Context);
+
+			if (const FEngineContext* EngineContext = (const FEngineContext*)(&Context))
+			{
+				if (UGeometryCollection* GeometryCollectionObject = Cast<UGeometryCollection>(EngineContext->Owner))
+				{
+					GeometryCollectionObject->Reset();
+
+					const UObject* Owner = EngineContext->Owner;
+					FName AName("GeometrySource");
+					if (Owner && Owner->GetClass())
+					{
+						if (const ::FProperty* UEProperty = Owner->GetClass()->FindPropertyByName(AName))
+						{
+							if (const FArrayProperty* ArrayProperty = CastField<FArrayProperty>(UEProperty))
+							{
+								FScriptArrayHelper_InContainer ArrayHelper(ArrayProperty, Owner);
+								const int32 ArraySize = ArrayHelper.Num();
+								for (int32 Index = 0; Index < ArraySize; ++Index)
+								{
+									if (FGeometryCollectionSource* SourceObject = (FGeometryCollectionSource*)(ArrayHelper.GetRawPtr(Index)))
+									{
+										if (UObject* ResolvedObject = SourceObject->SourceGeometryObject.ResolveObject())
+										{
+											if (UStaticMesh* StaticMesh = Cast<UStaticMesh>(ResolvedObject))
+											{
+												TArray<UMaterialInterface*> Materials;
+												Materials.Reserve(StaticMesh->GetStaticMaterials().Num());
+
+												for (int32 Index2 = 0; Index2 < StaticMesh->GetStaticMaterials().Num(); ++Index2)
+												{
+													UMaterialInterface* CurrMaterial = StaticMesh->GetMaterial(Index2);
+													Materials.Add(CurrMaterial);
+												}
+
+												// Geometry collections usually carry the selection material, which we'll delete before appending
+												UMaterialInterface* BoneSelectedMaterial = LoadObject<UMaterialInterface>(nullptr, UGeometryCollection::GetSelectedMaterialPath(), nullptr, LOAD_None, nullptr);
+												GeometryCollectionObject->Materials.Remove(BoneSelectedMaterial);
+												Materials.Remove(BoneSelectedMaterial);
+
+												FGeometryCollectionEngineConversion::AppendStaticMesh(StaticMesh, Materials, FTransform(), GeometryCollectionObject);
+											}
+										}
+									}
+								}
+							}
+						}
+					}
+					GeometryCollectionObject->InvalidateCollection();
+
+					if (const TSharedPtr<FGeometryCollection, ESPMode::ThreadSafe> AssetCollection = GeometryCollectionObject->GetGeometryCollection())
+					{
+						FManagedArrayCollection* NewCollection = AssetCollection->NewCopy<FManagedArrayCollection>();
+						ManagedArrayOut->SetValue(FManagedArrayCollectionSharedPtr(NewCollection), Context);
+					}
+				}
+			}
+		}
+	};
+
 	void GeometryCollectionEngineAssetNodes();
 
 
