GeometryCollection : Nodes
- Moved to cpp

diff --git a/Engine/Plugins/Experimental/GeometryCollectionPlugin/Source/GeometryCollectionNodes/Private/Dataflow/GeometryCollectionNodes.cpp b/Engine/Plugins/Experimental/GeometryCollectionPlugin/Source/GeometryCollectionNodes/Private/Dataflow/GeometryCollectionNodes.cpp
--- a/Engine/Plugins/Experimental/GeometryCollectionPlugin/Source/GeometryCollectionNodes/Private/Dataflow/GeometryCollectionNodes.cpp
+++ b/Engine/Plugins/Experimental/GeometryCollectionPlugin/Source/GeometryCollectionNodes/Private/Dataflow/GeometryCollectionNodes.cpp
@@ -2,6 +2,16 @@
 
 #include "Dataflow/GeometryCollectionNodes.h"
 
+#include "Engine/StaticMesh.h"
+#include "GeometryCollection/GeometryCollectionObject.h"
+#include "GeometryCollection/ManagedArrayCollection.h"
+#include "GeometryCollection/GeometryCollection.h"
+#include "GeometryCollection/GeometryCollectionEngineUtility.h"
+#include "GeometryCollection/GeometryCollectionEngineConversion.h"
+#include "Logging/LogMacros.h"
+#include "Templates/SharedPointer.h"
+#include "UObject/UnrealTypePrivate.h"
+
 namespace Dataflow
 {
 	void GeometryCollectionEngineAssetNodes()
@@ -11,5 +21,118 @@ namespace Dataflow
 		DATAFLOW_NODE_REGISTER_CREATION_FACTORY(SetCollectionAssetNode);
 		DATAFLOW_NODE_REGISTER_CREATION_FACTORY(ResetGeometryCollection);
 	}
+
+	void GetCollectionAssetNode::Evaluate(const FContext& Context, FConnection* Out) const
+	{
+		if (Output.Get() == Out)
+		{
+			if (const FEngineContext* EngineContext = Context.AsType<FEngineContext>(FName("UGeometryCollection")))
+			{
+				if (UGeometryCollection* CollectionAsset = Cast<UGeometryCollection>(EngineContext->Owner))
+				{
+					if (const TSharedPtr<FGeometryCollection, ESPMode::ThreadSafe> AssetCollection = CollectionAsset->GetGeometryCollection())
+					{
+						FManagedArrayCollection* NewCollection = AssetCollection->NewCopy<FManagedArrayCollection>();
+						Output->SetValue(FManagedArrayCollectionSharedPtr(NewCollection), Context);
+					}
+				}
+			}
+		}
+	}
+
+	void ExampleCollectionEditNode::Evaluate(const FContext& Context, FConnection* Out) const
+	{
+		if (Output.Get() == Out)
+		{
+			DataType Collection = Input->GetValue(Context);
+			if (Active.GetValue())
+			{
+				TManagedArray<FVector3f>* Vertex = Collection->FindAttribute<FVector3f>("Vertex", "Vertices");
+				for (int i = 0; i < Vertex->Num(); i++)
+				{
+					(*Vertex)[i][1] *= Scale.GetValue();
+				}
+			}
+			Output->SetValue(Collection, Context);
+		}
+	}
+
+	void SetCollectionAssetNode::Evaluate(const FContext& Context, FConnection* Out) const
+	{
+		if (Out == nullptr)
+		{
+			DataType Collection = Input->GetValue(Context);
+			if (const FEngineContext* EngineContext = Context.AsType<FEngineContext>(FName("UGeometryCollection")))
+			{
+				if (UGeometryCollection* CollectionAsset = Cast<UGeometryCollection>(EngineContext->Owner))
+				{
+					TSharedPtr<FGeometryCollection, ESPMode::ThreadSafe> NewCollection(Collection->NewCopy<FGeometryCollection>());
+					CollectionAsset->SetGeometryCollection(NewCollection);
+					CollectionAsset->InvalidateCollection();
+				}
+			}
+		}
+	}
+
+	void ResetGeometryCollection::Evaluate(const FContext& Context, FConnection* Out) const
+	{
+		ManagedArrayOut->SetValue(TSharedPtr<FManagedArrayCollection>(nullptr), Context);
+
+		if (const FEngineContext* EngineContext = (const FEngineContext*)(&Context))
+		{
+			if (UGeometryCollection* GeometryCollectionObject = Cast<UGeometryCollection>(EngineContext->Owner))
+			{
+				GeometryCollectionObject->Reset();
+
+				const UObject* Owner = EngineContext->Owner;
+				FName AName("GeometrySource");
+				if (Owner && Owner->GetClass())
+				{
+					if (const ::FProperty* UEProperty = Owner->GetClass()->FindPropertyByName(AName))
+					{
+						if (const FArrayProperty* ArrayProperty = CastField<FArrayProperty>(UEProperty))
+						{
+							FScriptArrayHelper_InContainer ArrayHelper(ArrayProperty, Owner);
+							const int32 ArraySize = ArrayHelper.Num();
+							for (int32 Index = 0; Index < ArraySize; ++Index)
+							{
+								if (FGeometryCollectionSource* SourceObject = (FGeometryCollectionSource*)(ArrayHelper.GetRawPtr(Index)))
+								{
+									if (UObject* ResolvedObject = SourceObject->SourceGeometryObject.ResolveObject())
+									{
+										if (UStaticMesh* StaticMesh = Cast<UStaticMesh>(ResolvedObject))
+										{
+											TArray<UMaterialInterface*> Materials;
+											Materials.Reserve(StaticMesh->GetStaticMaterials().Num());
+
+											for (int32 Index2 = 0; Index2 < StaticMesh->GetStaticMaterials().Num(); ++Index2)
+											{
+												UMaterialInterface* CurrMaterial = StaticMesh->GetMaterial(Index2);
+												Materials.Add(CurrMaterial);
+											}
+
+											// Geometry collections usually carry the selection material, which we'll delete before appending
+											UMaterialInterface* BoneSelectedMaterial = LoadObject<UMaterialInterface>(nullptr, UGeometryCollection::GetSelectedMaterialPath(), nullptr, LOAD_None, nullptr);
+											GeometryCollectionObject->Materials.Remove(BoneSelectedMaterial);
+											Materials.Remove(BoneSelectedMaterial);
+
+											FGeometryCollectionEngineConversion::AppendStaticMesh(StaticMesh, Materials, FTransform(), GeometryCollectionObject);
+										}
+									}
+								}
+							}
+						}
+					}
+				}
+				GeometryCollectionObject->InvalidateCollection();
+
+				if (const TSharedPtr<FGeometryCollection, ESPMode::ThreadSafe> AssetCollection = GeometryCollectionObject->GetGeometryCollection())
+				{
+					FManagedArrayCollection* NewCollection = AssetCollection->NewCopy<FManagedArrayCollection>();
+					ManagedArrayOut->SetValue(FManagedArrayCollectionSharedPtr(NewCollection), Context);
+				}
+			}
+		}
+	}
 }
 
diff --git a/Engine/Plugins/Experimental/GeometryCollectionPlugin/Source/GeometryCollectionNodes/Public/Dataflow/GeometryCollectionNodes.h b/Engine/Plugins/Experimental/GeometryCollectionPlugin/Source/GeometryCollectionNodes/Public/Dataflow/GeometryCollectionNodes.h
--- a/Engine/Plugins/Experimental/GeometryCollectionPlugin/Source/GeometryCollectionNodes/Public/Dataflow/GeometryCollectionNodes.h
+++ b/Engine/Plugins/Experimental/GeometryCollectionPlugin/Source/GeometryCollectionNodes/Public/Dataflow/GeometryCollectionNodes.h
@@ -5,19 +5,8 @@
 #include "GeometryCollectionNodeConnectionTypes.h"
 
 #include "CoreMinimal.h"
-#include "Dataflow/DataflowObjectInterface.h"
 #include "Dataflow/DataflowEngine.h"
-#include "Dataflow/DataflowEngineUtil.h"
-#include "Dataflow/DataflowProperty.h"
-#include "Engine/StaticMesh.h"
-#include "GeometryCollection/GeometryCollectionObject.h"
-#include "GeometryCollection/ManagedArrayCollection.h"
-#include "GeometryCollection/GeometryCollection.h"
-#include "GeometryCollection/GeometryCollectionEngineUtility.h"
-#include "GeometryCollection/GeometryCollectionEngineConversion.h"
-#include "Logging/LogMacros.h"
-#include "Templates/SharedPointer.h"
-#include "UObject/UnrealTypePrivate.h"
+
 
 
 namespace Dataflow
@@ -37,23 +26,7 @@ namespace Dataflow
 			, Output(new TOutput<DataType>(TOutputParameters<DataType>({ FName("CollectionOut"), this })))
 		{}
 
-		virtual void Evaluate(const FContext& Context, FConnection* Out) const override
-		{
-			if (Output.Get() == Out)
-			{
-				if (const FEngineContext* EngineContext = Context.AsType<FEngineContext>(FName("UGeometryCollection")))
-				{
-					if (UGeometryCollection* CollectionAsset = Cast<UGeometryCollection>(EngineContext->Owner))
-					{
-						if (const TSharedPtr<FGeometryCollection, ESPMode::ThreadSafe> AssetCollection = CollectionAsset->GetGeometryCollection())
-						{
-							FManagedArrayCollection* NewCollection = AssetCollection->NewCopy<FManagedArrayCollection>();
-							Output->SetValue(FManagedArrayCollectionSharedPtr(NewCollection), Context);
-						}
-					}
-				}
-			}
-		}
+		virtual void Evaluate(const FContext& Context, FConnection* Out) const override;
 	};
 
 
@@ -76,22 +49,8 @@ namespace Dataflow
 			, Scale({ FName("Scale"), 10.f, this })
 		{}
 
-		virtual void Evaluate(const FContext& Context, FConnection* Out) const override
-		{
-			if (Output.Get() == Out)
-			{
-				DataType Collection = Input->GetValue(Context);
-				if (Active.GetValue())
-				{
-					TManagedArray<FVector3f>* Vertex = Collection->FindAttribute<FVector3f>("Vertex", "Vertices");
-					for (int i = 0; i < Vertex->Num(); i++)
-					{
-						(*Vertex)[i][1] *= Scale.GetValue();
-					}
-				}
-				Output->SetValue(Collection,Context);
-			}
-		}
+		virtual void Evaluate(const FContext& Context, FConnection* Out) const override;
+
 	};
 
 	class SetCollectionAssetNode : public FNode
@@ -107,22 +66,8 @@ namespace Dataflow
 			, Input(new TInput<DataType>(TInputParameters<DataType>({ FName("CollectionIn"), this })))
 		{}
 
-		virtual void Evaluate(const FContext& Context, FConnection* Out) const override
-		{
-			if (Out == nullptr)
-			{
-				DataType Collection = Input->GetValue(Context);
-				if (const FEngineContext* EngineContext = Context.AsType<FEngineContext>(FName("UGeometryCollection")))
-				{
-					if (UGeometryCollection* CollectionAsset = Cast<UGeometryCollection>(EngineContext->Owner))
-					{
-						TSharedPtr<FGeometryCollection, ESPMode::ThreadSafe> NewCollection(Collection->NewCopy<FGeometryCollection>());
-						CollectionAsset->SetGeometryCollection(NewCollection);
-						CollectionAsset->InvalidateCollection();
-					}
-				}
-			}
-		}
+		virtual void Evaluate(const FContext& Context, FConnection* Out) const override;
+
 	};
 
 
@@ -141,66 +86,8 @@ namespace Dataflow
 		{}
 
 
-		virtual void Evaluate(const FContext& Context, FConnection* Out) const override
-		{
-			ManagedArrayOut->SetValue(TSharedPtr<FManagedArrayCollection>(nullptr), Context);
-
-			if (const FEngineContext* EngineContext = (const FEngineContext*)(&Context))
-			{
-				if (UGeometryCollection* GeometryCollectionObject = Cast<UGeometryCollection>(EngineContext->Owner))
-				{
-					GeometryCollectionObject->Reset();
-
-					const UObject* Owner = EngineContext->Owner;
-					FName AName("GeometrySource");
-					if (Owner && Owner->GetClass())
-					{
-						if (const ::FProperty* UEProperty = Owner->GetClass()->FindPropertyByName(AName))
-						{
-							if (const FArrayProperty* ArrayProperty = CastField<FArrayProperty>(UEProperty))
-							{
-								FScriptArrayHelper_InContainer ArrayHelper(ArrayProperty, Owner);
-								const int32 ArraySize = ArrayHelper.Num();
-								for (int32 Index = 0; Index < ArraySize; ++Index)
-								{
-									if (FGeometryCollectionSource* SourceObject = (FGeometryCollectionSource*)(ArrayHelper.GetRawPtr(Index)))
-									{
-										if (UObject* ResolvedObject = SourceObject->SourceGeometryObject.ResolveObject())
-										{
-											if (UStaticMesh* StaticMesh = Cast<UStaticMesh>(ResolvedObject))
-											{
-												TArray<UMaterialInterface*> Materials;
-												Materials.Reserve(StaticMesh->GetStaticMaterials().Num());
-
-												for (int32 Index2 = 0; Index2 < StaticMesh->GetStaticMaterials().Num(); ++Index2)
-												{
-													UMaterialInterface* CurrMaterial = StaticMesh->GetMaterial(Index2);
-													Materials.Add(CurrMaterial);
-												}
-
-												// Geometry collections usually carry the selection material, which we'll delete before appending
-												UMaterialInterface* BoneSelectedMaterial = LoadObject<UMaterialInterface>(nullptr, UGeometryCollection::GetSelectedMaterialPath(), nullptr, LOAD_None, nullptr);
-												GeometryCollectionObject->Materials.Remove(BoneSelectedMaterial);
-												Materials.Remove(BoneSelectedMaterial);
-
-												FGeometryCollectionEngineConversion::AppendStaticMesh(StaticMesh, Materials, FTransform(), GeometryCollectionObject);
-											}
-										}
-									}
-								}
-							}
-						}
-					}
-					GeometryCollectionObject->InvalidateCollection();
-
-					if (const TSharedPtr<FGeometryCollection, ESPMode::ThreadSafe> AssetCollection = GeometryCollectionObject->GetGeometryCollection())
-					{
-						FManagedArrayCollection* NewCollection = AssetCollection->NewCopy<FManagedArrayCollection>();
-						ManagedArrayOut->SetValue(FManagedArrayCollectionSharedPtr(NewCollection), Context);
-					}
-				}
-			}
-		}
+		virtual void Evaluate(const FContext& Context, FConnection* Out) const override;
+
 	};
 
 	void GeometryCollectionEngineAssetNodes();
diff --git a/Engine/Source/Runtime/Experimental/Dataflow/Engine/Public/Dataflow/DataflowEngine.h b/Engine/Source/Runtime/Experimental/Dataflow/Engine/Public/Dataflow/DataflowEngine.h
--- a/Engine/Source/Runtime/Experimental/Dataflow/Engine/Public/Dataflow/DataflowEngine.h
+++ b/Engine/Source/Runtime/Experimental/Dataflow/Engine/Public/Dataflow/DataflowEngine.h
@@ -6,3 +6,5 @@
 #include "Dataflow/DataflowObjectInterface.h"
 #include "Dataflow/DataflowEdNode.h"
 #include "Dataflow/DataflowObject.h"
+#include "Dataflow/DataflowEngineUtil.h"
+
