GeometryCollection : Node
-Update collection after init.

diff --git a/Engine/Plugins/Experimental/GeometryCollectionPlugin/Source/GeometryCollectionNodes/Private/Dataflow/GeometryCollectionNodes.cpp b/Engine/Plugins/Experimental/GeometryCollectionPlugin/Source/GeometryCollectionNodes/Private/Dataflow/GeometryCollectionNodes.cpp
--- a/Engine/Plugins/Experimental/GeometryCollectionPlugin/Source/GeometryCollectionNodes/Private/Dataflow/GeometryCollectionNodes.cpp
+++ b/Engine/Plugins/Experimental/GeometryCollectionPlugin/Source/GeometryCollectionNodes/Private/Dataflow/GeometryCollectionNodes.cpp
@@ -117,6 +117,7 @@ namespace Dataflow
 											Materials.Remove(BoneSelectedMaterial);
 
 											FGeometryCollectionEngineConversion::AppendStaticMesh(StaticMesh, Materials, FTransform(), GeometryCollectionObject);
+
 										}
 									}
 								}
@@ -124,6 +125,8 @@ namespace Dataflow
 						}
 					}
 				}
+				GeometryCollectionObject->UpdateConvexGeometry();
+				GeometryCollectionObject->InitializeMaterials();
 				GeometryCollectionObject->InvalidateCollection();
 
 				if (const TSharedPtr<FGeometryCollection, ESPMode::ThreadSafe> AssetCollection = GeometryCollectionObject->GetGeometryCollection())
