Dataflow :
- bug fix for getting invalid context

diff --git a/Engine/Plugins/Experimental/GeometryCollectionPlugin/Source/GeometryCollectionNodes/Private/Dataflow/GeometryCollectionNodes.cpp b/Engine/Plugins/Experimental/GeometryCollectionPlugin/Source/GeometryCollectionNodes/Private/Dataflow/GeometryCollectionNodes.cpp
--- a/Engine/Plugins/Experimental/GeometryCollectionPlugin/Source/GeometryCollectionNodes/Private/Dataflow/GeometryCollectionNodes.cpp
+++ b/Engine/Plugins/Experimental/GeometryCollectionPlugin/Source/GeometryCollectionNodes/Private/Dataflow/GeometryCollectionNodes.cpp
@@ -78,7 +78,7 @@ namespace Dataflow
 	{
 		ManagedArrayOut->SetValue(TSharedPtr<FManagedArrayCollection>(nullptr), Context);
 
-		if (const FEngineContext* EngineContext = (const FEngineContext*)(&Context))
+		if (const FEngineContext* EngineContext = Context.AsType<FEngineContext>(FName("UGeometryCollection")))
 		{
 			if (UGeometryCollection* GeometryCollectionObject = Cast<UGeometryCollection>(EngineContext->Owner))
 			{
