Dataflow : Support for FProperty on the FNode

diff --git a/Engine/Plugins/Experimental/Dataflow/Source/DataflowEditor/Private/Dataflow/DataflowAssetActions.cpp b/Engine/Plugins/Experimental/Dataflow/Source/DataflowEditor/Private/Dataflow/DataflowAssetActions.cpp
--- a/Engine/Plugins/Experimental/Dataflow/Source/DataflowEditor/Private/Dataflow/DataflowAssetActions.cpp
+++ b/Engine/Plugins/Experimental/Dataflow/Source/DataflowEditor/Private/Dataflow/DataflowAssetActions.cpp
@@ -30,6 +30,9 @@ void FDataflowAssetActions::GetActions(const TArray<UObject*>& InObjects,
 void FDataflowAssetActions::OpenAssetEditor(
 	const TArray<UObject*>& InObjects, TSharedPtr<IToolkitHost> EditWithinLevelEditor)
 {
+	FAssetTypeActions_Base::OpenAssetEditor(InObjects, EditWithinLevelEditor);
+
+	/*
 	EToolkitMode::Type Mode = EditWithinLevelEditor.IsValid() ? EToolkitMode::WorldCentric : EToolkitMode::Standalone;
 	for (auto ObjIt = InObjects.CreateConstIterator(); ObjIt; ++ObjIt)
 	{
@@ -39,6 +42,7 @@ void FDataflowAssetActions::OpenAssetEditor(
 			DataflowEditorPlugin->CreateDataflowAssetEditor(Mode, EditWithinLevelEditor, Object);
 		}
 	}
+	*/
 }
 
 
diff --git a/Engine/Plugins/Experimental/Dataflow/Source/DataflowEditor/Private/Dataflow/DataflowAssetFactory.cpp b/Engine/Plugins/Experimental/Dataflow/Source/DataflowEditor/Private/Dataflow/DataflowAssetFactory.cpp
--- a/Engine/Plugins/Experimental/Dataflow/Source/DataflowEditor/Private/Dataflow/DataflowAssetFactory.cpp
+++ b/Engine/Plugins/Experimental/Dataflow/Source/DataflowEditor/Private/Dataflow/DataflowAssetFactory.cpp
@@ -5,8 +5,6 @@
 #include "Dataflow/DataflowObject.h"
 #include "HAL/IConsoleManager.h"
 
-bool bDataflowEnableCreation = false;
-FAutoConsoleVariableRef CVarDataflowEnableCreation(TEXT("c.DataflowEnableCreation"), bDataflowEnableCreation, TEXT("Enable creation for the dataflow system (Curently Dev-Only)"));
 
 
 UDataflowAssetFactory::UDataflowAssetFactory()
@@ -16,7 +14,7 @@ UDataflowAssetFactory::UDataflowAssetFactory()
 
 bool UDataflowAssetFactory::CanCreateNew() const
 {
-	return bDataflowEnableCreation;
+	return true;
 }
 
 bool UDataflowAssetFactory::FactoryCanImport(const FString& Filename)
@@ -32,7 +30,7 @@ UObject* UDataflowAssetFactory::FactoryCreateNew(UClass* InClass, UObject* InPar
 
 bool UDataflowAssetFactory::ShouldShowInNewMenu() const
 {
-	return bDataflowEnableCreation;
+	return true;
 }
 
 bool UDataflowAssetFactory::ConfigureProperties()
diff --git a/Engine/Plugins/Experimental/Dataflow/Source/DataflowEditor/Private/Dataflow/DataflowEPropertyCustomizations.cpp b/Engine/Plugins/Experimental/Dataflow/Source/DataflowEditor/Private/Dataflow/DataflowEPropertyCustomizations.cpp
--- a/Engine/Plugins/Experimental/Dataflow/Source/DataflowEditor/Private/Dataflow/DataflowEPropertyCustomizations.cpp
+++ b/Engine/Plugins/Experimental/Dataflow/Source/DataflowEditor/Private/Dataflow/DataflowEPropertyCustomizations.cpp
@@ -2,11 +2,9 @@
 #include "Dataflow/DataflowEPropertyCustomizations.h"
 
 #include "Dataflow/DataflowNode.h"
-#include "Dataflow/DataflowProperty.h"
 #include "DetailLayoutBuilder.h"
 #include "DetailCategoryBuilder.h"
 #include "DetailWidgetRow.h"
-#include "Dataflow/DataflowEPropertyFactories.h"
 
 #define LOCTEXT_NAMESPACE "DataflowEProprty_Customizations"
 
@@ -32,10 +30,11 @@ void FDataflowSEditorCustomization::CustomizeDetails(IDetailLayoutBuilder& Detai
 	{
 		if (UDataflowSEditorObject* DataflowSObj = Cast< UDataflowSEditorObject>(Obj))
 		{
-			if (TSharedPtr<Dataflow::FNode> Node = DataflowSObj->Node)
+			if (TSharedPtr<FDataflowNode> Node = DataflowSObj->Node)
 			{
 				if (UDataflow* Graph = DataflowSObj->Graph)
 				{
+					/*
 					for (Dataflow::FProperty* Property : Node->GetProperties())
 					{
 						IDetailPropertyRow& PropertyRow = Category.AddProperty(FName("CustomProperty"));
@@ -66,6 +65,7 @@ void FDataflowSEditorCustomization::CustomizeDetails(IDetailLayoutBuilder& Detai
 						}
 						UE_LOG(LogCDataflowEProprtyCustomizations, Verbose, TEXT("FFloatObjectDetails::CustomizeDetails"));
 					}
+					*/
 				}
 			}
 		}
diff --git a/Engine/Plugins/Experimental/Dataflow/Source/DataflowEditor/Private/Dataflow/DataflowEPropertyFactories.cpp b/Engine/Plugins/Experimental/Dataflow/Source/DataflowEditor/Private/Dataflow/DataflowEPropertyFactories.cpp
deleted file mode 100644
diff --git a/Engine/Plugins/Experimental/Dataflow/Source/DataflowEditor/Private/Dataflow/DataflowEditorActions.cpp b/Engine/Plugins/Experimental/Dataflow/Source/DataflowEditor/Private/Dataflow/DataflowEditorActions.cpp
--- a/Engine/Plugins/Experimental/Dataflow/Source/DataflowEditor/Private/Dataflow/DataflowEditorActions.cpp
+++ b/Engine/Plugins/Experimental/Dataflow/Source/DataflowEditor/Private/Dataflow/DataflowEditorActions.cpp
@@ -5,6 +5,7 @@
 #include "Dataflow/DataflowEPropertyCustomizations.h"
 #include "Dataflow/DataflowNodeFactory.h"
 #include "Dataflow/DataflowObject.h"
+#include "IStructureDetailsView.h"
 
 
 #define LOCTEXT_NAMESPACE "DataflowEditorCommands"
@@ -58,7 +59,7 @@ void FDataflowEditorCommands::EvaluateNodes(const FGraphPanelSelectionSet& Selec
 		{
 			if (const TSharedPtr<Dataflow::FGraph> DataflowGraph = EdNode->GetDataflowGraph())
 			{
-				if (const TSharedPtr<Dataflow::FNode> DataflowNode = DataflowGraph->FindBaseNode(EdNode->GetDataflowNodeGuid()))
+				if (const TSharedPtr<FDataflowNode> DataflowNode = DataflowGraph->FindBaseNode(EdNode->GetDataflowNodeGuid()))
 				{
 					if (DataflowNode->GetOutputs().Num())
 					{
@@ -85,7 +86,7 @@ void FDataflowEditorCommands::DeleteNodes(UDataflow* Graph, const FGraphPanelSel
 		{
 			if (const TSharedPtr<Dataflow::FGraph> DataflowGraph = EdNode->GetDataflowGraph())
 			{
-				if (TSharedPtr<Dataflow::FNode> DataflowNode = DataflowGraph->FindBaseNode(EdNode->GetDataflowNodeGuid()))
+				if (TSharedPtr<FDataflowNode> DataflowNode = DataflowGraph->FindBaseNode(EdNode->GetDataflowNodeGuid()))
 				{
 					Graph->RemoveNode(EdNode);
 					DataflowGraph->RemoveNode(DataflowNode);
@@ -95,14 +96,15 @@ void FDataflowEditorCommands::DeleteNodes(UDataflow* Graph, const FGraphPanelSel
 	}
 }
 
-void FDataflowEditorCommands::OnSelectedNodesChanged(TSharedPtr<IDetailsView> PropertiesEditor, UObject* Asset, UDataflow* Graph, const TSet<UObject*>& NewSelection)
+void FDataflowEditorCommands::OnSelectedNodesChanged(TSharedPtr<IStructureDetailsView> PropertiesEditor, UObject* Asset, UDataflow* Graph, const TSet<UObject*>& NewSelection)
 {
+	PropertiesEditor->SetStructureData(nullptr);
+
 	if (Graph && PropertiesEditor)
 	{
 		if (const TSharedPtr<Dataflow::FGraph> DataflowGraph = Graph->GetDataflow())
 		{
-			PropertiesEditor->SetObject(Asset);
-			FGraphPanelSelectionSet SelectedNodes = NewSelection;//GetSelectedNodes();
+			FGraphPanelSelectionSet SelectedNodes = NewSelection;
 			if (SelectedNodes.Num())
 			{
 				TArray<UObject*> Objects;
@@ -110,17 +112,14 @@ void FDataflowEditorCommands::OnSelectedNodesChanged(TSharedPtr<IDetailsView> Pr
 				{
 					if (UDataflowEdNode* EdNode = Cast<UDataflowEdNode>(SelectedObject))
 					{
-						if (TSharedPtr<Dataflow::FNode> DataflowNode = DataflowGraph->FindBaseNode(EdNode->GetDataflowNodeGuid()))
+						if (TSharedPtr<FDataflowNode> DataflowNode = DataflowGraph->FindBaseNode(EdNode->GetDataflowNodeGuid()))
 						{
-							const FName NodeName = MakeUniqueObjectName(Graph, UDataflow::StaticClass(), DataflowNode->GetName());
-							UDataflowSEditorObject* Object = NewObject<UDataflowSEditorObject>(Asset, NodeName);
-							Object->Node = DataflowNode;
-							Object->Graph = Graph;
-							Objects.Add(Object);
+							TSharedPtr<FStructOnScope> Struct(DataflowNode->NewScructOnScope());
+							PropertiesEditor->SetStructureData(Struct);
 						}
 					}
 				}
-				PropertiesEditor->SetObjects(Objects);
+				
 			}
 		}
 	}
diff --git a/Engine/Plugins/Experimental/Dataflow/Source/DataflowEditor/Private/Dataflow/DataflowEditorToolkit.cpp b/Engine/Plugins/Experimental/Dataflow/Source/DataflowEditor/Private/Dataflow/DataflowEditorToolkit.cpp
--- a/Engine/Plugins/Experimental/Dataflow/Source/DataflowEditor/Private/Dataflow/DataflowEditorToolkit.cpp
+++ b/Engine/Plugins/Experimental/Dataflow/Source/DataflowEditor/Private/Dataflow/DataflowEditorToolkit.cpp
@@ -16,6 +16,7 @@
 #include "Styling/SlateStyleRegistry.h"
 #include "Framework/Commands/GenericCommands.h"
 #include "Widgets/Docking/SDockTab.h"
+#include "IStructureDetailsView.h"
 
 #define LOCTEXT_NAMESPACE "DataflowEditorToolkit"
 
@@ -82,20 +83,35 @@ TSharedRef<SGraphEditor> FDataflowEditorToolkit::CreateGraphEditorWidget(UDatafl
 		.DetailsView(GetPropertiesEditor());
 }
 
-TSharedPtr<IDetailsView> FDataflowEditorToolkit::CreatePropertiesEditorWidget(UObject* ObjectToEdit)
+TSharedPtr<IStructureDetailsView> FDataflowEditorToolkit::CreatePropertiesEditorWidget(UObject* ObjectToEdit)
 {
 	ensure(ObjectToEdit);
 	FPropertyEditorModule& PropertyEditorModule = FModuleManager::LoadModuleChecked<FPropertyEditorModule>(TEXT("PropertyEditor"));
 
 	FDetailsViewArgs DetailsViewArgs;
-	DetailsViewArgs.bAllowSearch = true;
-	DetailsViewArgs.bLockable = false;
-	DetailsViewArgs.bUpdatesFromSelection = false;
-	DetailsViewArgs.NameAreaSettings = FDetailsViewArgs::HideNameArea;
-	DetailsViewArgs.NotifyHook = this;
-
-	TSharedPtr<IDetailsView> DetailsView = PropertyEditorModule.CreateDetailView(DetailsViewArgs);
-	DetailsView->SetObject(ObjectToEdit);
+	{
+		DetailsViewArgs.bAllowSearch = false;
+		DetailsViewArgs.bHideSelectionTip = true;
+		DetailsViewArgs.bLockable = false;
+		DetailsViewArgs.bSearchInitialKeyFocus = true;
+		DetailsViewArgs.bUpdatesFromSelection = false;
+		DetailsViewArgs.NotifyHook = nullptr;
+		DetailsViewArgs.bShowOptions = true;
+		DetailsViewArgs.bShowModifiedPropertiesOption = false;
+		DetailsViewArgs.bShowScrollBar = false;
+	}
+
+	FStructureDetailsViewArgs StructureViewArgs;
+	{
+		StructureViewArgs.bShowObjects = true;
+		StructureViewArgs.bShowAssets = true;
+		StructureViewArgs.bShowClasses = true;
+		StructureViewArgs.bShowInterfaces = true;
+	}
+
+	TSharedPtr<IStructureDetailsView> DetailsView = PropertyEditorModule.CreateStructureDetailView(DetailsViewArgs, StructureViewArgs, nullptr);
+	DetailsView->GetDetailsView()->SetObject(ObjectToEdit);
+
 	return DetailsView;
 
 }
@@ -121,7 +137,7 @@ TSharedRef<SDockTab> FDataflowEditorToolkit::SpawnTab_Properties(const FSpawnTab
 	return SNew(SDockTab)
 		.Label(LOCTEXT("DataflowEditor_Properties_TabTitle", "Details"))
 		[
-			PropertiesEditor.ToSharedRef()
+			PropertiesEditor->GetWidget()->AsShared()
 		];
 }
 
diff --git a/Engine/Plugins/Experimental/Dataflow/Source/DataflowEditor/Private/Dataflow/DataflowEditorToolkit.h b/Engine/Plugins/Experimental/Dataflow/Source/DataflowEditor/Private/Dataflow/DataflowEditorToolkit.h
--- a/Engine/Plugins/Experimental/Dataflow/Source/DataflowEditor/Private/Dataflow/DataflowEditorToolkit.h
+++ b/Engine/Plugins/Experimental/Dataflow/Source/DataflowEditor/Private/Dataflow/DataflowEditorToolkit.h
@@ -36,8 +36,8 @@ public:
 	UDataflow* GetDataflow() { return Dataflow; }
 	const UDataflow* GetDataflow() const { return Dataflow; }
 
-	TSharedPtr<IDetailsView> GetPropertiesEditor() {return PropertiesEditor;}
-	const TSharedPtr<IDetailsView> GetPropertiesEditor() const { return PropertiesEditor; }
+	TSharedPtr<IStructureDetailsView> GetPropertiesEditor() {return PropertiesEditor;}
+	const TSharedPtr<IStructureDetailsView> GetPropertiesEditor() const { return PropertiesEditor; }
 
 	TSharedPtr<SGraphEditor> GetGraphEditor() { return GraphEditor; }
 	const TSharedPtr<SGraphEditor> GetGraphEditor() const { return GraphEditor; }
@@ -52,6 +52,6 @@ private:
 	TSharedRef<SGraphEditor> CreateGraphEditorWidget(UDataflow* ObjectToEdit);
 
 	static const FName PropertiesTabId;
-	TSharedPtr<IDetailsView> PropertiesEditor;
-	TSharedPtr<IDetailsView> CreatePropertiesEditorWidget(UObject* ObjectToEdit);
+	TSharedPtr<IStructureDetailsView> PropertiesEditor;
+	TSharedPtr<IStructureDetailsView> CreatePropertiesEditorWidget(UObject* ObjectToEdit);
 };
diff --git a/Engine/Plugins/Experimental/Dataflow/Source/DataflowEditor/Private/Dataflow/DataflowGraphEditor.cpp b/Engine/Plugins/Experimental/Dataflow/Source/DataflowEditor/Private/Dataflow/DataflowGraphEditor.cpp
--- a/Engine/Plugins/Experimental/Dataflow/Source/DataflowEditor/Private/Dataflow/DataflowGraphEditor.cpp
+++ b/Engine/Plugins/Experimental/Dataflow/Source/DataflowEditor/Private/Dataflow/DataflowGraphEditor.cpp
@@ -63,7 +63,7 @@ void SDataflowGraphEditor::EvaluateNode()
 	}
 	else
 	{
-		FDataflowEditorCommands::FGraphEvaluationCallback LocalEvaluateCallback = [](Dataflow::FNode* Node, Dataflow::FConnection* Out)
+		FDataflowEditorCommands::FGraphEvaluationCallback LocalEvaluateCallback = [](FDataflowNode* Node, Dataflow::FConnection* Out)
 		{
 			float EvalTime = FGameTime::GetTimeSinceAppStart().GetRealTimeSeconds();
 			return Node->Evaluate(Dataflow::FContext(EvalTime), Out);
diff --git a/Engine/Plugins/Experimental/Dataflow/Source/DataflowEditor/Private/Dataflow/DataflowSNode.cpp b/Engine/Plugins/Experimental/Dataflow/Source/DataflowEditor/Private/Dataflow/DataflowSNode.cpp
--- a/Engine/Plugins/Experimental/Dataflow/Source/DataflowEditor/Private/Dataflow/DataflowSNode.cpp
+++ b/Engine/Plugins/Experimental/Dataflow/Source/DataflowEditor/Private/Dataflow/DataflowSNode.cpp
@@ -62,7 +62,7 @@ UEdGraphNode* FAssetSchemaAction_Dataflow_CreateNode_DataflowEdNode::PerformActi
 
 			if (Dataflow::FNodeFactory* Factory = Dataflow::FNodeFactory::GetInstance())
 			{
-				if (TSharedPtr<Dataflow::FNode> DataflowNode = Factory->NewNodeFromRegisteredType(*Dataflow->GetDataflow(), { FGuid::NewGuid(),NodeTypeName,NodeName}))
+				if (TSharedPtr<FDataflowNode> DataflowNode = Factory->NewNodeFromRegisteredType(*Dataflow->GetDataflow(), { FGuid::NewGuid(),NodeTypeName,NodeName}))
 				{
 					EdNode->SetDataflowGraph(Dataflow->GetDataflow());
 					EdNode->SetDataflowNodeGuid(DataflowNode->GetGuid());
diff --git a/Engine/Plugins/Experimental/Dataflow/Source/DataflowEditor/Public/Dataflow/DataflowEPropertyCustomizations.h b/Engine/Plugins/Experimental/Dataflow/Source/DataflowEditor/Public/Dataflow/DataflowEPropertyCustomizations.h
--- a/Engine/Plugins/Experimental/Dataflow/Source/DataflowEditor/Public/Dataflow/DataflowEPropertyCustomizations.h
+++ b/Engine/Plugins/Experimental/Dataflow/Source/DataflowEditor/Public/Dataflow/DataflowEPropertyCustomizations.h
@@ -9,9 +9,8 @@
 
 #include "DataflowEPropertyCustomizations.generated.h"
 
-namespace Dataflow {
-	class FNode;
-}
+struct FDataflowNode;
+
 class UDataflow;
 
 UCLASS()
@@ -19,7 +18,7 @@ class DATAFLOWEDITOR_API UDataflowSEditorObject : public UObject
 {
 	GENERATED_BODY()
 public:
-	TSharedPtr<Dataflow::FNode> Node = nullptr;
+	TSharedPtr<FDataflowNode> Node = nullptr;
 	UDataflow* Graph = nullptr;
 
 	virtual void PostEditChangeProperty(struct FPropertyChangedEvent& PropertyChangedEvent) override;
diff --git a/Engine/Plugins/Experimental/Dataflow/Source/DataflowEditor/Public/Dataflow/DataflowEPropertyFactories.h b/Engine/Plugins/Experimental/Dataflow/Source/DataflowEditor/Public/Dataflow/DataflowEPropertyFactories.h
deleted file mode 100644
diff --git a/Engine/Plugins/Experimental/Dataflow/Source/DataflowEditor/Public/Dataflow/DataflowEditorActions.h b/Engine/Plugins/Experimental/Dataflow/Source/DataflowEditor/Public/Dataflow/DataflowEditorActions.h
--- a/Engine/Plugins/Experimental/Dataflow/Source/DataflowEditor/Public/Dataflow/DataflowEditorActions.h
+++ b/Engine/Plugins/Experimental/Dataflow/Source/DataflowEditor/Public/Dataflow/DataflowEditorActions.h
@@ -9,12 +9,13 @@
 namespace Dataflow 
 {
 	class FConnection;
-	class FNode;
 }
 
-class IDetailsView;
+class IStructureDetailsView;
 class UDataflow;
 class UDataflowEdNode;
+struct FDataflowNode;
+
 typedef TSet<class UObject*> FGraphPanelSelectionSet;
 
 /*
@@ -43,7 +44,7 @@ public:
 class DATAFLOWEDITOR_API FDataflowEditorCommands
 {
 public:
-	typedef TFunction<void(Dataflow::FNode*, Dataflow::FConnection*)> FGraphEvaluationCallback;
+	typedef TFunction<void(FDataflowNode*, Dataflow::FConnection*)> FGraphEvaluationCallback;
 	typedef TFunction<void(const FGeometry& MyGeometry, const FDragDropEvent& DragDropEvent)> FOnDragDropEventCallback;
 
 	static void Register();
@@ -64,6 +65,6 @@ public:
 	/*
 	*  OnSelectedNodesChanged
 	*/
-	static void OnSelectedNodesChanged(TSharedPtr<IDetailsView> PropertiesEditor, UObject* Asset, UDataflow* Graph, const TSet<UObject*>& NewSelection);
+	static void OnSelectedNodesChanged(TSharedPtr<IStructureDetailsView> PropertiesEditor, UObject* Asset, UDataflow* Graph, const TSet<UObject*>& NewSelection);
 
 };
diff --git a/Engine/Plugins/Experimental/Dataflow/Source/DataflowEditor/Public/Dataflow/DataflowGraphEditor.h b/Engine/Plugins/Experimental/Dataflow/Source/DataflowEditor/Public/Dataflow/DataflowGraphEditor.h
--- a/Engine/Plugins/Experimental/Dataflow/Source/DataflowEditor/Public/Dataflow/DataflowGraphEditor.h
+++ b/Engine/Plugins/Experimental/Dataflow/Source/DataflowEditor/Public/Dataflow/DataflowGraphEditor.h
@@ -44,7 +44,7 @@ public:
 	SLATE_ATTRIBUTE(FGraphAppearanceInfo, Appearance)
 	SLATE_ARGUMENT(UEdGraph*, GraphToEdit)
 	SLATE_ARGUMENT(FGraphEditorEvents, GraphEvents)
-	SLATE_ARGUMENT(TSharedPtr<IDetailsView>, DetailsView)
+	SLATE_ARGUMENT(TSharedPtr<IStructureDetailsView>, DetailsView)
 	SLATE_ARGUMENT(FDataflowEditorCommands::FGraphEvaluationCallback, EvaluateGraph)
 	SLATE_ARGUMENT(FDataflowEditorCommands::FOnDragDropEventCallback, OnDragDropEvent)
 	SLATE_END_ARGS()
@@ -88,5 +88,5 @@ private:
 	static TSharedPtr<FDataflowSNodeFactory> NodeFactory;
 
 	/** The details view that responds to this widget. */
-	TSharedPtr<IDetailsView> DetailsView;
+	TSharedPtr<IStructureDetailsView> DetailsView;
 };
diff --git a/Engine/Plugins/Experimental/Dataflow/Source/DataflowNodes/Private/Dataflow/DataflowSkeletalMeshNodes.cpp b/Engine/Plugins/Experimental/Dataflow/Source/DataflowNodes/Private/Dataflow/DataflowSkeletalMeshNodes.cpp
--- a/Engine/Plugins/Experimental/Dataflow/Source/DataflowNodes/Private/Dataflow/DataflowSkeletalMeshNodes.cpp
+++ b/Engine/Plugins/Experimental/Dataflow/Source/DataflowNodes/Private/Dataflow/DataflowSkeletalMeshNodes.cpp
@@ -13,32 +13,32 @@ namespace Dataflow
 {
 	void RegisterSkeletalMeshNodes()
 	{
-		DATAFLOW_NODE_REGISTER_CREATION_FACTORY(GetSkeletalMesh);
-		DATAFLOW_NODE_REGISTER_CREATION_FACTORY(SkeletalMeshBone);
+		DATAFLOW_NODE_REGISTER_CREATION_FACTORY(FGetSkeletalMeshDataflowNode);
+		DATAFLOW_NODE_REGISTER_CREATION_FACTORY(FSkeletalMeshBoneDataflowNode);
 	}
-
-	void GetSkeletalMesh::Evaluate(const FContext& Context, FConnection* Out) const
+}
+void FGetSkeletalMeshDataflowNode::Evaluate(const Dataflow::FContext& Context, Dataflow::FConnection* Out) const
+{
+	SkeletalMeshOut->SetValue(nullptr, Context);
+	if (const Dataflow::FEngineContext* EngineContext = (const Dataflow::FEngineContext*)(&Context))
 	{
-		SkeletalMeshOut->SetValue(nullptr, Context);
-		if (const FEngineContext* EngineContext = (const FEngineContext*)(&Context))
+		if (const USkeletalMesh* SkeletalMesh = Dataflow::Reflection::FindObjectPtrProperty<USkeletalMesh>(
+			EngineContext->Owner, SkeletalMeshAttributeName))
 		{
-			if (const USkeletalMesh* SkeletalMesh = Reflection::FindObjectPtrProperty<USkeletalMesh>(
-				EngineContext->Owner, SkeletalMeshAttributeName.GetValue()))
-			{
-				SkeletalMeshOut->SetValue(SkeletalMesh, Context);
-			}
+			SkeletalMeshOut->SetValue(SkeletalMesh, Context);
 		}
 	}
+}
 
-	void SkeletalMeshBone::Evaluate(const FContext& Context, FConnection* Out) const
-	{
-		BoneIndexOut->SetValue(INDEX_NONE, Context);
+void FSkeletalMeshBoneDataflowNode::Evaluate(const Dataflow::FContext& Context, Dataflow::FConnection* Out) const
+{
+	BoneIndexOut->SetValue(INDEX_NONE, Context);
 
-		if (const USkeletalMesh* SkeletalMesh = SkeletalMeshIn->GetValue(Context))
-		{
-			int32 Index = SkeletalMesh->GetRefSkeleton().FindBoneIndex(BoneName.GetValue());
-			BoneIndexOut->SetValue(Index, Context);
-		}
+	if (const USkeletalMesh* SkeletalMesh = SkeletalMeshIn->GetValue(Context))
+	{
+		int32 Index = SkeletalMesh->GetRefSkeleton().FindBoneIndex(BoneName);
+		BoneIndexOut->SetValue(Index, Context);
 	}
 }
 
+
diff --git a/Engine/Plugins/Experimental/Dataflow/Source/DataflowNodes/Public/Dataflow/DataflowSkeletalMeshNodes.h b/Engine/Plugins/Experimental/Dataflow/Source/DataflowNodes/Public/Dataflow/DataflowSkeletalMeshNodes.h
--- a/Engine/Plugins/Experimental/Dataflow/Source/DataflowNodes/Public/Dataflow/DataflowSkeletalMeshNodes.h
+++ b/Engine/Plugins/Experimental/Dataflow/Source/DataflowNodes/Public/Dataflow/DataflowSkeletalMeshNodes.h
@@ -8,55 +8,62 @@
 #include "Dataflow/DataflowCore.h"
 #include "Dataflow/DataflowEngine.h"
 
+
+#include "DataflowSkeletalMeshNodes.generated.h"
+
+
 DEFINE_LOG_CATEGORY_STATIC(LogDataflowSkeletalMeshNodes, Log, All);
 
 class USkeletalMesh;
 
-namespace Dataflow
+USTRUCT()
+struct DATAFLOWNODES_API FGetSkeletalMeshDataflowNode: public FDataflowNode
 {
+	GENERATED_USTRUCT_BODY()
+	DATAFLOW_NODE_DEFINE_INTERNAL(FGetSkeletalMeshDataflowNode, "SkeletalMesh", "Dataflow", "Skeletal Mesh")
+
+public:
 
-	class DATAFLOWNODES_API GetSkeletalMesh : public FNode
-	{
-		DATAFLOW_NODE_DEFINE_INTERNAL(GetSkeletalMesh)
 
-	public:
-		TProperty<FName> SkeletalMeshAttributeName;
+	UPROPERTY(EditAnywhere, Category = "Dataflow")
+	FName SkeletalMeshAttributeName;
 
-		TSharedPtr< class TOutput<USkeletalMeshPtr> > SkeletalMeshOut;
+	TSharedPtr< class Dataflow::TOutput<Dataflow::USkeletalMeshPtr> > SkeletalMeshOut;
 
-		GetSkeletalMesh(const FNodeParameters& InParam, FGuid InGuid = FGuid::NewGuid())
-			: FNode(InParam, InGuid)
-			, SkeletalMeshAttributeName({ FName("SkeletalMeshAttributeName"), FName("SkeletalMesh"), this })
-			, SkeletalMeshOut(new TOutput<USkeletalMeshPtr>(TOutputParameters<USkeletalMeshPtr>({ FName("SkeletalMeshOut"), this })))
-		{}
+	FGetSkeletalMeshDataflowNode(const Dataflow::FNodeParameters& InParam, FGuid InGuid = FGuid::NewGuid())
+		: FDataflowNode(InParam, InGuid)
+		, SkeletalMeshOut(new Dataflow::TOutput<Dataflow::USkeletalMeshPtr>(Dataflow::TOutputParameters<Dataflow::USkeletalMeshPtr>({ FName("SkeletalMeshOut"), this })))
+	{}
 
 
-		virtual void Evaluate(const FContext& Context, FConnection* Out) const override;
-	};
-	
-	class DATAFLOWNODES_API SkeletalMeshBone : public FNode
-	{
-		DATAFLOW_NODE_DEFINE_INTERNAL(SkeletalMeshBone)
+	virtual void Evaluate(const Dataflow::FContext& Context, Dataflow::FConnection* Out) const override;
+};
 
-	public:
-		TProperty<FName> BoneName;
+USTRUCT()
+struct DATAFLOWNODES_API FSkeletalMeshBoneDataflowNode : public FDataflowNode
+{
+	GENERATED_USTRUCT_BODY()
+	DATAFLOW_NODE_DEFINE_INTERNAL(FSkeletalMeshBoneDataflowNode, "SkeletalMeshBone", "Dataflow", "Skeletal Mesh")
 
-		TSharedPtr< class TInput<USkeletalMeshPtr> > SkeletalMeshIn;
-		TSharedPtr< class TOutput<int> > BoneIndexOut;
+public:
+	UPROPERTY(EditAnywhere, Category = "Dataflow")
+	FName BoneName;
 
-		SkeletalMeshBone(const FNodeParameters& InParam, FGuid InGuid = FGuid::NewGuid())
-			: FNode(InParam, InGuid)
-			, BoneName({ FName("BoneName"), FName(""), this})
-			, SkeletalMeshIn(new TInput<USkeletalMeshPtr>(TInputParameters<USkeletalMeshPtr>({ FName("SkeletalMeshIn"), this })))
-			, BoneIndexOut(new TOutput<int>(TOutputParameters<int>({ FName("BoneIndexOut"), this })))
-		{}
+	TSharedPtr< class Dataflow::TInput<Dataflow::USkeletalMeshPtr> > SkeletalMeshIn;
+	TSharedPtr< class Dataflow::TOutput<int> > BoneIndexOut;
 
+	FSkeletalMeshBoneDataflowNode(const Dataflow::FNodeParameters& InParam, FGuid InGuid = FGuid::NewGuid())
+		: FDataflowNode(InParam, InGuid)
+		, SkeletalMeshIn(new Dataflow::TInput<Dataflow::USkeletalMeshPtr>(Dataflow::TInputParameters<Dataflow::USkeletalMeshPtr>({ FName("SkeletalMeshIn"), this })))
+		, BoneIndexOut(new Dataflow::TOutput<int>(Dataflow::TOutputParameters<int>({ FName("BoneIndexOut"), this })))
+	{}
 
-		virtual void Evaluate(const FContext& Context, FConnection* Out) const override;
-	};
-	
-	void RegisterSkeletalMeshNodes();
 
+	virtual void Evaluate(const Dataflow::FContext& Context, Dataflow::FConnection* Out) const override;
+};
 
+namespace Dataflow
+{
+	void RegisterSkeletalMeshNodes();
 }
 
diff --git a/Engine/Plugins/Experimental/GeometryCollectionPlugin/Source/GeometryCollectionEditor/Private/GeometryCollection/GeometryCollectionEditorToolkit.cpp b/Engine/Plugins/Experimental/GeometryCollectionPlugin/Source/GeometryCollectionEditor/Private/GeometryCollection/GeometryCollectionEditorToolkit.cpp
--- a/Engine/Plugins/Experimental/GeometryCollectionPlugin/Source/GeometryCollectionEditor/Private/GeometryCollection/GeometryCollectionEditorToolkit.cpp
+++ b/Engine/Plugins/Experimental/GeometryCollectionPlugin/Source/GeometryCollectionEditor/Private/GeometryCollection/GeometryCollectionEditorToolkit.cpp
@@ -20,6 +20,7 @@
 #include "PropertyEditorModule.h"
 #include "Styling/SlateStyleRegistry.h"
 #include "Framework/Commands/GenericCommands.h"
+#include "IStructureDetailsView.h"
 
 #define LOCTEXT_NAMESPACE "GeometryCollectionEditorToolkit"
 
@@ -27,7 +28,8 @@
 
 
 const FName FGeometryCollectionEditorToolkit::GraphCanvasTabId(TEXT("GeometryCollectionEditor_GraphCanvas"));
-const FName FGeometryCollectionEditorToolkit::PropertiesTabId(TEXT("GeometryCollectionEditor_Properties"));
+const FName FGeometryCollectionEditorToolkit::AssetDetailsTabId(TEXT("GeometryCollectionEditor_AssetDetails"));
+const FName FGeometryCollectionEditorToolkit::NodeDetailsTabId(TEXT("GeometryCollectionEditor_NodeDetails"));
 
 void FGeometryCollectionEditorToolkit::InitGeometryCollectionAssetEditor(const EToolkitMode::Type Mode, const TSharedPtr<IToolkitHost>& InitToolkitHost, UObject* ObjectToEdit)
 {
@@ -44,10 +46,11 @@ void FGeometryCollectionEditorToolkit::InitGeometryCollectionAssetEditor(const E
 		Dataflow = GeometryCollection->Dataflow;
 		GeometryCollection->Dataflow->Schema = UDataflowSchema::StaticClass();
 
-		PropertiesEditor = CreatePropertiesEditorWidget(ObjectToEdit);
-		GraphEditor = CreateGraphEditorWidget(Dataflow, PropertiesEditor);
+		NodeDetailsEditor = CreateNodeDetailsEditorWidget(ObjectToEdit);
+		AssetDetailsEditor = CreateAssetDetailsEditorWidget(GeometryCollection);
+		GraphEditor = CreateGraphEditorWidget(Dataflow, NodeDetailsEditor);
 
-		const TSharedRef<FTabManager::FLayout> StandaloneDefaultLayout = FTabManager::NewLayout("GeometryCollectionDataflowEditor_Layout")
+		const TSharedRef<FTabManager::FLayout> StandaloneDefaultLayout = FTabManager::NewLayout("GeometryCollectionDataflowEditor_Layout.V1")
 			->AddArea
 			(
 				FTabManager::NewPrimaryArea()->SetOrientation(Orient_Vertical)
@@ -72,7 +75,7 @@ void FGeometryCollectionEditorToolkit::InitGeometryCollectionAssetEditor(const E
 							(
 								FTabManager::NewStack()
 								->SetSizeCoefficient(0.7f)
-								->AddTab(PropertiesTabId, ETabState::OpenedTab)
+								->AddTab(AssetDetailsTabId, ETabState::OpenedTab)
 							)
 						)
 					)
@@ -85,13 +88,13 @@ void FGeometryCollectionEditorToolkit::InitGeometryCollectionAssetEditor(const E
 	}
 }
 
-TSharedRef<SGraphEditor> FGeometryCollectionEditorToolkit::CreateGraphEditorWidget(UDataflow* DataflowToEdit, TSharedPtr<IDetailsView> InPropertiesEditor)
+TSharedRef<SGraphEditor> FGeometryCollectionEditorToolkit::CreateGraphEditorWidget(UDataflow* DataflowToEdit, TSharedPtr<IStructureDetailsView> InNodeDetailsEditor)
 {
 	ensure(DataflowToEdit);
 	using namespace Dataflow;
 	IDataflowEditorPlugin& DataflowEditorModule = FModuleManager::LoadModuleChecked<IDataflowEditorPlugin>(TEXT("DataflowEditor"));
 
-	FDataflowEditorCommands::FGraphEvaluationCallback Evaluate = [&](Dataflow::FNode* Node, Dataflow::FConnection* Out)
+	FDataflowEditorCommands::FGraphEvaluationCallback Evaluate = [&](FDataflowNode* Node, Dataflow::FConnection* Out)
 	{
 		float EvalTime = FGameTime::GetTimeSinceAppStart().GetRealTimeSeconds();
 		return Node->Evaluate(FEngineContext(GeometryCollection, Dataflow, EvalTime, FName("UGeometryCollection")), Out);
@@ -99,21 +102,56 @@ TSharedRef<SGraphEditor> FGeometryCollectionEditorToolkit::CreateGraphEditorWidg
 
 	return SNew(SDataflowGraphEditor, GeometryCollection)
 		.GraphToEdit(DataflowToEdit)
-		.DetailsView(InPropertiesEditor)
+		.DetailsView(InNodeDetailsEditor)
 		.EvaluateGraph(Evaluate);
 }
 
-TSharedPtr<IDetailsView> FGeometryCollectionEditorToolkit::CreatePropertiesEditorWidget(UObject* ObjectToEdit)
+TSharedPtr<IStructureDetailsView> FGeometryCollectionEditorToolkit::CreateNodeDetailsEditorWidget(UObject* ObjectToEdit)
 {
 	ensure(ObjectToEdit);
 	FPropertyEditorModule& PropertyEditorModule = FModuleManager::LoadModuleChecked<FPropertyEditorModule>(TEXT("PropertyEditor"));
 
 	FDetailsViewArgs DetailsViewArgs;
-	DetailsViewArgs.bAllowSearch = true;
-	DetailsViewArgs.bLockable = false;
-	DetailsViewArgs.bUpdatesFromSelection = false;
-	DetailsViewArgs.NameAreaSettings = FDetailsViewArgs::HideNameArea;
-	DetailsViewArgs.NotifyHook = this;
+	{
+		DetailsViewArgs.bAllowSearch = false;
+		DetailsViewArgs.bHideSelectionTip = true;
+		DetailsViewArgs.bLockable = false;
+		DetailsViewArgs.bSearchInitialKeyFocus = true;
+		DetailsViewArgs.bUpdatesFromSelection = false;
+		DetailsViewArgs.NotifyHook = nullptr;
+		DetailsViewArgs.bShowOptions = true;
+		DetailsViewArgs.bShowModifiedPropertiesOption = false;
+		DetailsViewArgs.bShowScrollBar = false;
+	}
+
+	FStructureDetailsViewArgs StructureViewArgs;
+	{
+		StructureViewArgs.bShowObjects = true;
+		StructureViewArgs.bShowAssets = true;
+		StructureViewArgs.bShowClasses = true;
+		StructureViewArgs.bShowInterfaces = true;
+	}
+	TSharedPtr<IStructureDetailsView> DetailsView = PropertyEditorModule.CreateStructureDetailView(DetailsViewArgs, StructureViewArgs, nullptr);
+	DetailsView->GetDetailsView()->SetObject(ObjectToEdit);
+
+	return DetailsView;
+
+}
+
+
+TSharedPtr<IDetailsView> FGeometryCollectionEditorToolkit::CreateAssetDetailsEditorWidget(UObject* ObjectToEdit)
+{
+	ensure(ObjectToEdit);
+	FPropertyEditorModule& PropertyEditorModule = FModuleManager::LoadModuleChecked<FPropertyEditorModule>(TEXT("PropertyEditor"));
+
+	FDetailsViewArgs DetailsViewArgs;
+	{
+		DetailsViewArgs.bAllowSearch = true;
+		DetailsViewArgs.bLockable = false;
+		DetailsViewArgs.bUpdatesFromSelection = false;
+		DetailsViewArgs.NameAreaSettings = FDetailsViewArgs::HideNameArea;
+		DetailsViewArgs.NotifyHook = this;
+	}
 
 	TSharedPtr<IDetailsView> DetailsView = PropertyEditorModule.CreateDetailView(DetailsViewArgs);
 	DetailsView->SetObject(ObjectToEdit);
@@ -122,6 +160,7 @@ TSharedPtr<IDetailsView> FGeometryCollectionEditorToolkit::CreatePropertiesEdito
 }
 
 
+
 TSharedRef<SDockTab> FGeometryCollectionEditorToolkit::SpawnTab_GraphCanvas(const FSpawnTabArgs& Args)
 {
 	check(Args.GetTabId() == GraphCanvasTabId);
@@ -135,14 +174,25 @@ TSharedRef<SDockTab> FGeometryCollectionEditorToolkit::SpawnTab_GraphCanvas(cons
 	return SpawnedTab;
 }
 
-TSharedRef<SDockTab> FGeometryCollectionEditorToolkit::SpawnTab_Properties(const FSpawnTabArgs& Args)
+TSharedRef<SDockTab> FGeometryCollectionEditorToolkit::SpawnTab_AssetDetails(const FSpawnTabArgs& Args)
+{
+	check(Args.GetTabId() == AssetDetailsTabId);
+
+	return SNew(SDockTab)
+		.Label(LOCTEXT("GeometryCollectionEditor_AssetDetails_TabTitle", "Asset Details"))
+		[
+			AssetDetailsEditor.ToSharedRef()
+		];
+}
+
+TSharedRef<SDockTab> FGeometryCollectionEditorToolkit::SpawnTab_NodeDetails(const FSpawnTabArgs& Args)
 {
-	check(Args.GetTabId() == PropertiesTabId);
+	check(Args.GetTabId() == NodeDetailsTabId);
 
 	return SNew(SDockTab)
-		.Label(LOCTEXT("GeometryCollectionEditor_Properties_TabTitle", "Details"))
+		.Label(LOCTEXT("GeometryCollectionEditor_NodeDetails_TabTitle", "Node Details"))
 		[
-			PropertiesEditor.ToSharedRef()
+			NodeDetailsEditor->GetWidget()->AsShared()
 		];
 }
 
@@ -156,8 +206,13 @@ void FGeometryCollectionEditorToolkit::RegisterTabSpawners(const TSharedRef<FTab
 		.SetGroup(WorkspaceMenuCategoryRef)
 		.SetIcon(FSlateIcon(FAppStyle::GetAppStyleSetName(), "GraphEditor.EventGraph_16x"));
 
-	InTabManager->RegisterTabSpawner(PropertiesTabId, FOnSpawnTab::CreateSP(this, &FGeometryCollectionEditorToolkit::SpawnTab_Properties))
-		.SetDisplayName(LOCTEXT("PropertiesTab", "Details"))
+	InTabManager->RegisterTabSpawner(AssetDetailsTabId, FOnSpawnTab::CreateSP(this, &FGeometryCollectionEditorToolkit::SpawnTab_AssetDetails))
+		.SetDisplayName(LOCTEXT("AssetDetailsTab", "Asset Details"))
+		.SetGroup(WorkspaceMenuCategoryRef)
+		.SetIcon(FSlateIcon(FAppStyle::GetAppStyleSetName(), "LevelEditor.Tabs.Details"));
+
+	InTabManager->RegisterTabSpawner(NodeDetailsTabId, FOnSpawnTab::CreateSP(this, &FGeometryCollectionEditorToolkit::SpawnTab_NodeDetails))
+		.SetDisplayName(LOCTEXT("NodeDetailsTab", "Node Details"))
 		.SetGroup(WorkspaceMenuCategoryRef)
 		.SetIcon(FSlateIcon(FAppStyle::GetAppStyleSetName(), "LevelEditor.Tabs.Details"));
 
diff --git a/Engine/Plugins/Experimental/GeometryCollectionPlugin/Source/GeometryCollectionEditor/Public/GeometryCollection/GeometryCollectionEditorToolkit.h b/Engine/Plugins/Experimental/GeometryCollectionPlugin/Source/GeometryCollectionEditor/Public/GeometryCollection/GeometryCollectionEditorToolkit.h
--- a/Engine/Plugins/Experimental/GeometryCollectionPlugin/Source/GeometryCollectionEditor/Public/GeometryCollection/GeometryCollectionEditorToolkit.h
+++ b/Engine/Plugins/Experimental/GeometryCollectionPlugin/Source/GeometryCollectionEditor/Public/GeometryCollection/GeometryCollectionEditorToolkit.h
@@ -34,7 +34,8 @@ public:
 	// Tab spawners 
 	virtual void RegisterTabSpawners(const TSharedRef<FTabManager>& TabManager) override;
 	TSharedRef<SDockTab> SpawnTab_GraphCanvas(const FSpawnTabArgs& Args);
-	TSharedRef<SDockTab> SpawnTab_Properties(const FSpawnTabArgs& Args);
+	TSharedRef<SDockTab> SpawnTab_AssetDetails(const FSpawnTabArgs& Args);
+	TSharedRef<SDockTab> SpawnTab_NodeDetails(const FSpawnTabArgs& Args);
 
 	UDataflow* GetDataflow() { return Dataflow; }
 	const UDataflow* GetDataflow() const { return Dataflow; }
@@ -42,8 +43,11 @@ public:
 	UGeometryCollection* GetGeometryCollection() { return GeometryCollection; }
 	const UGeometryCollection* GetGeometryCollection() const { return GeometryCollection; }
 
-	TSharedPtr<IDetailsView> GetPropertiesEditor() { return PropertiesEditor; }
-	const TSharedPtr<IDetailsView> GetPropertiesEditor() const { return PropertiesEditor; }
+	TSharedPtr<IDetailsView> GetAssetDetailsEditor() { return AssetDetailsEditor; }
+	const TSharedPtr<IDetailsView> GetAssetDetailsEditor() const { return AssetDetailsEditor; }
+
+	TSharedPtr<IStructureDetailsView> GetNodeDetailsEditor() { return NodeDetailsEditor; }
+	const TSharedPtr<IStructureDetailsView> GetNodeDetailsEditor() const { return NodeDetailsEditor; }
 
 	TSharedPtr<SGraphEditor> GetGraphEditor() { return GraphEditor; }
 	const TSharedPtr<SGraphEditor> GetGraphEditor() const { return GraphEditor; }
@@ -51,11 +55,15 @@ public:
 private:
 	static const FName GraphCanvasTabId;
 	TSharedPtr<SGraphEditor> GraphEditor;
-	TSharedRef<SGraphEditor> CreateGraphEditorWidget(UDataflow* ObjectToEdit, TSharedPtr<IDetailsView> PropertiesEditor);
+	TSharedRef<SGraphEditor> CreateGraphEditorWidget(UDataflow* ObjectToEdit, TSharedPtr<IStructureDetailsView> PropertiesEditor);
+
+	static const FName AssetDetailsTabId;
+	TSharedPtr<IDetailsView> AssetDetailsEditor;
+	TSharedPtr<IDetailsView> CreateAssetDetailsEditorWidget(UObject* ObjectToEdit);
 
-	static const FName PropertiesTabId;
-	TSharedPtr<IDetailsView> PropertiesEditor;
-	TSharedPtr<IDetailsView> CreatePropertiesEditorWidget(UObject* ObjectToEdit);
+	static const FName NodeDetailsTabId;
+	TSharedPtr<IStructureDetailsView> NodeDetailsEditor;
+	TSharedPtr<IStructureDetailsView> CreateNodeDetailsEditorWidget(UObject* ObjectToEdit);
 
 	UDataflow* Dataflow = nullptr;
 	UGeometryCollection* GeometryCollection = nullptr;
diff --git a/Engine/Plugins/Experimental/GeometryCollectionPlugin/Source/GeometryCollectionNodes/Private/Dataflow/GeometryCollectionNodes.cpp b/Engine/Plugins/Experimental/GeometryCollectionPlugin/Source/GeometryCollectionNodes/Private/Dataflow/GeometryCollectionNodes.cpp
--- a/Engine/Plugins/Experimental/GeometryCollectionPlugin/Source/GeometryCollectionNodes/Private/Dataflow/GeometryCollectionNodes.cpp
+++ b/Engine/Plugins/Experimental/GeometryCollectionPlugin/Source/GeometryCollectionNodes/Private/Dataflow/GeometryCollectionNodes.cpp
@@ -16,126 +16,128 @@ namespace Dataflow
 {
 	void GeometryCollectionEngineAssetNodes()
 	{
-		DATAFLOW_NODE_REGISTER_CREATION_FACTORY(GetCollectionAssetNode);
-		DATAFLOW_NODE_REGISTER_CREATION_FACTORY(ExampleCollectionEditNode);
-		DATAFLOW_NODE_REGISTER_CREATION_FACTORY(SetCollectionAssetNode);
-		DATAFLOW_NODE_REGISTER_CREATION_FACTORY(ResetGeometryCollection);
+		DATAFLOW_NODE_REGISTER_CREATION_FACTORY(FGetCollectionAssetDataflowNode);
+		DATAFLOW_NODE_REGISTER_CREATION_FACTORY(FExampleCollectionEditDataflowNode);
+		DATAFLOW_NODE_REGISTER_CREATION_FACTORY(FSetCollectionAssetDataflowNode);
+		DATAFLOW_NODE_REGISTER_CREATION_FACTORY(FResetGeometryCollectionDataflowNode);
 	}
+}
 
-	void GetCollectionAssetNode::Evaluate(const FContext& Context, FConnection* Out) const
+void FGetCollectionAssetDataflowNode::Evaluate(const Dataflow::FContext& Context, Dataflow::FConnection* Out) const
+{
+	if (Output.Get() == Out)
 	{
-		if (Output.Get() == Out)
+		if (const Dataflow::FEngineContext* EngineContext = Context.AsType<Dataflow::FEngineContext>(FName("UGeometryCollection")))
 		{
-			if (const FEngineContext* EngineContext = Context.AsType<FEngineContext>(FName("UGeometryCollection")))
+			if (UGeometryCollection* CollectionAsset = Cast<UGeometryCollection>(EngineContext->Owner))
 			{
-				if (UGeometryCollection* CollectionAsset = Cast<UGeometryCollection>(EngineContext->Owner))
+				if (const TSharedPtr<FGeometryCollection, ESPMode::ThreadSafe> AssetCollection = CollectionAsset->GetGeometryCollection())
 				{
-					if (const TSharedPtr<FGeometryCollection, ESPMode::ThreadSafe> AssetCollection = CollectionAsset->GetGeometryCollection())
-					{
-						FManagedArrayCollection* NewCollection = AssetCollection->NewCopy<FManagedArrayCollection>();
-						Output->SetValue(FManagedArrayCollectionSharedPtr(NewCollection), Context);
-					}
+					FManagedArrayCollection* NewCollection = AssetCollection->NewCopy<FManagedArrayCollection>();
+					Output->SetValue(DataType(NewCollection), Context);
 				}
 			}
 		}
 	}
+}
 
-	void ExampleCollectionEditNode::Evaluate(const FContext& Context, FConnection* Out) const
+void FExampleCollectionEditDataflowNode::Evaluate(const Dataflow::FContext& Context, Dataflow::FConnection* Out) const
+{
+	if (Output.Get() == Out)
 	{
-		if (Output.Get() == Out)
+		DataType Collection = Input->GetValue(Context);
+		if (Active)
 		{
-			DataType Collection = Input->GetValue(Context);
-			if (Active.GetValue())
+			TManagedArray<FVector3f>* Vertex = Collection->FindAttribute<FVector3f>("Vertex", "Vertices");
+			for (int i = 0; i < Vertex->Num(); i++)
 			{
-				TManagedArray<FVector3f>* Vertex = Collection->FindAttribute<FVector3f>("Vertex", "Vertices");
-				for (int i = 0; i < Vertex->Num(); i++)
-				{
-					(*Vertex)[i][1] *= Scale.GetValue();
-				}
+				(*Vertex)[i][1] *= Scale;
 			}
-			Output->SetValue(Collection, Context);
 		}
+		Output->SetValue(Collection, Context);
 	}
+}
 
-	void SetCollectionAssetNode::Evaluate(const FContext& Context, FConnection* Out) const
+void FSetCollectionAssetDataflowNode::Evaluate(const Dataflow::FContext& Context, Dataflow::FConnection* Out) const
+{
+	if (Out == nullptr)
 	{
-		if (Out == nullptr)
+		DataType Collection = Input->GetValue(Context);
+		if (const Dataflow::FEngineContext* EngineContext = Context.AsType<Dataflow::FEngineContext>(FName("UGeometryCollection")))
 		{
-			DataType Collection = Input->GetValue(Context);
-			if (const FEngineContext* EngineContext = Context.AsType<FEngineContext>(FName("UGeometryCollection")))
+			if (UGeometryCollection* CollectionAsset = Cast<UGeometryCollection>(EngineContext->Owner))
 			{
-				if (UGeometryCollection* CollectionAsset = Cast<UGeometryCollection>(EngineContext->Owner))
-				{
-					TSharedPtr<FGeometryCollection, ESPMode::ThreadSafe> NewCollection(Collection->NewCopy<FGeometryCollection>());
-					CollectionAsset->SetGeometryCollection(NewCollection);
-					CollectionAsset->InvalidateCollection();
-				}
+				TSharedPtr<FGeometryCollection, ESPMode::ThreadSafe> NewCollection(Collection->NewCopy<FGeometryCollection>());
+				CollectionAsset->SetGeometryCollection(NewCollection);
+				CollectionAsset->InvalidateCollection();
 			}
 		}
 	}
+}
 
-	void ResetGeometryCollection::Evaluate(const FContext& Context, FConnection* Out) const
-	{
-		ManagedArrayOut->SetValue(TSharedPtr<FManagedArrayCollection>(nullptr), Context);
+void FResetGeometryCollectionDataflowNode::Evaluate(const Dataflow::FContext& Context, Dataflow::FConnection* Out) const
+{
+	ManagedArrayOut->SetValue(TSharedPtr<FManagedArrayCollection>(nullptr), Context);
 
-		if (const FEngineContext* EngineContext = Context.AsType<FEngineContext>(FName("UGeometryCollection")))
+	if (const Dataflow::FEngineContext* EngineContext = Context.AsType<Dataflow::FEngineContext>(FName("UGeometryCollection")))
+	{
+		if (UGeometryCollection* GeometryCollectionObject = Cast<UGeometryCollection>(EngineContext->Owner))
 		{
-			if (UGeometryCollection* GeometryCollectionObject = Cast<UGeometryCollection>(EngineContext->Owner))
-			{
-				GeometryCollectionObject->Reset();
+			GeometryCollectionObject->Reset();
 
-				const UObject* Owner = EngineContext->Owner;
-				FName AName("GeometrySource");
-				if (Owner && Owner->GetClass())
+			const UObject* Owner = EngineContext->Owner;
+			FName AName("GeometrySource");
+			if (Owner && Owner->GetClass())
+			{
+				if (const ::FProperty* UEProperty = Owner->GetClass()->FindPropertyByName(AName))
 				{
-					if (const ::FProperty* UEProperty = Owner->GetClass()->FindPropertyByName(AName))
+					if (const FArrayProperty* ArrayProperty = CastField<FArrayProperty>(UEProperty))
 					{
-						if (const FArrayProperty* ArrayProperty = CastField<FArrayProperty>(UEProperty))
+						FScriptArrayHelper_InContainer ArrayHelper(ArrayProperty, Owner);
+						const int32 ArraySize = ArrayHelper.Num();
+						for (int32 Index = 0; Index < ArraySize; ++Index)
 						{
-							FScriptArrayHelper_InContainer ArrayHelper(ArrayProperty, Owner);
-							const int32 ArraySize = ArrayHelper.Num();
-							for (int32 Index = 0; Index < ArraySize; ++Index)
+							if (FGeometryCollectionSource* SourceObject = (FGeometryCollectionSource*)(ArrayHelper.GetRawPtr(Index)))
 							{
-								if (FGeometryCollectionSource* SourceObject = (FGeometryCollectionSource*)(ArrayHelper.GetRawPtr(Index)))
+								if (UObject* ResolvedObject = SourceObject->SourceGeometryObject.ResolveObject())
 								{
-									if (UObject* ResolvedObject = SourceObject->SourceGeometryObject.ResolveObject())
+									if (UStaticMesh* StaticMesh = Cast<UStaticMesh>(ResolvedObject))
 									{
-										if (UStaticMesh* StaticMesh = Cast<UStaticMesh>(ResolvedObject))
-										{
-											TArray<UMaterialInterface*> Materials;
-											Materials.Reserve(StaticMesh->GetStaticMaterials().Num());
+										TArray<UMaterialInterface*> Materials;
+										Materials.Reserve(StaticMesh->GetStaticMaterials().Num());
 
-											for (int32 Index2 = 0; Index2 < StaticMesh->GetStaticMaterials().Num(); ++Index2)
-											{
-												UMaterialInterface* CurrMaterial = StaticMesh->GetMaterial(Index2);
-												Materials.Add(CurrMaterial);
-											}
+										for (int32 Index2 = 0; Index2 < StaticMesh->GetStaticMaterials().Num(); ++Index2)
+										{
+											UMaterialInterface* CurrMaterial = StaticMesh->GetMaterial(Index2);
+											Materials.Add(CurrMaterial);
+										}
 
-											// Geometry collections usually carry the selection material, which we'll delete before appending
-											UMaterialInterface* BoneSelectedMaterial = LoadObject<UMaterialInterface>(nullptr, UGeometryCollection::GetSelectedMaterialPath(), nullptr, LOAD_None, nullptr);
-											GeometryCollectionObject->Materials.Remove(BoneSelectedMaterial);
-											Materials.Remove(BoneSelectedMaterial);
+										// Geometry collections usually carry the selection material, which we'll delete before appending
+										UMaterialInterface* BoneSelectedMaterial = LoadObject<UMaterialInterface>(nullptr, UGeometryCollection::GetSelectedMaterialPath(), nullptr, LOAD_None, nullptr);
+										GeometryCollectionObject->Materials.Remove(BoneSelectedMaterial);
+										Materials.Remove(BoneSelectedMaterial);
 
-											FGeometryCollectionEngineConversion::AppendStaticMesh(StaticMesh, Materials, FTransform(), GeometryCollectionObject);
+										FGeometryCollectionEngineConversion::AppendStaticMesh(StaticMesh, Materials, FTransform(), GeometryCollectionObject);
 
-										}
 									}
 								}
 							}
 						}
 					}
 				}
-				GeometryCollectionObject->UpdateConvexGeometry();
-				GeometryCollectionObject->InitializeMaterials();
-				GeometryCollectionObject->InvalidateCollection();
+			}
+			GeometryCollectionObject->UpdateConvexGeometry();
+			GeometryCollectionObject->InitializeMaterials();
+			GeometryCollectionObject->InvalidateCollection();
 
-				if (const TSharedPtr<FGeometryCollection, ESPMode::ThreadSafe> AssetCollection = GeometryCollectionObject->GetGeometryCollection())
-				{
-					FManagedArrayCollection* NewCollection = AssetCollection->NewCopy<FManagedArrayCollection>();
-					ManagedArrayOut->SetValue(FManagedArrayCollectionSharedPtr(NewCollection), Context);
-				}
+			if (const TSharedPtr<FGeometryCollection, ESPMode::ThreadSafe> AssetCollection = GeometryCollectionObject->GetGeometryCollection())
+			{
+				FManagedArrayCollection* NewCollection = AssetCollection->NewCopy<FManagedArrayCollection>();
+				ManagedArrayOut->SetValue(DataType(NewCollection), Context);
 			}
 		}
 	}
 }
 
+
+
diff --git a/Engine/Plugins/Experimental/GeometryCollectionPlugin/Source/GeometryCollectionNodes/Public/Dataflow/GeometryCollectionNodes.h b/Engine/Plugins/Experimental/GeometryCollectionPlugin/Source/GeometryCollectionNodes/Public/Dataflow/GeometryCollectionNodes.h
--- a/Engine/Plugins/Experimental/GeometryCollectionPlugin/Source/GeometryCollectionNodes/Public/Dataflow/GeometryCollectionNodes.h
+++ b/Engine/Plugins/Experimental/GeometryCollectionPlugin/Source/GeometryCollectionNodes/Public/Dataflow/GeometryCollectionNodes.h
@@ -7,89 +7,96 @@
 #include "CoreMinimal.h"
 #include "Dataflow/DataflowEngine.h"
 
+#include "GeometryCollectionNodes.generated.h"
 
-
-namespace Dataflow
+USTRUCT()
+struct FGetCollectionAssetDataflowNode : public FDataflowNode
 {
+	GENERATED_USTRUCT_BODY()
+	DATAFLOW_NODE_DEFINE_INTERNAL(FGetCollectionAssetDataflowNode, "GetCollectionAsset", "GeometryCollection", "")
 
-	class GetCollectionAssetNode : public FNode
-	{
-		DATAFLOW_NODE_DEFINE_INTERNAL(GetCollectionAssetNode)
+public:
+	typedef Dataflow::FManagedArrayCollectionSharedPtr DataType;
 
-	public:
-		typedef FManagedArrayCollectionSharedPtr DataType;
+	TSharedPtr< class Dataflow::TOutput<DataType> > Output;
 
-		TSharedPtr< class TOutput<DataType> > Output;
+	FGetCollectionAssetDataflowNode(const Dataflow::FNodeParameters& InParam, FGuid InGuid = FGuid::NewGuid())
+		: FDataflowNode(InParam, InGuid)
+		, Output(new Dataflow::TOutput<DataType>(Dataflow::TOutputParameters<DataType>({ FName("CollectionOut"), this })))
+	{}
 
-		GetCollectionAssetNode(const FNodeParameters& InParam, FGuid InGuid = FGuid::NewGuid())
-			: FNode(InParam, InGuid)
-			, Output(new TOutput<DataType>(TOutputParameters<DataType>({ FName("CollectionOut"), this })))
-		{}
-
-		virtual void Evaluate(const FContext& Context, FConnection* Out) const override;
-	};
+	virtual void Evaluate(const Dataflow::FContext& Context, Dataflow::FConnection* Out) const override;
+};
 
+USTRUCT()
+struct FExampleCollectionEditDataflowNode : public FDataflowNode
+{
+	GENERATED_USTRUCT_BODY()
+	DATAFLOW_NODE_DEFINE_INTERNAL(FExampleCollectionEditDataflowNode, "ExampleCollectionEdit", "GeometryCollection", "")
 
-	class ExampleCollectionEditNode : public FNode
-	{
-		DATAFLOW_NODE_DEFINE_INTERNAL(ExampleCollectionEditNode)
+public:
+	UPROPERTY(EditAnywhere, Category = "Dataflow");
+	bool Active = true;
 
-	public:
-		typedef FManagedArrayCollectionSharedPtr DataType;
-		TSharedPtr< class TInput<DataType> > Input;
-		TSharedPtr< class TOutput<DataType> > Output;
-		TProperty<bool> Active;
-		TProperty<float> Scale;
+	UPROPERTY(EditAnywhere, Category = "Dataflow");
+	float Scale = 1.0;
 
-		ExampleCollectionEditNode(const FNodeParameters& InParam, FGuid InGuid = FGuid::NewGuid())
-			: FNode(InParam, InGuid)
-			, Input(new TInput<DataType>(TInputParameters<DataType>({ FName("CollectionIn"), this })))
-			, Output(new TOutput<DataType>(TOutputParameters<DataType>({ FName("CollectionOut"), this })))
-			, Active({ FName("Active"), true, this })
-			, Scale({ FName("Scale"), 10.f, this })
-		{}
+	typedef Dataflow::FManagedArrayCollectionSharedPtr DataType;
+	TSharedPtr< class Dataflow::TInput<DataType> > Input;
+	TSharedPtr< class Dataflow::TOutput<DataType> > Output;
 
-		virtual void Evaluate(const FContext& Context, FConnection* Out) const override;
+	FExampleCollectionEditDataflowNode(const Dataflow::FNodeParameters& InParam, FGuid InGuid = FGuid::NewGuid())
+		: FDataflowNode(InParam, InGuid)
+		, Input(new Dataflow::TInput<DataType>(Dataflow::TInputParameters<DataType>({ FName("CollectionIn"), this })))
+		, Output(new Dataflow::TOutput<DataType>(Dataflow::TOutputParameters<DataType>({ FName("CollectionOut"), this })))
+	{}
 
-	};
+	virtual void Evaluate(const Dataflow::FContext& Context, Dataflow::FConnection* Out) const override;
 
-	class SetCollectionAssetNode : public FNode
-	{
-		DATAFLOW_NODE_DEFINE_INTERNAL(SetCollectionAssetNode)
+};
 
-	public:
-		typedef FManagedArrayCollectionSharedPtr DataType;
-		TSharedPtr< class TInput<DataType> > Input;
+USTRUCT()
+struct FSetCollectionAssetDataflowNode : public FDataflowNode
+{
+	GENERATED_USTRUCT_BODY()
+	DATAFLOW_NODE_DEFINE_INTERNAL(FSetCollectionAssetDataflowNode, "SetCollectionAsset", "GeometryCollection", "")
 
-		SetCollectionAssetNode(const FNodeParameters& InParam, FGuid InGuid = FGuid::NewGuid())
-			: FNode(InParam, InGuid)
-			, Input(new TInput<DataType>(TInputParameters<DataType>({ FName("CollectionIn"), this })))
-		{}
+public:
+	typedef Dataflow::FManagedArrayCollectionSharedPtr DataType;
+	TSharedPtr< class Dataflow::TInput<DataType> > Input;
 
-		virtual void Evaluate(const FContext& Context, FConnection* Out) const override;
+	FSetCollectionAssetDataflowNode(const Dataflow::FNodeParameters& InParam, FGuid InGuid = FGuid::NewGuid())
+		: FDataflowNode(InParam, InGuid)
+		, Input(new Dataflow::TInput<DataType>(Dataflow::TInputParameters<DataType>({ FName("CollectionIn"), this })))
+	{}
 
-	};
+	virtual void Evaluate(const Dataflow::FContext& Context, Dataflow::FConnection* Out) const override;
 
+};
 
-	class ResetGeometryCollection : public FNode
-	{
-		DATAFLOW_NODE_DEFINE_INTERNAL(ResetGeometryCollection)
+USTRUCT()
+struct FResetGeometryCollectionDataflowNode : public FDataflowNode
+{
+	GENERATED_USTRUCT_BODY()
+	DATAFLOW_NODE_DEFINE_INTERNAL(FResetGeometryCollectionDataflowNode, "ResetGeometryCollection", "GeometryCollection", "")
 
-	public:
-		typedef FManagedArrayCollectionSharedPtr DataType;
+public:
+	typedef Dataflow::FManagedArrayCollectionSharedPtr DataType;
 
-		TSharedPtr< class TOutput<DataType> > ManagedArrayOut;
+	TSharedPtr< class Dataflow::TOutput<DataType> > ManagedArrayOut;
 
-		ResetGeometryCollection(const FNodeParameters& InParam, FGuid InGuid = FGuid::NewGuid())
-			: FNode(InParam, InGuid)
-			, ManagedArrayOut(new TOutput<DataType>(TOutputParameters<DataType>({ FName("ManagedArrayOut"), this })))
-		{}
+	FResetGeometryCollectionDataflowNode(const Dataflow::FNodeParameters& InParam, FGuid InGuid = FGuid::NewGuid())
+		: FDataflowNode(InParam, InGuid)
+		, ManagedArrayOut(new Dataflow::TOutput<DataType>(Dataflow::TOutputParameters<DataType>({ FName("ManagedArrayOut"), this })))
+	{}
 
 
-		virtual void Evaluate(const FContext& Context, FConnection* Out) const override;
+	virtual void Evaluate(const Dataflow::FContext& Context, Dataflow::FConnection* Out) const override;
 
-	};
+};
 
+namespace Dataflow
+{
 	void GeometryCollectionEngineAssetNodes();
 
 
diff --git a/Engine/Source/Runtime/Experimental/Dataflow/Core/Private/Dataflow/DataflowConnection.cpp b/Engine/Source/Runtime/Experimental/Dataflow/Core/Private/Dataflow/DataflowConnection.cpp
--- a/Engine/Source/Runtime/Experimental/Dataflow/Core/Private/Dataflow/DataflowConnection.cpp
+++ b/Engine/Source/Runtime/Experimental/Dataflow/Core/Private/Dataflow/DataflowConnection.cpp
@@ -7,7 +7,7 @@
 
 namespace Dataflow
 {
-	FConnection::FConnection(FPin::EDirection InDirection, FName InType, FName InName, FNode* InOwningNode, FGuid InGuid)
+	FConnection::FConnection(FPin::EDirection InDirection, FName InType, FName InName, FDataflowNode* InOwningNode, FGuid InGuid)
 		: Direction(InDirection)
 		, Type(InType)
 		, Name(InName)
@@ -15,9 +15,9 @@ namespace Dataflow
 		, OwningNode(InOwningNode)
 	{}
 
-	void FConnection::BindInput(FNode* InNode, FConnection* That) 
+	void FConnection::BindInput(FDataflowNode* InNode, FConnection* That) 
 	{ InNode->AddInput(That); }
-	void FConnection::BindOutput(FNode* InNode, FConnection* That) 
+	void FConnection::BindOutput(FDataflowNode* InNode, FConnection* That) 
 	{ InNode->AddOutput(That); }
 
 }
diff --git a/Engine/Source/Runtime/Experimental/Dataflow/Core/Private/Dataflow/DataflowGraph.cpp b/Engine/Source/Runtime/Experimental/Dataflow/Core/Private/Dataflow/DataflowGraph.cpp
--- a/Engine/Source/Runtime/Experimental/Dataflow/Core/Private/Dataflow/DataflowGraph.cpp
+++ b/Engine/Source/Runtime/Experimental/Dataflow/Core/Private/Dataflow/DataflowGraph.cpp
@@ -19,7 +19,7 @@ namespace Dataflow
 	}
 
 
-	void FGraph::RemoveNode(TSharedPtr<FNode> Node)
+	void FGraph::RemoveNode(TSharedPtr<FDataflowNode> Node)
 	{
 		for (FConnection* Output : Node->GetOutputs())
 		{
@@ -106,7 +106,7 @@ namespace Dataflow
 			int32 ArNum = Nodes.Num();
 
 			Ar << ArNum;
-			for (TSharedPtr<FNode> Node : Nodes)
+			for (TSharedPtr<FDataflowNode> Node : Nodes)
 			{
 				ArGuid = Node->GetGuid();
 				ArType = Node->GetType();
@@ -126,6 +126,7 @@ namespace Dataflow
 						ArName = Conn->GetName();
 						Ar << ArGuid << ArType << ArName;
 					}
+
 					Node->SerializeInternal(Ar);
 				}
 				DATAFLOW_OPTIONAL_BLOCK_WRITE_END();
@@ -140,7 +141,7 @@ namespace Dataflow
 			FName ArType, ArName;
 			int32 ArNum = 0;
 
-			TMap<FGuid, TSharedPtr<FNode> > NodeGuidMap;
+			TMap<FGuid, TSharedPtr<FDataflowNode> > NodeGuidMap;
 			TMap<FGuid, FConnection* > ConnectionGuidMap;
 
 			Ar << ArNum;
@@ -148,7 +149,7 @@ namespace Dataflow
 			{
 				Ar << ArGuid << ArType << ArName;
 
-				TSharedPtr<FNode> Node = FNodeFactory::GetInstance()->NewNodeFromRegisteredType(*this, { ArGuid,ArType,ArName });
+				TSharedPtr<FDataflowNode> Node = FNodeFactory::GetInstance()->NewNodeFromRegisteredType(*this, { ArGuid,ArType,ArName });
 				DATAFLOW_OPTIONAL_BLOCK_READ_BEGIN(Node != nullptr)
 				{
 					ensure(!NodeGuidMap.Contains(ArGuid));
diff --git a/Engine/Source/Runtime/Experimental/Dataflow/Core/Private/Dataflow/DataflowNode.cpp b/Engine/Source/Runtime/Experimental/Dataflow/Core/Private/Dataflow/DataflowNode.cpp
--- a/Engine/Source/Runtime/Experimental/Dataflow/Core/Private/Dataflow/DataflowNode.cpp
+++ b/Engine/Source/Runtime/Experimental/Dataflow/Core/Private/Dataflow/DataflowNode.cpp
@@ -3,167 +3,69 @@
 #include "Dataflow/DataflowNode.h"
 
 #include "Dataflow/DataflowNode.h"
-#include "Dataflow/DataflowProperty.h"
 #include "Dataflow/DataflowArchive.h"
 
-namespace Dataflow
+void FDataflowNode::AddInput(Dataflow::FConnection* InPtr)
 {
-
-	void FNode::AddProperty(FProperty* InPtr)
+	for (Dataflow::FConnection* In : Inputs)
 	{
-		for (FProperty* In : Properties)
-		{
-			ensureMsgf(!In->GetName().IsEqual(InPtr->GetName()), TEXT("Add Property Failed: Existing Node property already defined with name (%s)"), *InPtr->GetName().ToString());
-		}
-		Properties.Add(InPtr);
-	}
-
-	void FNode::AddInput(FConnection* InPtr)
-	{ 
-		for (FConnection* In : Inputs)
-		{
-			ensureMsgf(!In->GetName().IsEqual(InPtr->GetName()), TEXT("Add Input Failed: Existing Node input already defined with name (%s)"), *InPtr->GetName().ToString());
-		}
-		Inputs.Add(InPtr); 
-	}
-
-	void FNode::AddOutput(FConnection* InPtr)
-	{ 
-		for (FConnection* Out : Outputs)
-		{
-			ensureMsgf(!Out->GetName().IsEqual(InPtr->GetName()), TEXT("Add Output Failed: Existing Node output already defined with name (%s)"), *InPtr->GetName().ToString());
-		}
-		Outputs.Add(InPtr);
+		ensureMsgf(!In->GetName().IsEqual(InPtr->GetName()), TEXT("Add Input Failed: Existing Node input already defined with name (%s)"), *InPtr->GetName().ToString());
 	}
+	Inputs.Add(InPtr);
+}
 
-
-
-	TArray<FPin> FNode::GetPins() const
+void FDataflowNode::AddOutput(Dataflow::FConnection* InPtr)
+{
+	for (Dataflow::FConnection* Out : Outputs)
 	{
-		TArray<FPin> RetVal;
-		for (FConnection* Con : Inputs)
-			RetVal.Add({ FPin::EDirection::INPUT,Con->GetType(), Con->GetName() });
-		for (FConnection* Con : Outputs)
-			RetVal.Add({ FPin::EDirection::OUTPUT,Con->GetType(), Con->GetName() });
-		return RetVal;
+		ensureMsgf(!Out->GetName().IsEqual(InPtr->GetName()), TEXT("Add Output Failed: Existing Node output already defined with name (%s)"), *InPtr->GetName().ToString());
 	}
+	Outputs.Add(InPtr);
+}
 
-	void FNode::InvalidateOutputs()
-	{
-		for(FConnection * Output :  Outputs)
-		{
-			Output->Invalidate();
-		}
-	}
 
-	FConnection* FNode::FindInput(FName InName) const
-	{
-		for (FConnection* Input : Inputs)
-		{
-			if (Input->GetName().IsEqual(InName))
-			{
-				return Input;
-			}
-		}
-		return nullptr;
-	}
 
-	FConnection* FNode::FindOutput(FName InName) const
-	{
-		for (FConnection* Output : Outputs)
-		{
-			if (Output->GetName().IsEqual(InName))
-			{
-				return Output;
-			}
-		}
-		return nullptr;
-	}
+TArray<Dataflow::FPin> FDataflowNode::GetPins() const
+{
+	TArray<Dataflow::FPin> RetVal;
+	for (Dataflow::FConnection* Con : Inputs)
+		RetVal.Add({ Dataflow::FPin::EDirection::INPUT,Con->GetType(), Con->GetName() });
+	for (Dataflow::FConnection* Con : Outputs)
+		RetVal.Add({ Dataflow::FPin::EDirection::OUTPUT,Con->GetType(), Con->GetName() });
+	return RetVal;
+}
 
-	TMap<FName, FProperty*> FNode::GetPropertyMap()
+void FDataflowNode::InvalidateOutputs()
+{
+	for (Dataflow::FConnection* Output : Outputs)
 	{
-		TMap<FName, FProperty*> Map;
-		for (FProperty* Prop : GetProperties())
-		{
-			Map.Add(TTuple<FName, FProperty*>(Prop->GetName(),Prop));
-		}
-		return Map;
+		Output->Invalidate();
 	}
+}
 
-	const TMap<FName, const FProperty*> FNode::GetPropertyMap() const
+Dataflow::FConnection* FDataflowNode::FindInput(FName InName) const
+{
+	for (Dataflow::FConnection* Input : Inputs)
 	{
-		TMap<FName, const FProperty*> Map;
-		for (FProperty* Prop : GetProperties())
+		if (Input->GetName().IsEqual(InName))
 		{
-			Map.Add(TTuple<FName, const FProperty*>(Prop->GetName(), Prop));
+			return Input;
 		}
-		return Map;
 	}
+	return nullptr;
+}
 
-	void FNode::SerializeInternal(FArchive& Ar)
+Dataflow::FConnection* FDataflowNode::FindOutput(FName InName) const
+{
+	for (Dataflow::FConnection* Output : Outputs)
 	{
-		TMap<FName, FProperty*> Map = GetPropertyMap();
-		TArray< FProperty* >& CurrentProperties = GetProperties();
-		uint8 NumProperties = CurrentProperties.Num();
-		Ar << NumProperties;
-
-		for (int32 i = NumProperties - 1; i >= 0; i--)
+		if (Output->GetName().IsEqual(InName))
 		{
-			uint8 Type;
-			FString StrName;
-			int64 PropertySizeOf = 0, PropertyDataSize = 0;
-
-			if (Ar.IsSaving())
-			{
-				FProperty* Prop = CurrentProperties[i];
-				Type = (uint8)Prop->GetType();
-				StrName = Prop->GetName().ToString();
-				PropertySizeOf = Prop->SizeOf();
-				Ar << Type << StrName;
-
-				DATAFLOW_OPTIONAL_BLOCK_WRITE_BEGIN()
-				{
-					Prop->Serialize(Ar);
-				}
-				DATAFLOW_OPTIONAL_BLOCK_WRITE_END()
-
-			}
-			else if (Ar.IsLoading())
-			{
-
-				Ar << Type << StrName;
-				FName PropName(StrName);
-
-				FProperty* NewProperty = FProperty::NewProperty((FProperty::EType)Type, PropName);
-				DATAFLOW_OPTIONAL_BLOCK_READ_BEGIN(NewProperty != nullptr)
-				{
-					if (Map.Contains(PropName) && Map[PropName]->GetType() == (FProperty::EType)Type)
-					{
-						Map[PropName]->Serialize(Ar);
-						delete NewProperty;
-					}
-					else
-					{
-						NewProperty->Serialize(Ar);
-
-						if (!Map.Contains(PropName))
-						{
-							AddProperty(NewProperty);
-						}
-						else
-						{
-							delete NewProperty;
-						}
-					}
-				}
-				DATAFLOW_OPTIONAL_BLOCK_READ_ELSE()
-				{
-					// Silently skip
-				}
-				DATAFLOW_OPTIONAL_BLOCK_READ_END()
-			}
+			return Output;
 		}
 	}
-
+	return nullptr;
 }
 
+
+
diff --git a/Engine/Source/Runtime/Experimental/Dataflow/Core/Private/Dataflow/DataflowNodeFactory.cpp b/Engine/Source/Runtime/Experimental/Dataflow/Core/Private/Dataflow/DataflowNodeFactory.cpp
--- a/Engine/Source/Runtime/Experimental/Dataflow/Core/Private/Dataflow/DataflowNodeFactory.cpp
+++ b/Engine/Source/Runtime/Experimental/Dataflow/Core/Private/Dataflow/DataflowNodeFactory.cpp
@@ -8,13 +8,13 @@ namespace Dataflow
 {
 	FNodeFactory* FNodeFactory::Instance = nullptr;
 
-	TSharedPtr<FNode> FNodeFactory::NewNodeFromRegisteredType(FGraph& Graph, const FNewNodeParameters& Param)
+	TSharedPtr<FDataflowNode> FNodeFactory::NewNodeFromRegisteredType(FGraph& Graph, const FNewNodeParameters& Param)
 	{ 
 		if (ClassMap.Contains(Param.Type))
 		{
 			return Graph.AddNode(ClassMap[Param.Type](Param));
 		}
-		return TSharedPtr<FNode>(nullptr);
+		return TSharedPtr<FDataflowNode>(nullptr);
 	}
 }
 
diff --git a/Engine/Source/Runtime/Experimental/Dataflow/Core/Private/Dataflow/DataflowProperty.cpp b/Engine/Source/Runtime/Experimental/Dataflow/Core/Private/Dataflow/DataflowProperty.cpp
deleted file mode 100644
diff --git a/Engine/Source/Runtime/Experimental/Dataflow/Core/Public/Dataflow/DataflowConnection.h b/Engine/Source/Runtime/Experimental/Dataflow/Core/Public/Dataflow/DataflowConnection.h
--- a/Engine/Source/Runtime/Experimental/Dataflow/Core/Public/Dataflow/DataflowConnection.h
+++ b/Engine/Source/Runtime/Experimental/Dataflow/Core/Public/Dataflow/DataflowConnection.h
@@ -5,9 +5,12 @@
 #include "CoreMinimal.h"
 #include "Dataflow/DataflowNodeParameters.h"
 
+
+struct FDataflowNode;
+
+
 namespace Dataflow
 {
-	class FNode;
 
 	template<class T> inline DATAFLOWCORE_API FName GraphConnectionTypeName();
 	template<class T> inline DATAFLOWCORE_API T DeepCopy(const T&);
@@ -36,19 +39,19 @@ namespace Dataflow
 		FName Type;
 		FName Name;
 		FGuid  Guid;
-		FNode* OwningNode = nullptr;
+		FDataflowNode* OwningNode = nullptr;
 
-		friend class FNode;
-		static void BindInput(FNode* InNode, FConnection*);
-		static void BindOutput(FNode* InNode, FConnection*);
+		friend struct FDataflowNode;
+		static void BindInput(FDataflowNode* InNode, FConnection*);
+		static void BindOutput(FDataflowNode* InNode, FConnection*);
 
 
 	public:
-		FConnection(FPin::EDirection Direction, FName InType, FName InName, FNode* OwningNode = nullptr, FGuid InGuid = FGuid::NewGuid());
+		FConnection(FPin::EDirection Direction, FName InType, FName InName, FDataflowNode* OwningNode = nullptr, FGuid InGuid = FGuid::NewGuid());
 		virtual ~FConnection() {};
 
-		FNode* GetOwningNode() { return OwningNode; }
-		const FNode* GetOwningNode() const { return OwningNode; }
+		FDataflowNode* GetOwningNode() { return OwningNode; }
+		const FDataflowNode* GetOwningNode() const { return OwningNode; }
 
 		FPin::EDirection GetDirection() const { return Direction; }
 
diff --git a/Engine/Source/Runtime/Experimental/Dataflow/Core/Public/Dataflow/DataflowCore.h b/Engine/Source/Runtime/Experimental/Dataflow/Core/Public/Dataflow/DataflowCore.h
--- a/Engine/Source/Runtime/Experimental/Dataflow/Core/Public/Dataflow/DataflowCore.h
+++ b/Engine/Source/Runtime/Experimental/Dataflow/Core/Public/Dataflow/DataflowCore.h
@@ -8,4 +8,3 @@
 #include "Dataflow/DataflowNode.h"
 #include "Dataflow/DataflowNodeFactory.h"
 #include "Dataflow/DataflowNodeParameters.h"
-#include "Dataflow/DataflowProperty.h"
diff --git a/Engine/Source/Runtime/Experimental/Dataflow/Core/Public/Dataflow/DataflowGraph.h b/Engine/Source/Runtime/Experimental/Dataflow/Core/Public/Dataflow/DataflowGraph.h
--- a/Engine/Source/Runtime/Experimental/Dataflow/Core/Public/Dataflow/DataflowGraph.h
+++ b/Engine/Source/Runtime/Experimental/Dataflow/Core/Public/Dataflow/DataflowGraph.h
@@ -47,15 +47,15 @@ namespace Dataflow
 	{
 
 		FGuid  Guid;
-		TArray< TSharedPtr<FNode> > Nodes;
+		TArray< TSharedPtr<FDataflowNode> > Nodes;
 		TArray< FLink > Connections;
 		TSet< FName > DisabledNodes;
 	public:
 		FGraph(FGuid InGuid = FGuid::NewGuid());
 		virtual ~FGraph() {}
 
-		const TArray< TSharedPtr<FNode> >& GetNodes() const {return Nodes;}
-		TArray< TSharedPtr<FNode> >& GetNodes() { return Nodes; }
+		const TArray< TSharedPtr<FDataflowNode> >& GetNodes() const {return Nodes;}
+		TArray< TSharedPtr<FDataflowNode> >& GetNodes() { return Nodes; }
 		int NumNodes() { return Nodes.Num(); }
 
 		template<class T> TSharedPtr<T> AddNode(T* InNode)
@@ -65,20 +65,20 @@ namespace Dataflow
 			return NewNode;
 		}
 
-		TSharedPtr<FNode> FindBaseNode(FGuid InGuid)
+		TSharedPtr<FDataflowNode> FindBaseNode(FGuid InGuid)
 		{
-			for (TSharedPtr<FNode> Node : Nodes)
+			for (TSharedPtr<FDataflowNode> Node : Nodes)
 			{
 				if (Node->GetGuid() == InGuid)
 				{
 					return Node;
 				}
 			}
-			return TSharedPtr<FNode>(nullptr);
+			return TSharedPtr<FDataflowNode>(nullptr);
 		}
 
 
-		void RemoveNode(TSharedPtr<FNode> Node);
+		void RemoveNode(TSharedPtr<FDataflowNode> Node);
 
 		void ClearConnections(FConnection*);
 		void Connect(FConnection* Input, FConnection* Output);
diff --git a/Engine/Source/Runtime/Experimental/Dataflow/Core/Public/Dataflow/DataflowInputOutput.h b/Engine/Source/Runtime/Experimental/Dataflow/Core/Public/Dataflow/DataflowInputOutput.h
--- a/Engine/Source/Runtime/Experimental/Dataflow/Core/Public/Dataflow/DataflowInputOutput.h
+++ b/Engine/Source/Runtime/Experimental/Dataflow/Core/Public/Dataflow/DataflowInputOutput.h
@@ -17,14 +17,14 @@ namespace Dataflow
 
 	template<class T>
 	struct DATAFLOWCORE_API TInputParameters {
-		TInputParameters(FName InName, FNode* InOwner, T InDefault = T())
+		TInputParameters(FName InName, FDataflowNode* InOwner, T InDefault = T())
 			: Type(GraphConnectionTypeName<T>())
 			, Name(InName)
 			, Owner(InOwner)
 			, Default(InDefault) {}
 		FName Type;
 		FName Name;
-		FNode* Owner = nullptr;
+		FDataflowNode* Owner = nullptr;
 		T Default;
 	};
 
@@ -113,13 +113,13 @@ namespace Dataflow
 
 	template<class T>
 	struct DATAFLOWCORE_API TOutputParameters {
-		TOutputParameters(FName InName, FNode * InOwner) 
+		TOutputParameters(FName InName, FDataflowNode * InOwner) 
 			: Type(GraphConnectionTypeName<T>())
 			, Name(InName)
 			, Owner(InOwner) {}
 		FName Type;
 		FName Name;
-		FNode* Owner = nullptr;
+		FDataflowNode* Owner = nullptr;
 	};
 
 	template<class T>
diff --git a/Engine/Source/Runtime/Experimental/Dataflow/Core/Public/Dataflow/DataflowNode.h b/Engine/Source/Runtime/Experimental/Dataflow/Core/Public/Dataflow/DataflowNode.h
--- a/Engine/Source/Runtime/Experimental/Dataflow/Core/Public/Dataflow/DataflowNode.h
+++ b/Engine/Source/Runtime/Experimental/Dataflow/Core/Public/Dataflow/DataflowNode.h
@@ -5,86 +5,103 @@
 #include "CoreMinimal.h"
 #include "Dataflow/DataflowConnection.h"
 #include "Dataflow/DataflowNodeParameters.h"
+#include "UObject/StructOnScope.h"
 
-namespace Dataflow
-{
-	class FProperty;
+#include "DataflowNode.generated.h"
 
+namespace Dataflow {
 	struct DATAFLOWCORE_API FNodeParameters {
 		FName Name;
 	};
+	class FGraph;
+}
 
 
-	//
-	// FNode
-	//
-	class DATAFLOWCORE_API FNode
-	{
-		friend class FGraph;
-		friend class FConnection;
-		friend class FProperty;
+//
+// FNode
+//
+USTRUCT()
+struct DATAFLOWCORE_API FDataflowNode
+{
+	GENERATED_USTRUCT_BODY()
 
-		FGuid Guid;
-		FName Name;
+	friend class Dataflow::FGraph;
+	friend class Dataflow::FConnection;
 
-		TArray< FConnection* > Inputs;
-		TArray< FConnection* > Outputs;
-		TArray< FProperty* > Properties;
-	public:
+	FGuid Guid;
+	FName Name;
 
-		FNode(const FNodeParameters& Param, FGuid InGuid = FGuid::NewGuid())
-			: Guid(InGuid)
-			, Name(Param.Name)
-		{}
+	TArray< Dataflow::FConnection* > Inputs;
+	TArray< Dataflow::FConnection* > Outputs;
 
-		virtual ~FNode() {}
+	FDataflowNode()
+		: Guid(FGuid())
+		, Name("Invalid")
+	{}
 
-		FName GetName() const { return Name; }
-		void SetName(FName InName) { Name = InName; }
+	FDataflowNode(const Dataflow::FNodeParameters& Param, FGuid InGuid = FGuid::NewGuid())
+		: Guid(InGuid)
+		, Name(Param.Name)
+	{}
 
-		virtual FName GetType() const { check(true); return FName("invalid"); }
-		FGuid GetGuid() const { return Guid; }
+	virtual ~FDataflowNode() {}
 
-		TArray<FPin> GetPins() const;
-		FConnection* FindInput(FName Name) const;
-		FConnection* FindOutput(FName Name) const;
+	FName GetName() const { return Name; }
+	void SetName(FName InName) { Name = InName; }
 
-		virtual void Evaluate(const FContext& Context, FConnection*) const { ensure(false); }
-		void InvalidateOutputs();
+	virtual FString GetDisplayName() { return ""; }
+	virtual FString GetCatagory() { return ""; }
+	virtual FString GetTags() { return ""; }
 
-		virtual void SerializeInternal(FArchive& Ar);
+	virtual FName GetType() const { check(true); return FName("invalid"); }
+	FGuid GetGuid() const { return Guid; }
 
-		void AddProperty(FProperty* InPtr);
-		const TArray< FProperty* >& GetProperties() const { return Properties; }
-		TArray< FProperty* >& GetProperties() { return Properties; }
+	TArray<Dataflow::FPin> GetPins() const;
+	Dataflow::FConnection* FindInput(FName Name) const;
+	Dataflow::FConnection* FindOutput(FName Name) const;
 
-		const TMap<FName, const FProperty*> GetPropertyMap() const;
- 		TMap<FName, FProperty*> GetPropertyMap();
+	virtual void Evaluate(const  Dataflow::FContext& Context, Dataflow::FConnection*) const { ensure(false); }
+	void InvalidateOutputs();
 
-		void AddInput(FConnection* InPtr);
-		const TArray< FConnection* >& GetInputs() const { return Inputs; }
-		TArray< FConnection* >& GetInputs() { return Inputs; }
+	void AddInput(Dataflow::FConnection* InPtr);
+	const TArray< Dataflow::FConnection* >& GetInputs() const { return Inputs; }
+	TArray< Dataflow::FConnection* >& GetInputs() { return Inputs; }
 
-		void AddOutput(FConnection* InPtr);
-		const TArray< FConnection* >& GetOutputs() const { return Outputs; }
-		TArray< FConnection* >& GetOutputs() { return Outputs; }
+	void AddOutput(Dataflow::FConnection* InPtr);
+	const TArray< Dataflow::FConnection* >& GetOutputs() const { return Outputs; }
+	TArray< Dataflow::FConnection* >& GetOutputs() { return Outputs; }
 
-	};
+	virtual void SerializeInternal(FArchive& Ar) { check(false); }
+	virtual FStructOnScope* NewScructOnScope() { return nullptr; }
+};
 
-//
-// Used these macros to register dataflow nodes. 
-//
+namespace Dataflow
+{
+	//
+	// Used these macros to register dataflow nodes. 
+	//
 
-#define DATAFLOW_NODE_REGISTER_CREATION_FACTORY(A)					\
-	FNodeFactory::GetInstance()->RegisterNode(					\
-		A::StaticType(),										\
-		[](const FNewNodeParameters& InParam)					\
+#define DATAFLOW_NODE_REGISTER_CREATION_FACTORY(A)						\
+	FNodeFactory::GetInstance()->RegisterNode(							\
+		A::StaticType(),													\
+		[](const FNewNodeParameters& InParam)							\
 			{return new A({InParam.Name}, InParam.Guid); });
 
-#define DATAFLOW_NODE_DEFINE_INTERNAL(A)								\
-public:															\
-   static FName StaticType() {return #A;}						\
-   virtual FName GetType() const { return #A; }					\
+#define DATAFLOW_NODE_DEFINE_INTERNAL(TYPE, DISPLAY_NAME, CATAGORY, TAGS)\
+public:																	\
+	static FName StaticType() {return #TYPE;}							\
+	static FName StaticDisplay() {return DISPLAY_NAME;}					\
+	virtual FName GetType() const { return #TYPE; }						\
+	virtual FStructOnScope* NewScructOnScope() override {				\
+	   return new FStructOnScope(TYPE::StaticStruct(), (uint8*)this);}	\
+	virtual void SerializeInternal(FArchive& Ar) override {				\
+		UScriptStruct* const Struct = TYPE::StaticStruct();				\
+		Struct->SerializeTaggedProperties(Ar, (uint8*)this,				\
+		Struct, nullptr);}												\
+	virtual FString GetDisplayName() override { return DISPLAY_NAME; }	\
+	virtual FString GetCatagory() override { return CATAGORY; }			\
+	virtual FString GetTags() override { return TAGS; }					\
+	TYPE() {}															\
 private:
 
 
diff --git a/Engine/Source/Runtime/Experimental/Dataflow/Core/Public/Dataflow/DataflowNodeFactory.h b/Engine/Source/Runtime/Experimental/Dataflow/Core/Public/Dataflow/DataflowNodeFactory.h
--- a/Engine/Source/Runtime/Experimental/Dataflow/Core/Public/Dataflow/DataflowNodeFactory.h
+++ b/Engine/Source/Runtime/Experimental/Dataflow/Core/Public/Dataflow/DataflowNodeFactory.h
@@ -5,9 +5,10 @@
 #include "CoreMinimal.h"
 #include "Dataflow/DataflowGraph.h"
 
+struct FDataflowNode;
+
 namespace Dataflow
 {
-	class FNode;
 	class FConnection;
 
 	struct DATAFLOWCORE_API FNewNodeParameters {
@@ -21,7 +22,7 @@ namespace Dataflow
 	//
 	class DATAFLOWCORE_API FNodeFactory
 	{
-		typedef TFunction<FNode* (const FNewNodeParameters&)> FNewNodeFunction;
+		typedef TFunction<FDataflowNode* (const FNewNodeParameters&)> FNewNodeFunction;
 
 		TMap<FName, FNewNodeFunction > ClassMap;
 
@@ -48,7 +49,7 @@ namespace Dataflow
 			}
 		}
 
-		TSharedPtr<FNode> NewNodeFromRegisteredType(FGraph& Graph, const FNewNodeParameters& Param);
+		TSharedPtr<FDataflowNode> NewNodeFromRegisteredType(FGraph& Graph, const FNewNodeParameters& Param);
 
 		template<class T> TSharedPtr<T> NewNode(FGraph& Graph, const FNewNodeParameters& Param)
 		{
diff --git a/Engine/Source/Runtime/Experimental/Dataflow/Core/Public/Dataflow/DataflowProperty.h b/Engine/Source/Runtime/Experimental/Dataflow/Core/Public/Dataflow/DataflowProperty.h
deleted file mode 100644
diff --git a/Engine/Source/Runtime/Experimental/Dataflow/Engine/Private/Dataflow/DataflowEdNode.cpp b/Engine/Source/Runtime/Experimental/Dataflow/Engine/Private/Dataflow/DataflowEdNode.cpp
--- a/Engine/Source/Runtime/Experimental/Dataflow/Engine/Private/Dataflow/DataflowEdNode.cpp
+++ b/Engine/Source/Runtime/Experimental/Dataflow/Engine/Private/Dataflow/DataflowEdNode.cpp
@@ -24,7 +24,7 @@ void UDataflowEdNode::AllocateDefaultPins()
 	{
 		if (DataflowNodeGuid.IsValid())
 		{
-			if (TSharedPtr<Dataflow::FNode> DataflowNode = DataflowGraph->FindBaseNode(DataflowNodeGuid))
+			if (TSharedPtr<FDataflowNode> DataflowNode = DataflowGraph->FindBaseNode(DataflowNodeGuid))
 			{
 				for (const Dataflow::FPin& Pin : DataflowNode->GetPins())
 				{
@@ -54,7 +54,7 @@ void UDataflowEdNode::PinConnectionListChanged(UEdGraphPin* Pin)
 {
 	if (ensure(IsBound()))
 	{
-		if (TSharedPtr<Dataflow::FNode> DataflowNode = DataflowGraph->FindBaseNode(DataflowNodeGuid))
+		if (TSharedPtr<FDataflowNode> DataflowNode = DataflowGraph->FindBaseNode(DataflowNodeGuid))
 		{
 			if (Dataflow::FConnection* ConnectionInput = DataflowNode->FindInput(FName(Pin->GetName())) )
 			{
@@ -65,7 +65,7 @@ void UDataflowEdNode::PinConnectionListChanged(UEdGraphPin* Pin)
 					{
 						if (ensure(LinkedNode->IsBound()))
 						{
-							if (TSharedPtr<Dataflow::FNode> LinkedDataflowNode = DataflowGraph->FindBaseNode(LinkedNode->GetDataflowNodeGuid()))
+							if (TSharedPtr<FDataflowNode> LinkedDataflowNode = DataflowGraph->FindBaseNode(LinkedNode->GetDataflowNodeGuid()))
 							{
 								if (Dataflow::FConnection* LinkedConBase = LinkedDataflowNode->FindOutput(FName(LinkedCon->GetName())))
 								{
@@ -87,7 +87,7 @@ void UDataflowEdNode::PinConnectionListChanged(UEdGraphPin* Pin)
 						{
 							if (ensure(LinkedNode->IsBound()))
 							{
-								if (TSharedPtr<Dataflow::FNode> LinkedDataflowNode = DataflowGraph->FindBaseNode(LinkedNode->GetDataflowNodeGuid()))
+								if (TSharedPtr<FDataflowNode> LinkedDataflowNode = DataflowGraph->FindBaseNode(LinkedNode->GetDataflowNodeGuid()))
 								{
 									if (Dataflow::FConnection* LinkedConBase = LinkedDataflowNode->FindInput(FName(LinkedCon->GetName())))
 									{
