Dataflow : FContext CRT udpate

diff --git a/Engine/Plugins/Experimental/Dataflow/Source/DataflowNodes/Private/Dataflow/DataflowSkeletalMeshNodes.cpp b/Engine/Plugins/Experimental/Dataflow/Source/DataflowNodes/Private/Dataflow/DataflowSkeletalMeshNodes.cpp
--- a/Engine/Plugins/Experimental/Dataflow/Source/DataflowNodes/Private/Dataflow/DataflowSkeletalMeshNodes.cpp
+++ b/Engine/Plugins/Experimental/Dataflow/Source/DataflowNodes/Private/Dataflow/DataflowSkeletalMeshNodes.cpp
@@ -24,7 +24,7 @@ void FGetSkeletalMeshDataflowNode::Evaluate(const Dataflow::FContext& Context, D
 	{
 		SkeletalMeshOut->SetValue(SkeletalMesh, Context);
 	}
-	else if (const Dataflow::FEngineContext* EngineContext = Context.AsType<Dataflow::FEngineContext>(FName("UFleshAsset")))
+	else if (const Dataflow::FEngineContext* EngineContext = Context.AsType<Dataflow::FEngineContext>())
 	{
 		if (const USkeletalMesh* SkeletalMeshFromOwner = Dataflow::Reflection::FindObjectPtrProperty<USkeletalMesh>(
 			EngineContext->Owner, PropertyName))
diff --git a/Engine/Plugins/Experimental/Dataflow/Source/DataflowNodes/Private/Dataflow/DataflowStaticMeshNodes.cpp b/Engine/Plugins/Experimental/Dataflow/Source/DataflowNodes/Private/Dataflow/DataflowStaticMeshNodes.cpp
--- a/Engine/Plugins/Experimental/Dataflow/Source/DataflowNodes/Private/Dataflow/DataflowStaticMeshNodes.cpp
+++ b/Engine/Plugins/Experimental/Dataflow/Source/DataflowNodes/Private/Dataflow/DataflowStaticMeshNodes.cpp
@@ -18,7 +18,7 @@ void FGetStaticMeshDataflowNode::Evaluate(const Dataflow::FContext& Context, Dat
 	{
 		StaticMeshOut->SetValue(StaticMesh, Context);
 	}
-	else if (const Dataflow::FEngineContext* EngineContext = Context.AsType<Dataflow::FEngineContext>(FName("UFleshAsset")))
+	else if (const Dataflow::FEngineContext* EngineContext = Context.AsType<Dataflow::FEngineContext>())
 	{
 		if (const UStaticMesh* StaticMeshFromOwner = Dataflow::Reflection::FindObjectPtrProperty<UStaticMesh>(
 			EngineContext->Owner, PropertyName))
diff --git a/Engine/Plugins/Experimental/GeometryCollectionPlugin/Source/GeometryCollectionEditor/Private/GeometryCollection/GeometryCollectionEditorToolkit.cpp b/Engine/Plugins/Experimental/GeometryCollectionPlugin/Source/GeometryCollectionEditor/Private/GeometryCollection/GeometryCollectionEditorToolkit.cpp
--- a/Engine/Plugins/Experimental/GeometryCollectionPlugin/Source/GeometryCollectionEditor/Private/GeometryCollection/GeometryCollectionEditorToolkit.cpp
+++ b/Engine/Plugins/Experimental/GeometryCollectionPlugin/Source/GeometryCollectionEditor/Private/GeometryCollection/GeometryCollectionEditorToolkit.cpp
@@ -97,7 +97,7 @@ TSharedRef<SGraphEditor> FGeometryCollectionEditorToolkit::CreateGraphEditorWidg
 	FDataflowEditorCommands::FGraphEvaluationCallback Evaluate = [&](FDataflowNode* Node, Dataflow::FConnection* Out)
 	{
 		float EvalTime = FGameTime::GetTimeSinceAppStart().GetRealTimeSeconds();
-		return Node->Evaluate(FEngineContext(GeometryCollection, Dataflow, EvalTime, FName("UGeometryCollection")), Out);
+		return Node->Evaluate(FEngineContext(GeometryCollection, Dataflow, EvalTime, FString("UGeometryCollection")), Out);
 	};
 
 	return SNew(SDataflowGraphEditor, GeometryCollection)
diff --git a/Engine/Plugins/Experimental/GeometryCollectionPlugin/Source/GeometryCollectionNodes/Private/Dataflow/GeometryCollectionNodes.cpp b/Engine/Plugins/Experimental/GeometryCollectionPlugin/Source/GeometryCollectionNodes/Private/Dataflow/GeometryCollectionNodes.cpp
--- a/Engine/Plugins/Experimental/GeometryCollectionPlugin/Source/GeometryCollectionNodes/Private/Dataflow/GeometryCollectionNodes.cpp
+++ b/Engine/Plugins/Experimental/GeometryCollectionPlugin/Source/GeometryCollectionNodes/Private/Dataflow/GeometryCollectionNodes.cpp
@@ -27,7 +27,7 @@ void FGetCollectionAssetDataflowNode::Evaluate(const Dataflow::FContext& Context
 {
 	if (Output.Get() == Out)
 	{
-		if (const Dataflow::FEngineContext* EngineContext = Context.AsType<Dataflow::FEngineContext>(FName("UGeometryCollection")))
+		if (const Dataflow::FEngineContext* EngineContext = Context.AsType<Dataflow::FEngineContext>())
 		{
 			if (UGeometryCollection* CollectionAsset = Cast<UGeometryCollection>(EngineContext->Owner))
 			{
@@ -63,7 +63,7 @@ void FSetCollectionAssetDataflowNode::Evaluate(const Dataflow::FContext& Context
 	if (Out == nullptr)
 	{
 		DataType Collection = Input->GetValue(Context);
-		if (const Dataflow::FEngineContext* EngineContext = Context.AsType<Dataflow::FEngineContext>(FName("UGeometryCollection")))
+		if (const Dataflow::FEngineContext* EngineContext = Context.AsType<Dataflow::FEngineContext>())
 		{
 			if (UGeometryCollection* CollectionAsset = Cast<UGeometryCollection>(EngineContext->Owner))
 			{
@@ -79,7 +79,7 @@ void FResetGeometryCollectionDataflowNode::Evaluate(const Dataflow::FContext& Co
 {
 	ManagedArrayOut->SetValue(TSharedPtr<FManagedArrayCollection>(nullptr), Context);
 
-	if (const Dataflow::FEngineContext* EngineContext = Context.AsType<Dataflow::FEngineContext>(FName("UGeometryCollection")))
+	if (const Dataflow::FEngineContext* EngineContext = Context.AsType<Dataflow::FEngineContext>())
 	{
 		if (UGeometryCollection* GeometryCollectionObject = Cast<UGeometryCollection>(EngineContext->Owner))
 		{
diff --git a/Engine/Source/Runtime/Experimental/Dataflow/Core/Public/Dataflow/DataflowNodeParameters.h b/Engine/Source/Runtime/Experimental/Dataflow/Core/Public/Dataflow/DataflowNodeParameters.h
--- a/Engine/Source/Runtime/Experimental/Dataflow/Core/Public/Dataflow/DataflowNodeParameters.h
+++ b/Engine/Source/Runtime/Experimental/Dataflow/Core/Public/Dataflow/DataflowNodeParameters.h
@@ -11,12 +11,14 @@ namespace Dataflow
 	class DATAFLOWCORE_API FContext
 	{
 	public:
-		FContext(float InTime, FName InType = FName("Unknown"))
-		: Timestamp(InTime) 
-		, Type(InType){}
+		FContext(float InTime, FString InType = FString(""))
+		: Timestamp(InTime)
+		, Type(StaticType().Append(InType))
+		{}
 
 		float Timestamp = 0.f;
-		FName Type = FName("Unknown");
+		FString Type;
+		static FString StaticType() { return "FContext"; }
 
 		uint32 GetTypeHash() const
 		{
@@ -24,9 +26,9 @@ namespace Dataflow
 		}
 
 		template<class T>
-		const T* AsType(FName InType) const
+		const T* AsType() const
 		{
-			if (Type.IsEqual(InType))
+			if (Type.Contains(T::StaticType()))
 			{
 				return (T*)this;
 			}
diff --git a/Engine/Source/Runtime/Experimental/Dataflow/Engine/Private/Dataflow/DataflowObjectInterface.cpp b/Engine/Source/Runtime/Experimental/Dataflow/Engine/Private/Dataflow/DataflowObjectInterface.cpp
--- a/Engine/Source/Runtime/Experimental/Dataflow/Engine/Private/Dataflow/DataflowObjectInterface.cpp
+++ b/Engine/Source/Runtime/Experimental/Dataflow/Engine/Private/Dataflow/DataflowObjectInterface.cpp
@@ -4,8 +4,8 @@
 
 namespace Dataflow
 {
-	FEngineContext::FEngineContext(UObject* InOwner, UDataflow* InGraph, float InTime, FName InType)
-			: FContext(InTime, InType)
+	FEngineContext::FEngineContext(UObject* InOwner, UDataflow* InGraph, float InTime, FString InType)
+			: FContext(InTime, StaticType().Append(InType) )
 			, Owner(InOwner)
 			, Graph(InGraph)
 		{}
diff --git a/Engine/Source/Runtime/Experimental/Dataflow/Engine/Public/Dataflow/DataflowObjectInterface.h b/Engine/Source/Runtime/Experimental/Dataflow/Engine/Public/Dataflow/DataflowObjectInterface.h
--- a/Engine/Source/Runtime/Experimental/Dataflow/Engine/Public/Dataflow/DataflowObjectInterface.h
+++ b/Engine/Source/Runtime/Experimental/Dataflow/Engine/Public/Dataflow/DataflowObjectInterface.h
@@ -13,7 +13,8 @@ namespace Dataflow
 	class DATAFLOWENGINE_API FEngineContext : public FContext
 	{
 	public:
-		FEngineContext(UObject* Owner, UDataflow* InGraph, float InTime, FName InType=FName("Unknown"));
+		FEngineContext(UObject* Owner, UDataflow* InGraph, float InTime, FString InType = "");
+		static FString StaticType() { return "FEngineContext"; }
 
 		UObject* Owner = nullptr;
 		UDataflow* Graph = nullptr;
