Dataflow : Nodes
-Removed shared pointer inputs and outputs.

diff --git a/Engine/Plugins/Experimental/Dataflow/Source/DataflowNodes/Private/Dataflow/DataflowSkeletalMeshNodes.cpp b/Engine/Plugins/Experimental/Dataflow/Source/DataflowNodes/Private/Dataflow/DataflowSkeletalMeshNodes.cpp
--- a/Engine/Plugins/Experimental/Dataflow/Source/DataflowNodes/Private/Dataflow/DataflowSkeletalMeshNodes.cpp
+++ b/Engine/Plugins/Experimental/Dataflow/Source/DataflowNodes/Private/Dataflow/DataflowSkeletalMeshNodes.cpp
@@ -19,29 +19,29 @@ namespace Dataflow
 }
 void FGetSkeletalMeshDataflowNode::Evaluate(const Dataflow::FContext& Context, Dataflow::FConnection* Out) const
 {
-	SkeletalMeshOut->SetValue(nullptr, Context);
+	SkeletalMeshOut.SetValue(nullptr, Context);
 	if (SkeletalMesh)
 	{
-		SkeletalMeshOut->SetValue(SkeletalMesh, Context);
+		SkeletalMeshOut.SetValue(SkeletalMesh, Context);
 	}
 	else if (const Dataflow::FEngineContext* EngineContext = Context.AsType<Dataflow::FEngineContext>())
 	{
 		if (const USkeletalMesh* SkeletalMeshFromOwner = Dataflow::Reflection::FindObjectPtrProperty<USkeletalMesh>(
 			EngineContext->Owner, PropertyName))
 		{
-			SkeletalMeshOut->SetValue(SkeletalMeshFromOwner, Context);
+			SkeletalMeshOut.SetValue(SkeletalMeshFromOwner, Context);
 		}
 	}
 }
 
 void FSkeletalMeshBoneDataflowNode::Evaluate(const Dataflow::FContext& Context, Dataflow::FConnection* Out) const
 {
-	BoneIndexOut->SetValue(INDEX_NONE, Context);
+	BoneIndexOut.SetValue(INDEX_NONE, Context);
 
-	if (const USkeletalMesh* SkeletalMesh = SkeletalMeshIn->GetValue(Context))
+	if (const USkeletalMesh* SkeletalMesh = SkeletalMeshIn.GetValue(Context))
 	{
 		int32 Index = SkeletalMesh->GetRefSkeleton().FindBoneIndex(BoneName);
-		BoneIndexOut->SetValue(Index, Context);
+		BoneIndexOut.SetValue(Index, Context);
 	}
 }
 
diff --git a/Engine/Plugins/Experimental/Dataflow/Source/DataflowNodes/Private/Dataflow/DataflowStaticMeshNodes.cpp b/Engine/Plugins/Experimental/Dataflow/Source/DataflowNodes/Private/Dataflow/DataflowStaticMeshNodes.cpp
--- a/Engine/Plugins/Experimental/Dataflow/Source/DataflowNodes/Private/Dataflow/DataflowStaticMeshNodes.cpp
+++ b/Engine/Plugins/Experimental/Dataflow/Source/DataflowNodes/Private/Dataflow/DataflowStaticMeshNodes.cpp
@@ -13,17 +13,17 @@ namespace Dataflow
 
 void FGetStaticMeshDataflowNode::Evaluate(const Dataflow::FContext& Context, Dataflow::FConnection* Out) const
 {
-	StaticMeshOut->SetValue(nullptr, Context);
+	StaticMeshOut.SetValue(nullptr, Context);
 	if (StaticMesh)
 	{
-		StaticMeshOut->SetValue(StaticMesh, Context);
+		StaticMeshOut.SetValue(StaticMesh, Context);
 	}
 	else if (const Dataflow::FEngineContext* EngineContext = Context.AsType<Dataflow::FEngineContext>())
 	{
 		if (const UStaticMesh* StaticMeshFromOwner = Dataflow::Reflection::FindObjectPtrProperty<UStaticMesh>(
 			EngineContext->Owner, PropertyName))
 		{
-			StaticMeshOut->SetValue(StaticMeshFromOwner, Context);
+			StaticMeshOut.SetValue(StaticMeshFromOwner, Context);
 		}
 	}
 }
diff --git a/Engine/Plugins/Experimental/Dataflow/Source/DataflowNodes/Public/Dataflow/DataflowSkeletalMeshNodes.h b/Engine/Plugins/Experimental/Dataflow/Source/DataflowNodes/Public/Dataflow/DataflowSkeletalMeshNodes.h
--- a/Engine/Plugins/Experimental/Dataflow/Source/DataflowNodes/Public/Dataflow/DataflowSkeletalMeshNodes.h
+++ b/Engine/Plugins/Experimental/Dataflow/Source/DataflowNodes/Public/Dataflow/DataflowSkeletalMeshNodes.h
@@ -31,11 +31,11 @@ public:
 	UPROPERTY(EditAnywhere, Category = "Dataflow" )
 	FName PropertyName = "SkeletalMesh";
 
-	TSharedPtr< class Dataflow::TOutput<Dataflow::USkeletalMeshPtr> > SkeletalMeshOut;
+	Dataflow::TOutput<Dataflow::USkeletalMeshPtr> SkeletalMeshOut;
 
 	FGetSkeletalMeshDataflowNode(const Dataflow::FNodeParameters& InParam, FGuid InGuid = FGuid::NewGuid())
 		: FDataflowNode(InParam, InGuid)
-		, SkeletalMeshOut(new Dataflow::TOutput<Dataflow::USkeletalMeshPtr>(Dataflow::TOutputParameters<Dataflow::USkeletalMeshPtr>({ FName("SkeletalMeshOut"), this })))
+		, SkeletalMeshOut({ FName("SkeletalMeshOut"), this })
 	{}
 
 
@@ -52,13 +52,13 @@ public:
 	UPROPERTY(EditAnywhere, Category = "Dataflow")
 	FName BoneName;
 
-	TSharedPtr< class Dataflow::TInput<Dataflow::USkeletalMeshPtr> > SkeletalMeshIn;
-	TSharedPtr< class Dataflow::TOutput<int> > BoneIndexOut;
+	Dataflow::TInput<Dataflow::USkeletalMeshPtr> SkeletalMeshIn;
+	Dataflow::TOutput<int> BoneIndexOut;
 
 	FSkeletalMeshBoneDataflowNode(const Dataflow::FNodeParameters& InParam, FGuid InGuid = FGuid::NewGuid())
 		: FDataflowNode(InParam, InGuid)
-		, SkeletalMeshIn(new Dataflow::TInput<Dataflow::USkeletalMeshPtr>(Dataflow::TInputParameters<Dataflow::USkeletalMeshPtr>({ FName("SkeletalMeshIn"), this })))
-		, BoneIndexOut(new Dataflow::TOutput<int>(Dataflow::TOutputParameters<int>({ FName("BoneIndexOut"), this, INDEX_NONE })))
+		, SkeletalMeshIn({ FName("SkeletalMeshIn"), this })
+		, BoneIndexOut({ FName("BoneIndexOut"), this, INDEX_NONE })
 	{}
 
 
diff --git a/Engine/Plugins/Experimental/Dataflow/Source/DataflowNodes/Public/Dataflow/DataflowStaticMeshNodes.h b/Engine/Plugins/Experimental/Dataflow/Source/DataflowNodes/Public/Dataflow/DataflowStaticMeshNodes.h
--- a/Engine/Plugins/Experimental/Dataflow/Source/DataflowNodes/Public/Dataflow/DataflowStaticMeshNodes.h
+++ b/Engine/Plugins/Experimental/Dataflow/Source/DataflowNodes/Public/Dataflow/DataflowStaticMeshNodes.h
@@ -32,11 +32,11 @@ public:
 	UPROPERTY(EditAnywhere, Category = "Dataflow")
 		FName PropertyName = "StaticMesh";
 
-	TSharedPtr< class Dataflow::TOutput<Dataflow::UStaticMeshPtr> > StaticMeshOut;
+	Dataflow::TOutput<Dataflow::UStaticMeshPtr> StaticMeshOut;
 
 	FGetStaticMeshDataflowNode(const Dataflow::FNodeParameters& InParam, FGuid InGuid = FGuid::NewGuid())
 		: FDataflowNode(InParam, InGuid)
-		, StaticMeshOut(new Dataflow::TOutput<Dataflow::UStaticMeshPtr>(Dataflow::TOutputParameters<Dataflow::UStaticMeshPtr>({ FName("StaticMeshOut"), this })))
+		, StaticMeshOut({ FName("StaticMeshOut"), this })
 	{}
 
 
diff --git a/Engine/Plugins/Experimental/GeometryCollectionPlugin/Source/GeometryCollectionNodes/Private/Dataflow/GeometryCollectionNodes.cpp b/Engine/Plugins/Experimental/GeometryCollectionPlugin/Source/GeometryCollectionNodes/Private/Dataflow/GeometryCollectionNodes.cpp
--- a/Engine/Plugins/Experimental/GeometryCollectionPlugin/Source/GeometryCollectionNodes/Private/Dataflow/GeometryCollectionNodes.cpp
+++ b/Engine/Plugins/Experimental/GeometryCollectionPlugin/Source/GeometryCollectionNodes/Private/Dataflow/GeometryCollectionNodes.cpp
@@ -25,7 +25,7 @@ namespace Dataflow
 
 void FGetCollectionAssetDataflowNode::Evaluate(const Dataflow::FContext& Context, Dataflow::FConnection* Out) const
 {
-	if (Output.Get() == Out)
+	if (&Output == Out)
 	{
 		if (const Dataflow::FEngineContext* EngineContext = Context.AsType<Dataflow::FEngineContext>())
 		{
@@ -34,7 +34,7 @@ void FGetCollectionAssetDataflowNode::Evaluate(const Dataflow::FContext& Context
 				if (const TSharedPtr<FGeometryCollection, ESPMode::ThreadSafe> AssetCollection = CollectionAsset->GetGeometryCollection())
 				{
 					FManagedArrayCollection* NewCollection = AssetCollection->NewCopy<FManagedArrayCollection>();
-					Output->SetValue(DataType(NewCollection), Context);
+					Output.SetValue(DataType(NewCollection), Context);
 				}
 			}
 		}
@@ -43,9 +43,9 @@ void FGetCollectionAssetDataflowNode::Evaluate(const Dataflow::FContext& Context
 
 void FExampleCollectionEditDataflowNode::Evaluate(const Dataflow::FContext& Context, Dataflow::FConnection* Out) const
 {
-	if (Output.Get() == Out)
+	if (&Output == Out)
 	{
-		DataType Collection = Input->GetValue(Context);
+		DataType Collection = Input.GetValue(Context);
 		if (Active)
 		{
 			TManagedArray<FVector3f>* Vertex = Collection->FindAttribute<FVector3f>("Vertex", "Vertices");
@@ -54,7 +54,7 @@ void FExampleCollectionEditDataflowNode::Evaluate(const Dataflow::FContext& Cont
 				(*Vertex)[i][1] *= Scale;
 			}
 		}
-		Output->SetValue(Collection, Context);
+		Output.SetValue(Collection, Context);
 	}
 }
 
@@ -62,7 +62,7 @@ void FSetCollectionAssetDataflowNode::Evaluate(const Dataflow::FContext& Context
 {
 	if (Out == nullptr)
 	{
-		DataType Collection = Input->GetValue(Context);
+		DataType Collection = Input.GetValue(Context);
 		if (const Dataflow::FEngineContext* EngineContext = Context.AsType<Dataflow::FEngineContext>())
 		{
 			if (UGeometryCollection* CollectionAsset = Cast<UGeometryCollection>(EngineContext->Owner))
@@ -77,7 +77,7 @@ void FSetCollectionAssetDataflowNode::Evaluate(const Dataflow::FContext& Context
 
 void FResetGeometryCollectionDataflowNode::Evaluate(const Dataflow::FContext& Context, Dataflow::FConnection* Out) const
 {
-	ManagedArrayOut->SetValue(TSharedPtr<FManagedArrayCollection>(nullptr), Context);
+	ManagedArrayOut.SetValue(TSharedPtr<FManagedArrayCollection>(nullptr), Context);
 
 	if (const Dataflow::FEngineContext* EngineContext = Context.AsType<Dataflow::FEngineContext>())
 	{
@@ -133,7 +133,7 @@ void FResetGeometryCollectionDataflowNode::Evaluate(const Dataflow::FContext& Co
 			if (const TSharedPtr<FGeometryCollection, ESPMode::ThreadSafe> AssetCollection = GeometryCollectionObject->GetGeometryCollection())
 			{
 				FManagedArrayCollection* NewCollection = AssetCollection->NewCopy<FManagedArrayCollection>();
-				ManagedArrayOut->SetValue(DataType(NewCollection), Context);
+				ManagedArrayOut.SetValue(DataType(NewCollection), Context);
 			}
 		}
 	}
diff --git a/Engine/Plugins/Experimental/GeometryCollectionPlugin/Source/GeometryCollectionNodes/Public/Dataflow/GeometryCollectionNodes.h b/Engine/Plugins/Experimental/GeometryCollectionPlugin/Source/GeometryCollectionNodes/Public/Dataflow/GeometryCollectionNodes.h
--- a/Engine/Plugins/Experimental/GeometryCollectionPlugin/Source/GeometryCollectionNodes/Public/Dataflow/GeometryCollectionNodes.h
+++ b/Engine/Plugins/Experimental/GeometryCollectionPlugin/Source/GeometryCollectionNodes/Public/Dataflow/GeometryCollectionNodes.h
@@ -18,11 +18,11 @@ struct FGetCollectionAssetDataflowNode : public FDataflowNode
 public:
 	typedef Dataflow::FManagedArrayCollectionSharedPtr DataType;
 
-	TSharedPtr< class Dataflow::TOutput<DataType> > Output;
+	Dataflow::TOutput<DataType> Output;
 
 	FGetCollectionAssetDataflowNode(const Dataflow::FNodeParameters& InParam, FGuid InGuid = FGuid::NewGuid())
 		: FDataflowNode(InParam, InGuid)
-		, Output(new Dataflow::TOutput<DataType>(Dataflow::TOutputParameters<DataType>({ FName("CollectionOut"), this })))
+		, Output({ FName("CollectionOut"), this })
 	{}
 
 	virtual void Evaluate(const Dataflow::FContext& Context, Dataflow::FConnection* Out) const override;
@@ -42,13 +42,13 @@ public:
 	float Scale = 1.0;
 
 	typedef Dataflow::FManagedArrayCollectionSharedPtr DataType;
-	TSharedPtr< class Dataflow::TInput<DataType> > Input;
-	TSharedPtr< class Dataflow::TOutput<DataType> > Output;
+	Dataflow::TInput<DataType> Input;
+	Dataflow::TOutput<DataType> Output;
 
 	FExampleCollectionEditDataflowNode(const Dataflow::FNodeParameters& InParam, FGuid InGuid = FGuid::NewGuid())
 		: FDataflowNode(InParam, InGuid)
-		, Input(new Dataflow::TInput<DataType>(Dataflow::TInputParameters<DataType>({ FName("CollectionIn"), this })))
-		, Output(new Dataflow::TOutput<DataType>(Dataflow::TOutputParameters<DataType>({ FName("CollectionOut"), this })))
+		, Input({ FName("CollectionIn"), this })
+		, Output({ FName("CollectionOut"), this })
 	{}
 
 	virtual void Evaluate(const Dataflow::FContext& Context, Dataflow::FConnection* Out) const override;
@@ -63,11 +63,11 @@ struct FSetCollectionAssetDataflowNode : public FDataflowNode
 
 public:
 	typedef Dataflow::FManagedArrayCollectionSharedPtr DataType;
-	TSharedPtr< class Dataflow::TInput<DataType> > Input;
+	Dataflow::TInput<DataType> Input;
 
 	FSetCollectionAssetDataflowNode(const Dataflow::FNodeParameters& InParam, FGuid InGuid = FGuid::NewGuid())
 		: FDataflowNode(InParam, InGuid)
-		, Input(new Dataflow::TInput<DataType>(Dataflow::TInputParameters<DataType>({ FName("CollectionIn"), this })))
+		, Input({ FName("CollectionIn"), this })
 	{}
 
 	virtual void Evaluate(const Dataflow::FContext& Context, Dataflow::FConnection* Out) const override;
@@ -83,11 +83,11 @@ struct FResetGeometryCollectionDataflowNode : public FDataflowNode
 public:
 	typedef Dataflow::FManagedArrayCollectionSharedPtr DataType;
 
-	TSharedPtr< class Dataflow::TOutput<DataType> > ManagedArrayOut;
+	Dataflow::TOutput<DataType> ManagedArrayOut;
 
 	FResetGeometryCollectionDataflowNode(const Dataflow::FNodeParameters& InParam, FGuid InGuid = FGuid::NewGuid())
 		: FDataflowNode(InParam, InGuid)
-		, ManagedArrayOut(new Dataflow::TOutput<DataType>(Dataflow::TOutputParameters<DataType>({ FName("ManagedArrayOut"), this })))
+		, ManagedArrayOut({ FName("ManagedArrayOut"), this })
 	{}
 
 
diff --git a/Engine/Source/Runtime/Experimental/Dataflow/Core/Private/Dataflow/DataflowConnection.cpp b/Engine/Source/Runtime/Experimental/Dataflow/Core/Private/Dataflow/DataflowConnection.cpp
--- a/Engine/Source/Runtime/Experimental/Dataflow/Core/Private/Dataflow/DataflowConnection.cpp
+++ b/Engine/Source/Runtime/Experimental/Dataflow/Core/Private/Dataflow/DataflowConnection.cpp
@@ -16,9 +16,13 @@ namespace Dataflow
 	{}
 
 	void FConnection::BindInput(FDataflowNode* InNode, FConnection* That) 
-	{ InNode->AddInput(That); }
-	void FConnection::BindOutput(FDataflowNode* InNode, FConnection* That) 
-	{ InNode->AddOutput(That); }
+	{
+		if (InNode && That) { InNode->AddInput(That); }
+	}
+	void FConnection::BindOutput(FDataflowNode* InNode, FConnection* That)
+	{
+		if (InNode && That) { InNode->AddOutput(That); }
+	}
 
 }
 
diff --git a/Engine/Source/Runtime/Experimental/Dataflow/Core/Public/Dataflow/DataflowInputOutput.h b/Engine/Source/Runtime/Experimental/Dataflow/Core/Public/Dataflow/DataflowInputOutput.h
--- a/Engine/Source/Runtime/Experimental/Dataflow/Core/Public/Dataflow/DataflowInputOutput.h
+++ b/Engine/Source/Runtime/Experimental/Dataflow/Core/Public/Dataflow/DataflowInputOutput.h
@@ -18,7 +18,7 @@ namespace Dataflow
 
 	template<class T>
 	struct DATAFLOWCORE_API TInputParameters {
-		TInputParameters(FName InName, FDataflowNode* InOwner, T InDefault = T())
+		TInputParameters(FName InName = FName(""), FDataflowNode * InOwner = nullptr, T InDefault = T())
 			: Type(GraphConnectionTypeName<T>())
 			, Name(InName)
 			, Owner(InOwner)
@@ -38,7 +38,7 @@ namespace Dataflow
 		T Default = T();
 		TOutput<T>* Connection;
 	public:
-		TInput(const TInputParameters<T>& Param, FGuid InGuid = FGuid::NewGuid())
+		TInput(const TInputParameters<T>& Param = {}, FGuid InGuid = FGuid::NewGuid())
 			: FConnection(FPin::EDirection::INPUT, Param.Type, Param.Name, Param.Owner, InGuid)
 			, Default(Param.Default)
 			, Connection(nullptr)
@@ -117,7 +117,7 @@ namespace Dataflow
 	{
 		typedef TFunction<T(const  Dataflow::FContext& Context, Dataflow::FConnection*)> PassthroughCallback;
 
-		TOutputParameters(FName InName, FDataflowNode* InOwner, T InDefault = T())
+		TOutputParameters(FName InName = FName(""), FDataflowNode * InOwner = nullptr, T InDefault = T())
 			: Type(GraphConnectionTypeName<T>())
 			, Name(InName)
 			, Owner(InOwner)
@@ -143,14 +143,15 @@ namespace Dataflow
 		typedef FConnection Super;
 		friend class FConnection;
 
-		uint32 CacheKey = UINT_MAX;
-		TCacheValue<T> Cache;
+		mutable uint32 CacheKey = UINT_MAX;
+		mutable TCacheValue<T> Cache;
+
 		TArray< TInput<T>* > Connections;
 		PassthroughCallback Passthrough;
 		T Default = T();
 
 	public:
-		TOutput(const TOutputParameters<T>& Param, FGuid InGuid = FGuid::NewGuid())
+		TOutput(const TOutputParameters<T>& Param = {}, FGuid InGuid = FGuid::NewGuid())
 			: FConnection(FPin::EDirection::OUTPUT, Param.Type, Param.Name, Param.Owner, InGuid)
 			, Passthrough(Param.Passthrough)
 			, Default(Param.Default)
@@ -182,7 +183,7 @@ namespace Dataflow
 		
 		virtual bool RemoveConnection(FConnection* InInput) override { Connections.RemoveSwap((TInput<T>*)InInput); return true; }
 
-		void SetValue(T InVal, const FContext& Context)
+		void SetValue(T InVal, const FContext& Context) const
 		{
 			CacheKey = Context.GetTypeHash();
 			Cache.Data = InVal;
