GeometryCollection : Init nodes

diff --git a/Engine/Plugins/Experimental/GeometryCollectionPlugin/Source/GeometryCollectionNodes/Private/Dataflow/GeometryCollectionNodesPlugin.cpp b/Engine/Plugins/Experimental/GeometryCollectionPlugin/Source/GeometryCollectionNodes/Private/Dataflow/GeometryCollectionNodesPlugin.cpp
--- a/Engine/Plugins/Experimental/GeometryCollectionPlugin/Source/GeometryCollectionNodes/Private/Dataflow/GeometryCollectionNodesPlugin.cpp
+++ b/Engine/Plugins/Experimental/GeometryCollectionPlugin/Source/GeometryCollectionNodes/Private/Dataflow/GeometryCollectionNodesPlugin.cpp
@@ -6,6 +6,9 @@
 
 #include "Dataflow/DataflowNodeFactory.h"
 #include "Dataflow/GeometryCollectionNodes.h"
+#include "Dataflow/GeometryCollectionProcessingNodes.h"
+#include "Dataflow/GeometryCollectionSkeletalMeshNodes.h"
+#include "Dataflow/GeometryCollectionTetrahedralNodes.h"
 
 #define LOCTEXT_NAMESPACE "DataflowNodes"
 
@@ -13,6 +16,9 @@
 void IGeometryCollectionNodesPlugin::StartupModule()
 {
 	Dataflow::GeometryCollectionEngineAssetNodes();
+	Dataflow::GeometryCollectionProcessingNodes();
+	Dataflow::GeometryCollectionSkeletalMeshNodes();
+	Dataflow::GeometryCollectionTetrahedralNodes();
 }
 
 void IGeometryCollectionNodesPlugin::ShutdownModule()
diff --git a/Engine/Plugins/Experimental/GeometryCollectionPlugin/Source/GeometryCollectionNodes/Private/Dataflow/GeometryCollectionProcessingNodes.cpp b/Engine/Plugins/Experimental/GeometryCollectionPlugin/Source/GeometryCollectionNodes/Private/Dataflow/GeometryCollectionProcessingNodes.cpp
new file mode 100644
--- /dev/null
+++ b/Engine/Plugins/Experimental/GeometryCollectionPlugin/Source/GeometryCollectionNodes/Private/Dataflow/GeometryCollectionProcessingNodes.cpp
@@ -0,0 +1,34 @@
+// Copyright Epic Games, Inc. All Rights Reserved.
+
+#include "Dataflow/GeometryCollectionProcessingNodes.h"
+
+#include "GeometryCollection/ManagedArrayCollection.h"
+#include "GeometryCollection/GeometryCollection.h"
+
+
+namespace Dataflow
+{
+	void GeometryCollectionProcessingNodes()
+	{
+		DATAFLOW_NODE_REGISTER_CREATION_FACTORY(FCloseGeometryOnCollectionDataflowNode);
+	}
+}
+
+void FCloseGeometryOnCollectionDataflowNode::Evaluate(Dataflow::FContext& Context, const FDataflowOutput* Out) const
+{
+	if (Out->IsA<DataType>(&Collection))
+	{
+		DataType InCollection = GetInput(&Collection)->GetValue<DataType>(Context, Collection);
+
+		//
+		// @todo(dataflow) : Implemention that closes the geometry on the collection
+		//
+
+		Out->SetValue<DataType>(InCollection, Context);
+	}
+}
+
+
+
+
+
diff --git a/Engine/Plugins/Experimental/GeometryCollectionPlugin/Source/GeometryCollectionNodes/Private/Dataflow/GeometryCollectionSkeletalMeshNodes.cpp b/Engine/Plugins/Experimental/GeometryCollectionPlugin/Source/GeometryCollectionNodes/Private/Dataflow/GeometryCollectionSkeletalMeshNodes.cpp
new file mode 100644
--- /dev/null
+++ b/Engine/Plugins/Experimental/GeometryCollectionPlugin/Source/GeometryCollectionNodes/Private/Dataflow/GeometryCollectionSkeletalMeshNodes.cpp
@@ -0,0 +1,34 @@
+// Copyright Epic Games, Inc. All Rights Reserved.
+
+#include "Dataflow/GeometryCollectionSkeletalMeshNodes.h"
+
+#include "Engine/SkeletalMesh.h"
+#include "GeometryCollection/GeometryCollectionObject.h"
+#include "GeometryCollection/ManagedArrayCollection.h"
+#include "GeometryCollection/GeometryCollection.h"
+
+
+namespace Dataflow
+{
+	void GeometryCollectionSkeletalMeshNodes()
+	{
+		DATAFLOW_NODE_REGISTER_CREATION_FACTORY(FSkeletalMeshToCollectionDataflowNode);
+
+	}
+}
+
+void FSkeletalMeshToCollectionDataflowNode::Evaluate(Dataflow::FContext& Context, const FDataflowOutput* Out) const
+{
+	if (Out->IsA<DataType>(&Collection))
+	{
+		DataType InCollection = GetInput(&Collection)->GetValue<DataType>(Context, Collection);
+
+		//
+		// @todo(dataflow) : Implemention conversion from skeletal mesh to TManagedArrayCollection
+		//
+
+		Out->SetValue<DataType>(InCollection, Context);
+	}
+}
+
+
diff --git a/Engine/Plugins/Experimental/GeometryCollectionPlugin/Source/GeometryCollectionNodes/Private/Dataflow/GeometryCollectionTetrahedralNodes.cpp b/Engine/Plugins/Experimental/GeometryCollectionPlugin/Source/GeometryCollectionNodes/Private/Dataflow/GeometryCollectionTetrahedralNodes.cpp
new file mode 100644
--- /dev/null
+++ b/Engine/Plugins/Experimental/GeometryCollectionPlugin/Source/GeometryCollectionNodes/Private/Dataflow/GeometryCollectionTetrahedralNodes.cpp
@@ -0,0 +1,30 @@
+// Copyright Epic Games, Inc. All Rights Reserved.
+
+#include "Dataflow/GeometryCollectionTetrahedralNodes.h"
+
+#include "GeometryCollection/ManagedArrayCollection.h"
+
+
+namespace Dataflow
+{
+	void GeometryCollectionTetrahedralNodes()
+	{
+		DATAFLOW_NODE_REGISTER_CREATION_FACTORY(FGenerateTetrahedralCollectionDataflowNodes);
+	}
+}
+
+
+void FGenerateTetrahedralCollectionDataflowNodes::Evaluate(Dataflow::FContext& Context, const FDataflowOutput* Out) const
+{
+	if (Out->IsA<DataType>(&Collection))
+	{
+		DataType InCollection = GetInput(&Collection)->GetValue<DataType>(Context, Collection);
+
+		//
+		// @todo(dataflow) : Implemention tetrahedral generation for a closed surface.
+		//
+
+		Out->SetValue<DataType>(InCollection, Context);
+	}
+}
+
diff --git a/Engine/Plugins/Experimental/GeometryCollectionPlugin/Source/GeometryCollectionNodes/Public/Dataflow/GeometryCollectionProcessingNodes.h b/Engine/Plugins/Experimental/GeometryCollectionPlugin/Source/GeometryCollectionNodes/Public/Dataflow/GeometryCollectionProcessingNodes.h
new file mode 100644
--- /dev/null
+++ b/Engine/Plugins/Experimental/GeometryCollectionPlugin/Source/GeometryCollectionNodes/Public/Dataflow/GeometryCollectionProcessingNodes.h
@@ -0,0 +1,40 @@
+// Copyright Epic Games, Inc. All Rights Reserved.
+
+#pragma once 
+
+#include "CoreMinimal.h"
+#include "Dataflow/DataflowEngine.h"
+#include "GeometryCollection/ManagedArrayCollection.h"
+
+#include "GeometryCollectionProcessingNodes.generated.h"
+
+
+USTRUCT()
+struct FCloseGeometryOnCollectionDataflowNode : public FDataflowNode
+{
+	GENERATED_USTRUCT_BODY()
+	DATAFLOW_NODE_DEFINE_INTERNAL(FCloseGeometryOnCollectionDataflowNode, "CloseGeometryOnCollection", "GeometryCollection", "")
+
+public:
+	typedef FManagedArrayCollection DataType;
+
+	UPROPERTY(meta = (DataflowInput, DataflowOutput, DisplayName = "Collection"))
+		FManagedArrayCollection Collection;
+
+	FCloseGeometryOnCollectionDataflowNode(const Dataflow::FNodeParameters& InParam, FGuid InGuid = FGuid::NewGuid())
+		: FDataflowNode(InParam, InGuid)
+	{
+		RegisterInputConnection(&Collection);
+		RegisterOutputConnection(&Collection);
+	}
+
+	virtual void Evaluate(Dataflow::FContext& Context, const FDataflowOutput* Out) const override;
+};
+
+namespace Dataflow
+{
+	void GeometryCollectionProcessingNodes();
+
+
+}
+
diff --git a/Engine/Plugins/Experimental/GeometryCollectionPlugin/Source/GeometryCollectionNodes/Public/Dataflow/GeometryCollectionSkeletalMeshNodes.h b/Engine/Plugins/Experimental/GeometryCollectionPlugin/Source/GeometryCollectionNodes/Public/Dataflow/GeometryCollectionSkeletalMeshNodes.h
new file mode 100644
--- /dev/null
+++ b/Engine/Plugins/Experimental/GeometryCollectionPlugin/Source/GeometryCollectionNodes/Public/Dataflow/GeometryCollectionSkeletalMeshNodes.h
@@ -0,0 +1,46 @@
+// Copyright Epic Games, Inc. All Rights Reserved.
+
+#pragma once 
+
+#include "CoreMinimal.h"
+#include "Dataflow/DataflowEngine.h"
+#include "GeometryCollection/ManagedArrayCollection.h"
+
+#include "GeometryCollectionSkeletalMeshNodes.generated.h"
+
+class USkeletalMesh;
+
+USTRUCT()
+struct FSkeletalMeshToCollectionDataflowNode : public FDataflowNode
+{
+	GENERATED_USTRUCT_BODY()
+	DATAFLOW_NODE_DEFINE_INTERNAL(FSkeletalMeshToCollectionDataflowNode, "SkeletalMeshToCollection", "GeometryCollection", "")
+
+public:
+	typedef FManagedArrayCollection DataType;
+
+
+	UPROPERTY(meta = (DataflowInput, DisplayName = "SkeletalMesh"))
+	TObjectPtr<const USkeletalMesh> SkeletalMesh = nullptr;
+
+	UPROPERTY(meta = (DataflowInput, DataflowOutput, DisplayName = "Collection"))
+	FManagedArrayCollection Collection;
+
+	FSkeletalMeshToCollectionDataflowNode(const Dataflow::FNodeParameters& InParam, FGuid InGuid = FGuid::NewGuid())
+		: FDataflowNode(InParam, InGuid)
+	{
+		RegisterInputConnection(&SkeletalMesh);
+		RegisterInputConnection(&Collection);
+		RegisterOutputConnection(&Collection);
+	}
+
+	virtual void Evaluate(Dataflow::FContext& Context, const FDataflowOutput* Out) const override;
+};
+
+namespace Dataflow
+{
+	void GeometryCollectionSkeletalMeshNodes();
+
+
+}
+
diff --git a/Engine/Plugins/Experimental/GeometryCollectionPlugin/Source/GeometryCollectionNodes/Public/Dataflow/GeometryCollectionTetrahedralNodes.h b/Engine/Plugins/Experimental/GeometryCollectionPlugin/Source/GeometryCollectionNodes/Public/Dataflow/GeometryCollectionTetrahedralNodes.h
new file mode 100644
--- /dev/null
+++ b/Engine/Plugins/Experimental/GeometryCollectionPlugin/Source/GeometryCollectionNodes/Public/Dataflow/GeometryCollectionTetrahedralNodes.h
@@ -0,0 +1,40 @@
+// Copyright Epic Games, Inc. All Rights Reserved.
+
+#pragma once 
+
+#include "CoreMinimal.h"
+#include "Dataflow/DataflowEngine.h"
+#include "GeometryCollection/ManagedArrayCollection.h"
+
+#include "GeometryCollectionTetrahedralNodes.generated.h"
+
+
+USTRUCT()
+struct FGenerateTetrahedralCollectionDataflowNodes : public FDataflowNode
+{
+	GENERATED_USTRUCT_BODY()
+		DATAFLOW_NODE_DEFINE_INTERNAL(FGenerateTetrahedralCollectionDataflowNodes, "GenerateTetrahedralCollection", "GeometryCollection", "")
+
+public:
+	typedef FManagedArrayCollection DataType;
+
+	UPROPERTY(meta = (DataflowInput, DataflowOutput, DisplayName = "Collection"))
+		FManagedArrayCollection Collection;
+
+	FGenerateTetrahedralCollectionDataflowNodes(const Dataflow::FNodeParameters& InParam, FGuid InGuid = FGuid::NewGuid())
+		: FDataflowNode(InParam, InGuid)
+	{
+		RegisterInputConnection(&Collection);
+		RegisterOutputConnection(&Collection);
+	}
+
+	virtual void Evaluate(Dataflow::FContext& Context, const FDataflowOutput* Out) const override;
+};
+
+namespace Dataflow
+{
+	void GeometryCollectionTetrahedralNodes();
+
+
+}
+
