Dataflow : GetValue /SetValue from the node

diff --git a/Engine/Plugins/Experimental/Dataflow/Source/DataflowNodes/Private/Dataflow/DataflowSkeletalMeshNodes.cpp b/Engine/Plugins/Experimental/Dataflow/Source/DataflowNodes/Private/Dataflow/DataflowSkeletalMeshNodes.cpp
--- a/Engine/Plugins/Experimental/Dataflow/Source/DataflowNodes/Private/Dataflow/DataflowSkeletalMeshNodes.cpp
+++ b/Engine/Plugins/Experimental/Dataflow/Source/DataflowNodes/Private/Dataflow/DataflowSkeletalMeshNodes.cpp
@@ -23,18 +23,18 @@ void FGetSkeletalMeshDataflowNode::Evaluate(Dataflow::FContext& Context, const F
 	typedef TObjectPtr<const USkeletalMesh> DataType;
 	if (Out->IsA<DataType>(&SkeletalMesh))
 	{
-		GetOutput(&SkeletalMesh)->SetValue<DataType>(SkeletalMesh, Context); // prime to avoid ensure
+		SetValue<DataType>(Context, SkeletalMesh, &SkeletalMesh);
 
 		if (SkeletalMesh)
 		{
-			GetOutput(&SkeletalMesh)->SetValue<DataType>(SkeletalMesh, Context);
+			SetValue<DataType>(Context, SkeletalMesh, &SkeletalMesh);
 		}
 		else if (const Dataflow::FEngineContext* EngineContext = Context.AsType<Dataflow::FEngineContext>())
 		{
 			if (const USkeletalMesh* SkeletalMeshFromOwner = Dataflow::Reflection::FindObjectPtrProperty<USkeletalMesh>(
 				EngineContext->Owner, PropertyName))
 			{
-				GetOutput(&SkeletalMesh)->SetValue<DataType>(DataType(SkeletalMeshFromOwner), Context);
+				SetValue<DataType>(Context, DataType(SkeletalMeshFromOwner), &SkeletalMesh);
 			}
 		}
 	}
@@ -45,12 +45,12 @@ void FSkeletalMeshBoneDataflowNode::Evaluate(Dataflow::FContext& Context, const
 	typedef TObjectPtr<const USkeletalMesh> InDataType;
 	if (Out->IsA<int>(&BoneIndexOut))
 	{
-		GetOutput(&BoneIndexOut)->SetValue<int>(INDEX_NONE, Context); // prime to avoid ensure
+		SetValue<int>(Context, INDEX_NONE, &BoneIndexOut);
 
 		if( InDataType InSkeletalMesh = GetInput(&SkeletalMesh)->GetValue<InDataType>(Context, SkeletalMesh) )
 		{
 			int32 Index = InSkeletalMesh->GetRefSkeleton().FindBoneIndex(BoneName);
-			Out->SetValue<int>(Index, Context);
+			SetValue<int>(Context, Index, &BoneIndexOut);
 		}
 
 	}
diff --git a/Engine/Plugins/Experimental/Dataflow/Source/DataflowNodes/Private/Dataflow/DataflowStaticMeshNodes.cpp b/Engine/Plugins/Experimental/Dataflow/Source/DataflowNodes/Private/Dataflow/DataflowStaticMeshNodes.cpp
--- a/Engine/Plugins/Experimental/Dataflow/Source/DataflowNodes/Private/Dataflow/DataflowStaticMeshNodes.cpp
+++ b/Engine/Plugins/Experimental/Dataflow/Source/DataflowNodes/Private/Dataflow/DataflowStaticMeshNodes.cpp
@@ -16,18 +16,19 @@ void FGetStaticMeshDataflowNode::Evaluate(Dataflow::FContext& Context, const FDa
 	typedef TObjectPtr<const UStaticMesh> DataType;
 	if (Out->IsA<DataType>(&StaticMesh))
 	{
-		GetOutput(&StaticMesh)->SetValue<DataType>(StaticMesh, Context); // prime to avoid ensure
+		SetValue<DataType>(Context, StaticMesh, &StaticMesh); // prime to avoid ensure
 
 		if (StaticMesh)
 		{
-			GetOutput(&StaticMesh)->SetValue<DataType>(StaticMesh, Context);
+			SetValue<DataType>(Context, StaticMesh, &StaticMesh);
 		}
 		else if (const Dataflow::FEngineContext* EngineContext = Context.AsType<Dataflow::FEngineContext>())
 		{
 			if (const UStaticMesh* StaticMeshFromOwner = Dataflow::Reflection::FindObjectPtrProperty<UStaticMesh>(
 				EngineContext->Owner, PropertyName))
 			{
-				GetOutput(&StaticMesh)->SetValue<DataType>(DataType(StaticMeshFromOwner), Context);
+				SetValue<DataType>(Context, DataType(StaticMeshFromOwner), &StaticMesh);
+
 			}
 		}
 	}
diff --git a/Engine/Plugins/Experimental/GeometryCollectionPlugin/Source/GeometryCollectionNodes/Private/Dataflow/GeometryCollectionNodes.cpp b/Engine/Plugins/Experimental/GeometryCollectionPlugin/Source/GeometryCollectionNodes/Private/Dataflow/GeometryCollectionNodes.cpp
--- a/Engine/Plugins/Experimental/GeometryCollectionPlugin/Source/GeometryCollectionNodes/Private/Dataflow/GeometryCollectionNodes.cpp
+++ b/Engine/Plugins/Experimental/GeometryCollectionPlugin/Source/GeometryCollectionNodes/Private/Dataflow/GeometryCollectionNodes.cpp
@@ -33,7 +33,7 @@ void FGetCollectionAssetDataflowNode::Evaluate(Dataflow::FContext& Context, cons
 			{
 				if (const TSharedPtr<FGeometryCollection, ESPMode::ThreadSafe> AssetCollection = CollectionAsset->GetGeometryCollection())
 				{
-					GetOutput(&Output)->SetValue<DataType>(DataType(*AssetCollection), Context);
+					SetValue<DataType>(Context, DataType(*AssetCollection), &Output);
 				}
 			}
 		}
@@ -44,7 +44,7 @@ void FExampleCollectionEditDataflowNode::Evaluate(Dataflow::FContext& Context, c
 {
 	if (Out->IsA<DataType>(&Collection))
 	{
-		DataType InCollection = GetInput(&Collection)->GetValue<DataType>(Context, Collection);
+		DataType InCollection = GetValue<DataType>(Context, &Collection);
 
 		if (Active)
 		{
@@ -54,13 +54,13 @@ void FExampleCollectionEditDataflowNode::Evaluate(Dataflow::FContext& Context, c
 				(*Vertex)[i][1] *= Scale;
 			}
 		}
-		Out->SetValue<DataType>(InCollection, Context);
+		SetValue<DataType>(Context, InCollection, &Collection);
 	}
 }
 
 void FSetCollectionAssetDataflowNode::Evaluate(Dataflow::FContext& Context, const FDataflowOutput* Out) const
 {
-	DataType InCollection = GetInput(&Collection)->GetValue<DataType>(Context, Collection);
+	DataType InCollection = GetValue<DataType>(Context, &Collection);
 
 	if (const Dataflow::FEngineContext* EngineContext = Context.AsType<Dataflow::FEngineContext>())
 	{
@@ -77,7 +77,7 @@ void FResetGeometryCollectionDataflowNode::Evaluate(Dataflow::FContext& Context,
 {
 	if (Out->IsA<DataType>(&Collection))
 	{
-		Out->SetValue<DataType>(Collection, Context); // prime to avoid ensure
+		SetValue<DataType>(Context, Collection, &Collection); // prime to avoid ensure
 
 		if (const Dataflow::FEngineContext* EngineContext = Context.AsType<Dataflow::FEngineContext>())
 		{
@@ -132,7 +132,7 @@ void FResetGeometryCollectionDataflowNode::Evaluate(Dataflow::FContext& Context,
 
 				if (const TSharedPtr<FGeometryCollection, ESPMode::ThreadSafe> AssetCollection = GeometryCollectionObject->GetGeometryCollection())
 				{
-					Out->SetValue<DataType>(*AssetCollection, Context);
+					SetValue<DataType>(Context, *AssetCollection, &Collection);
 				}
 			}
 		}
diff --git a/Engine/Plugins/Experimental/GeometryCollectionPlugin/Source/GeometryCollectionNodes/Private/Dataflow/GeometryCollectionProcessingNodes.cpp b/Engine/Plugins/Experimental/GeometryCollectionPlugin/Source/GeometryCollectionNodes/Private/Dataflow/GeometryCollectionProcessingNodes.cpp
--- a/Engine/Plugins/Experimental/GeometryCollectionPlugin/Source/GeometryCollectionNodes/Private/Dataflow/GeometryCollectionProcessingNodes.cpp
+++ b/Engine/Plugins/Experimental/GeometryCollectionPlugin/Source/GeometryCollectionNodes/Private/Dataflow/GeometryCollectionProcessingNodes.cpp
@@ -18,13 +18,13 @@ void FCloseGeometryOnCollectionDataflowNode::Evaluate(Dataflow::FContext& Contex
 {
 	if (Out->IsA<DataType>(&Collection))
 	{
-		DataType InCollection = GetInput(&Collection)->GetValue<DataType>(Context, Collection);
+		DataType InCollection = GetValue<DataType>(Context, &Collection);
 
 		//
 		// @todo(dataflow) : Implemention that closes the geometry on the collection
 		//
 
-		Out->SetValue<DataType>(InCollection, Context);
+		SetValue<DataType>(Context, InCollection, &Collection);
 	}
 }
 
diff --git a/Engine/Plugins/Experimental/GeometryCollectionPlugin/Source/GeometryCollectionNodes/Private/Dataflow/GeometryCollectionSkeletalMeshNodes.cpp b/Engine/Plugins/Experimental/GeometryCollectionPlugin/Source/GeometryCollectionNodes/Private/Dataflow/GeometryCollectionSkeletalMeshNodes.cpp
--- a/Engine/Plugins/Experimental/GeometryCollectionPlugin/Source/GeometryCollectionNodes/Private/Dataflow/GeometryCollectionSkeletalMeshNodes.cpp
+++ b/Engine/Plugins/Experimental/GeometryCollectionPlugin/Source/GeometryCollectionNodes/Private/Dataflow/GeometryCollectionSkeletalMeshNodes.cpp
@@ -21,13 +21,13 @@ void FSkeletalMeshToCollectionDataflowNode::Evaluate(Dataflow::FContext& Context
 {
 	if (Out->IsA<DataType>(&Collection))
 	{
-		DataType InCollection = GetInput(&Collection)->GetValue<DataType>(Context, Collection);
+		DataType InCollection = GetValue<DataType>(Context, &Collection);
 
 		//
 		// @todo(dataflow) : Implemention conversion from skeletal mesh to TManagedArrayCollection
 		//
 
-		Out->SetValue<DataType>(InCollection, Context);
+		SetValue<DataType>(Context, InCollection, &Collection);
 	}
 }
 
diff --git a/Engine/Plugins/Experimental/GeometryCollectionPlugin/Source/GeometryCollectionNodes/Private/Dataflow/GeometryCollectionTetrahedralNodes.cpp b/Engine/Plugins/Experimental/GeometryCollectionPlugin/Source/GeometryCollectionNodes/Private/Dataflow/GeometryCollectionTetrahedralNodes.cpp
--- a/Engine/Plugins/Experimental/GeometryCollectionPlugin/Source/GeometryCollectionNodes/Private/Dataflow/GeometryCollectionTetrahedralNodes.cpp
+++ b/Engine/Plugins/Experimental/GeometryCollectionPlugin/Source/GeometryCollectionNodes/Private/Dataflow/GeometryCollectionTetrahedralNodes.cpp
@@ -18,13 +18,13 @@ void FGenerateTetrahedralCollectionDataflowNodes::Evaluate(Dataflow::FContext& C
 {
 	if (Out->IsA<DataType>(&Collection))
 	{
-		DataType InCollection = GetInput(&Collection)->GetValue<DataType>(Context, Collection);
+		DataType InCollection = GetValue<DataType>(Context, &Collection);
 
 		//
 		// @todo(dataflow) : Implemention tetrahedral generation for a closed surface.
 		//
 
-		Out->SetValue<DataType>(InCollection, Context);
+		SetValue<DataType>(Context, InCollection, &Collection);
 	}
 }
 
diff --git a/Engine/Source/Runtime/Experimental/Dataflow/Core/Private/Dataflow/DataflowInputOutput.cpp b/Engine/Source/Runtime/Experimental/Dataflow/Core/Private/Dataflow/DataflowInputOutput.cpp
--- a/Engine/Source/Runtime/Experimental/Dataflow/Core/Private/Dataflow/DataflowInputOutput.cpp
+++ b/Engine/Source/Runtime/Experimental/Dataflow/Core/Private/Dataflow/DataflowInputOutput.cpp
@@ -50,6 +50,16 @@ TArray< FDataflowConnection* > FDataflowInput::GetConnectedOutputs()
 	return RetList;
 }
 
+const TArray< const FDataflowConnection* > FDataflowInput::GetConnectedOutputs() const
+{
+	TArray<const FDataflowConnection* > RetList;
+	if (const FDataflowConnection* Conn = GetConnection())
+	{
+		RetList.Add(Conn);
+	}
+	return RetList;
+}
+
 void FDataflowInput::Invalidate()
 {
 	OwningNode->InvalidateOutputs();
@@ -71,14 +81,27 @@ FDataflowOutput::FDataflowOutput(const Dataflow::FOutputParameters& Param, FGuid
 const TArray<FDataflowInput*>& FDataflowOutput::GetConnections() const { return Connections; }
 TArray<FDataflowInput*>& FDataflowOutput::GetConnections() { return Connections; }
 
-TArray<FDataflowConnection*> FDataflowOutput::GetConnectedInputs()
+const TArray< const FDataflowConnection*> FDataflowOutput::GetConnectedInputs() const
 {
-	TArray<FDataflowConnection*> RetList;
+	TArray<const FDataflowConnection*> RetList;
 	RetList.Reserve(Connections.Num());
-	for (FDataflowInput* Ptr : Connections) { RetList.Add(Ptr); }
+	for (FDataflowInput* Ptr : Connections) 
+	{ 
+		RetList.Add(Ptr); 
+	}
 	return RetList;
 }
 
+TArray< FDataflowConnection*> FDataflowOutput::GetConnectedInputs()
+{
+	TArray<FDataflowConnection*> RetList;
+	RetList.Reserve(Connections.Num());
+	for (FDataflowInput* Ptr : Connections) 
+	{ 
+		RetList.Add(Ptr); 
+	}
+	return RetList;
+}
 
 bool FDataflowOutput::AddConnection(FDataflowConnection* InOutput)
 {
diff --git a/Engine/Source/Runtime/Experimental/Dataflow/Core/Public/Dataflow/DataflowConnection.h b/Engine/Source/Runtime/Experimental/Dataflow/Core/Public/Dataflow/DataflowConnection.h
--- a/Engine/Source/Runtime/Experimental/Dataflow/Core/Public/Dataflow/DataflowConnection.h
+++ b/Engine/Source/Runtime/Experimental/Dataflow/Core/Public/Dataflow/DataflowConnection.h
@@ -73,7 +73,10 @@ public:
 	virtual bool RemoveConnection(FDataflowConnection* In) { return false; }
 
 	virtual TArray< FDataflowConnection* > GetConnectedInputs() { return TArray<FDataflowConnection* >(); }
+	virtual const TArray< const FDataflowConnection* > GetConnectedInputs() const { return TArray<const FDataflowConnection* >(); }
+
 	virtual TArray< FDataflowConnection* > GetConnectedOutputs() { return TArray<FDataflowConnection* >(); }
+	virtual const TArray< const FDataflowConnection* > GetConnectedOutputs() const { return TArray<const FDataflowConnection* >(); }
 
 	template<class T>
 	bool IsA(const T* InVar) const
@@ -82,6 +85,26 @@ public:
 
 	}
 
+	template<class T> const T& GetValueAsInput(Dataflow::FContext& Context, const T& Default) const
+	{
+		if (GetConnectedOutputs().Num())
+		{
+			ensure(GetConnectedOutputs().Num() == 1);
+			if (const FDataflowConnection* ConnectionBase = (GetConnectedOutputs()[0]))
+			{
+				if (!ConnectionBase->Evaluate(Context))
+				{
+					Context.SetData(ConnectionBase->CacheKey(), new Dataflow::ContextCache<T>(Property, new T(Default)));
+				}
+				if (Context.HasData(ConnectionBase->CacheKey()))
+				{
+					return Context.GetDataReference<T>(ConnectionBase->CacheKey(), Default);
+				}
+			}
+		}
+		return Default;
+	}
+
 	virtual void Invalidate() {};
 	virtual bool Evaluate(Dataflow::FContext& Context) const { return false; };
 
diff --git a/Engine/Source/Runtime/Experimental/Dataflow/Core/Public/Dataflow/DataflowInputOutput.h b/Engine/Source/Runtime/Experimental/Dataflow/Core/Public/Dataflow/DataflowInputOutput.h
--- a/Engine/Source/Runtime/Experimental/Dataflow/Core/Public/Dataflow/DataflowInputOutput.h
+++ b/Engine/Source/Runtime/Experimental/Dataflow/Core/Public/Dataflow/DataflowInputOutput.h
@@ -49,25 +49,16 @@ public:
 
 	FDataflowOutput* GetConnection() { return Connection; }
 	const FDataflowOutput* GetConnection() const { return Connection; }
+
 	virtual TArray< FDataflowConnection* > GetConnectedOutputs() override;
+	virtual const TArray< const FDataflowConnection* > GetConnectedOutputs() const override;
 
-	// @todo(dataflow) : needs testing. 
 	template<class T> const T& GetValue(Dataflow::FContext& Context, const T& Default) const
 	{
-		if (FDataflowConnection* ConnectionBase = ((FDataflowConnection*)Connection) )
-		{
-			if (!ConnectionBase->Evaluate(Context))
-			{
-				Context.SetData(ConnectionBase->CacheKey(), new Dataflow::ContextCache<T>(Property, new T(Default)));
-			}
-			if (Context.HasData(ConnectionBase->CacheKey()))
-			{
-				return Context.GetDataReference<T>(ConnectionBase->CacheKey(), Default);
-			}
-		}
-		return Default;
+		return GetValueAsInput<T>(Context, Default);
 	}
 
+
 	virtual void Invalidate() override;
 };
 
@@ -112,6 +103,7 @@ public:
 	const TArray<FDataflowInput*>& GetConnections() const;
 
 	virtual TArray<FDataflowConnection*> GetConnectedInputs() override;
+	virtual const TArray<const FDataflowConnection*> GetConnectedInputs() const override;
 
 	virtual bool AddConnection(FDataflowConnection* InOutput) override;
 
diff --git a/Engine/Source/Runtime/Experimental/Dataflow/Core/Public/Dataflow/DataflowNode.h b/Engine/Source/Runtime/Experimental/Dataflow/Core/Public/Dataflow/DataflowNode.h
--- a/Engine/Source/Runtime/Experimental/Dataflow/Core/Public/Dataflow/DataflowNode.h
+++ b/Engine/Source/Runtime/Experimental/Dataflow/Core/Public/Dataflow/DataflowNode.h
@@ -76,7 +76,9 @@ struct DATAFLOWCORE_API FDataflowNode
 
 	TArray<Dataflow::FPin> GetPins() const;
 
+
 	void AddInput(FDataflowConnection* InPtr);
+
 	FDataflowInput* FindInput(FName Name);
 	FDataflowInput* FindInput(void* Reference);
 	const FDataflowInput* FindInput(const void* Reference) const;
@@ -92,6 +94,8 @@ struct DATAFLOWCORE_API FDataflowNode
 	TArray< FDataflowConnection* > GetOutputs() const;
 	void ClearOutputs();
 
+
+
 	//
 	//  Struct Support
 	//
@@ -108,22 +112,37 @@ struct DATAFLOWCORE_API FDataflowNode
 	//
 	virtual void Evaluate(Dataflow::FContext& Context, const FDataflowOutput*) const { ensure(false); }
 
-	void InvalidateOutputs();
-
+	/**
+	*   GetValue(...)
+	*
+	*	Get the value of the Reference output, invoking up stream evaluations if not 
+	*   cached in the contexts data store. 
+	* 
+	*   @param Context : The evaluation context that holds the data store.
+	*   @param Value : The value to store in the contexts data store.
+	*   @param Reference : Pointer to a member of this node that corresponds with the output to set.
+	*/
+	template<class T> const T& GetValue(Dataflow::FContext& Context, const T* Reference) const
+	{
+		return GetInput(Reference)->template GetValueAsInput<T>(Context, *Reference);
+	}
 
-	template<class T>
-	void SetValue(Dataflow::FContext& Context, const T& NewValue, const T* Reference) const
+	/** 
+	*   SetValue(...)
+	*
+	*	Set the value of the Reference output.
+	* 
+	*   @param Context : The evaluation context that holds the data store.
+	*   @param Value : The value to store in the contexts data store. 
+	*   @param Reference : Pointer to a member of this node that corresponds with the output to set. 
+	*/
+	template<class T> void SetValue(Dataflow::FContext& Context, const T& Value, const T* Reference) const
 	{
-		if (const FDataflowConnection* Con = FindOutput(Reference))
-		{
-			Context.SetData(Con->RealAddress(), new Dataflow::ContextCache<T>( Con->Property, new T(NewValue)));
-		}
-		else
-		{
-			ensure(false);
-		}
+		GetOutput(Reference)->template SetValue<T>(Value, Context);
 	}
 
+	void InvalidateOutputs();
+
 	bool ValidateConnections();
 
 };
