Dataflow Node :
- implemented 34 Dataflow nodes, Voronoi Fracture, Planar Cut, Point generators, Exploded View, Convex Collision generator, conversion and utilities, math functions

diff --git a/Engine/Plugins/Experimental/GeometryCollectionPlugin/GeometryCollectionPlugin.uplugin b/Engine/Plugins/Experimental/GeometryCollectionPlugin/GeometryCollectionPlugin.uplugin
--- a/Engine/Plugins/Experimental/GeometryCollectionPlugin/GeometryCollectionPlugin.uplugin
+++ b/Engine/Plugins/Experimental/GeometryCollectionPlugin/GeometryCollectionPlugin.uplugin
@@ -48,6 +48,10 @@
 		{
 			"Name": "Dataflow",
 			"Enabled": true
+		},
+		{
+			"Name": "PlanarCut",
+			"Enabled": true
 		}
 	]
 }
diff --git a/Engine/Plugins/Experimental/GeometryCollectionPlugin/Source/GeometryCollectionNodes/GeometryCollectionNodes.Build.cs b/Engine/Plugins/Experimental/GeometryCollectionPlugin/Source/GeometryCollectionNodes/GeometryCollectionNodes.Build.cs
--- a/Engine/Plugins/Experimental/GeometryCollectionPlugin/Source/GeometryCollectionNodes/GeometryCollectionNodes.Build.cs
+++ b/Engine/Plugins/Experimental/GeometryCollectionPlugin/Source/GeometryCollectionNodes/GeometryCollectionNodes.Build.cs
@@ -19,6 +19,9 @@ namespace UnrealBuildTool.Rules
 					"DataflowCore",
 					"DataflowEngine",
 					"GeometryCollectionEngine",
+					"GeometryCore",
+					"Voronoi",
+					"PlanarCut",
 					"Engine"
 				}
 			);
diff --git a/Engine/Plugins/Experimental/GeometryCollectionPlugin/Source/GeometryCollectionNodes/Private/Dataflow/GeometryCollectionNodes.cpp b/Engine/Plugins/Experimental/GeometryCollectionPlugin/Source/GeometryCollectionNodes/Private/Dataflow/GeometryCollectionNodes.cpp
--- a/Engine/Plugins/Experimental/GeometryCollectionPlugin/Source/GeometryCollectionNodes/Private/Dataflow/GeometryCollectionNodes.cpp
+++ b/Engine/Plugins/Experimental/GeometryCollectionPlugin/Source/GeometryCollectionNodes/Private/Dataflow/GeometryCollectionNodes.cpp
@@ -1,6 +1,7 @@
 // Copyright Epic Games, Inc. All Rights Reserved.
 
 #include "Dataflow/GeometryCollectionNodes.h"
+#include "Dataflow/DataflowCore.h"
 
 #include "Engine/StaticMesh.h"
 #include "GeometryCollection/GeometryCollectionObject.h"
@@ -12,6 +13,14 @@
 #include "Templates/SharedPointer.h"
 #include "UObject/UnrealTypePrivate.h"
 
+#include "GeometryCollection/GeometryCollectionAlgo.h"
+#include "GeometryCollection/GeometryCollectionClusteringUtility.h"
+#include "GeometryCollection/GeometryCollectionConvexUtility.h"
+#include "Voronoi/Voronoi.h"
+#include "PlanarCut.h"
+#include "Runtime/Engine/Public/EngineGlobals.h"
+#include "Math/UnrealMathUtility.h"
+
 namespace Dataflow
 {
 	void GeometryCollectionEngineAssetNodes()
@@ -20,6 +29,41 @@ namespace Dataflow
 		DATAFLOW_NODE_REGISTER_CREATION_FACTORY(FExampleCollectionEditDataflowNode);
 		DATAFLOW_NODE_REGISTER_CREATION_FACTORY(FSetCollectionAssetDataflowNode);
 		DATAFLOW_NODE_REGISTER_CREATION_FACTORY(FResetGeometryCollectionDataflowNode);
+
+		DATAFLOW_NODE_REGISTER_CREATION_FACTORY(FPrintStringDataflowNode);
+		DATAFLOW_NODE_REGISTER_CREATION_FACTORY(FLogStringDataflowNode);
+		DATAFLOW_NODE_REGISTER_CREATION_FACTORY(FMakeLiteralStringDataflowNode);
+		DATAFLOW_NODE_REGISTER_CREATION_FACTORY(FBoundingBoxDataflowNode);
+		DATAFLOW_NODE_REGISTER_CREATION_FACTORY(FExpandBoundingBoxDataflowNode);
+		DATAFLOW_NODE_REGISTER_CREATION_FACTORY(FVectorToStringDataflowNode);
+		DATAFLOW_NODE_REGISTER_CREATION_FACTORY(FFloatToStringDataflowNode);
+		DATAFLOW_NODE_REGISTER_CREATION_FACTORY(FMakePointsDataflowNode);
+		DATAFLOW_NODE_REGISTER_CREATION_FACTORY(FMakeBoxDataflowNode);
+		DATAFLOW_NODE_REGISTER_CREATION_FACTORY(FUniformScatterPointsDataflowNode);
+		DATAFLOW_NODE_REGISTER_CREATION_FACTORY(FRadialScatterPointsDataflowNode);
+		DATAFLOW_NODE_REGISTER_CREATION_FACTORY(FMakeLiteralFloatDataflowNode);
+		DATAFLOW_NODE_REGISTER_CREATION_FACTORY(FMakeLiteralIntDataflowNode);
+		DATAFLOW_NODE_REGISTER_CREATION_FACTORY(FMakeLiteralBoolDataflowNode);
+		DATAFLOW_NODE_REGISTER_CREATION_FACTORY(FMakeLiteralVectorDataflowNode);
+		DATAFLOW_NODE_REGISTER_CREATION_FACTORY(FIntToStringDataflowNode);
+		DATAFLOW_NODE_REGISTER_CREATION_FACTORY(FBoolToStringDataflowNode);
+		DATAFLOW_NODE_REGISTER_CREATION_FACTORY(FExpandVectorDataflowNode);
+		DATAFLOW_NODE_REGISTER_CREATION_FACTORY(FIntToFloatDataflowNode);
+		DATAFLOW_NODE_REGISTER_CREATION_FACTORY(FVoronoiFractureDataflowNode);
+		DATAFLOW_NODE_REGISTER_CREATION_FACTORY(FStringAppendDataflowNode);
+		DATAFLOW_NODE_REGISTER_CREATION_FACTORY(FRandomFloatDataflowNode);
+		DATAFLOW_NODE_REGISTER_CREATION_FACTORY(FRandomFloatInRangeDataflowNode);
+		DATAFLOW_NODE_REGISTER_CREATION_FACTORY(FRandomUnitVectorDataflowNode);
+		DATAFLOW_NODE_REGISTER_CREATION_FACTORY(FRandomUnitVectorInConeDataflowNode);
+		DATAFLOW_NODE_REGISTER_CREATION_FACTORY(FRadiansToDegreesDataflowNode);
+		DATAFLOW_NODE_REGISTER_CREATION_FACTORY(FDegreesToRadiansDataflowNode);
+		DATAFLOW_NODE_REGISTER_CREATION_FACTORY(FExplodedViewDataflowNode);
+		DATAFLOW_NODE_REGISTER_CREATION_FACTORY(FCreateNonOverlappingConvexHullsDataflowNode);
+		DATAFLOW_NODE_REGISTER_CREATION_FACTORY(FPlaneCutterDataflowNode);
+		DATAFLOW_NODE_REGISTER_CREATION_FACTORY(FHashStringDataflowNode);
+		DATAFLOW_NODE_REGISTER_CREATION_FACTORY(FHashVectorDataflowNode);
+		DATAFLOW_NODE_REGISTER_CREATION_FACTORY(FFloatToIntDataflowNode);
+		DATAFLOW_NODE_REGISTER_CREATION_FACTORY(FMathConstantsDataflowNode);
 	}
 }
 
@@ -46,7 +90,7 @@ void FExampleCollectionEditDataflowNode::Evaluate(Dataflow::FContext& Context, c
 	{
 		DataType InCollection = GetValue<DataType>(Context, &Collection);
 
-		if (Active)
+		if (bActive)
 		{
 			TManagedArray<FVector3f>* Vertex = InCollection.FindAttribute<FVector3f>("Vertex", "Vertices");
 			for (int i = 0; i < Vertex->Num(); i++)
@@ -139,6 +183,797 @@ void FResetGeometryCollectionDataflowNode::Evaluate(Dataflow::FContext& Context,
 	}
 }
 
+void FPrintStringDataflowNode::Evaluate(Dataflow::FContext& Context, const FDataflowOutput* Out) const
+{
+	FString Value = GetValue<FString>(Context, &String);
+
+	if (PrintToScreen)
+	{
+		GEngine->AddOnScreenDebugMessage(-1, Duration, Color, Value);
+	}
+	if (PrintToLog)
+	{
+		UE_LOG(LogTemp, Warning, TEXT("Text, %s"), *Value);
+	}
+}
+
+void FLogStringDataflowNode::Evaluate(Dataflow::FContext& Context, const FDataflowOutput* Out) const
+{
+	if (PrintToLog)
+	{
+		FString Value = GetValue<FString>(Context, &String);
+		UE_LOG(LogTemp, Warning, TEXT("Text, %s"), *Value);
+	}
+}
+
+void FMakeLiteralStringDataflowNode::Evaluate(Dataflow::FContext& Context, const FDataflowOutput* Out) const
+{
+	if (Out->IsA<FString>(&String))
+	{
+		SetValue<FString>(Context, Value, &String);
+	}
+}
+
+void ComputeBoundingBox(FBoundingBoxDataflowNode::DataType& Collection, FBox& BoundingBox)
+{
+	if (Collection.HasAttribute("Transform", "Transform") &&
+		Collection.HasAttribute("Parent", "Transform") &&
+		Collection.HasAttribute("TransformIndex", "Geometry") &&
+		Collection.HasAttribute("BoundingBox", "Geometry"))
+	{
+		const TManagedArray<FTransform>& Transforms = Collection.GetAttribute<FTransform>("Transform", "Transform");
+		const TManagedArray<int32>& ParentIndices = Collection.GetAttribute<int32>("Parent", "Transform");
+		const TManagedArray<int32>& TransformIndices = Collection.GetAttribute<int32>("TransformIndex", "Geometry");
+		const TManagedArray<FBox>& BoundingBoxes = Collection.GetAttribute<FBox>("BoundingBox", "Geometry");
+
+		TArray<FMatrix> TmpGlobalMatrices;
+		GeometryCollectionAlgo::GlobalMatrices(Transforms, ParentIndices, TmpGlobalMatrices);
+
+		if (TmpGlobalMatrices.Num() > 0)
+		{
+			for (int32 BoxIdx = 0; BoxIdx < BoundingBoxes.Num(); ++BoxIdx)
+			{
+				const int32 TransformIndex = TransformIndices[BoxIdx];
+				BoundingBox += BoundingBoxes[BoxIdx].TransformBy(TmpGlobalMatrices[TransformIndex]);
+			}
+		}
+	}
+}
+
+void FBoundingBoxDataflowNode::Evaluate(Dataflow::FContext& Context, const FDataflowOutput* Out) const
+{
+	if (Out->IsA<FBox>(&BoundingBox))
+	{
+		DataType InCollection = GetValue<DataType>(Context, &Collection);
+
+		FBox BBox(ForceInit);
+
+		ComputeBoundingBox(InCollection, BBox);
+
+		SetValue<FBox>(Context, BBox, &BoundingBox);
+	}
+}
+
+void FExpandBoundingBoxDataflowNode::Evaluate(Dataflow::FContext& Context, const FDataflowOutput* Out) const
+{
+	FBox BBox = GetValue<FBox>(Context, &BoundingBox);
+
+	if (Out->IsA<FVector>(&Min))
+	{
+		SetValue<FVector>(Context, BBox.Min, &Min);
+	}
+	else if (Out->IsA<FVector>(&Max))
+	{
+		SetValue<FVector>(Context, BBox.Max, &Max);
+	}
+	else if (Out->IsA<FVector>(&Center))
+	{
+		SetValue<FVector>(Context, BBox.GetCenter(), &Center);
+	}
+	else if (Out->IsA<FVector>(&HalfExtents))
+	{
+		SetValue<FVector>(Context, BBox.GetExtent(), &HalfExtents);
+	}
+	else if (Out->IsA<float>(&Volume))
+	{
+		SetValue<float>(Context, BBox.GetVolume(), &Volume);
+	}
+}
+
+void FVectorToStringDataflowNode::Evaluate(Dataflow::FContext& Context, const FDataflowOutput* Out) const
+{
+	if (Out->IsA<FString>(&String))
+	{
+		FString Value = GetValue<FVector>(Context, &Vector).ToString();
+		SetValue<FString>(Context, Value, &String);
+	}
+}
+
+void FFloatToStringDataflowNode::Evaluate(Dataflow::FContext& Context, const FDataflowOutput* Out) const
+{
+	if (Out->IsA<FString>(&String))
+	{
+		FString Value = FString::Printf(TEXT("%f"), GetValue<float>(Context, &Float));
+		SetValue<FString>(Context, Value, &String);
+	}
+}
+
+void FMakePointsDataflowNode::Evaluate(Dataflow::FContext& Context, const FDataflowOutput* Out) const
+{
+	if (Out->IsA<TArray<FVector>>(&Points))
+	{
+		SetValue<TArray<FVector>>(Context, Point, &Points);
+	}
+}
+
+void FMakeBoxDataflowNode::Evaluate(Dataflow::FContext& Context, const FDataflowOutput* Out) const
+{
+	if (Out->IsA<FBox>(&Box))
+	{
+		if (DataType == EMakeBoxDataTypeEnum::Dataflow_MakeBox_DataType_MinMax)
+		{
+			FVector MinVal = GetValue<FVector>(Context, &Min);
+			FVector MaxVal = GetValue<FVector>(Context, &Max);
+
+			SetValue<FBox>(Context, FBox(MinVal, MaxVal), &Box);
+		}
+		else if (DataType == EMakeBoxDataTypeEnum::Dataflow_MakeBox_DataType_CenterSize)
+		{
+			FVector CenterVal = GetValue<FVector>(Context, &Center);
+			FVector SizeVal = GetValue<FVector>(Context, &Size);
+
+			SetValue<FBox>(Context, FBox(CenterVal - 0.5 * SizeVal, CenterVal + 0.5 * SizeVal), &Box);
+		}
+	}
+}
+
+void FUniformScatterPointsDataflowNode::Evaluate(Dataflow::FContext& Context, const FDataflowOutput* Out) const
+{
+	if (Out->IsA<TArray<FVector>>(&Points))
+	{
+		FBox BBox = GetValue<FBox>(Context, &BoundingBox);
+		if (BBox.GetVolume() > 0.f)
+		{
+			FRandomStream RandStream(GetValue<float>(Context, &RandomSeed));
+
+			const FVector Extent(BBox.Max - BBox.Min);
+			const int32 NumPoints = RandStream.RandRange(GetValue<int32>(Context, &MinNumberOfPoints), GetValue<int32>(Context, &MaxNumberOfPoints));
+
+			TArray<FVector> PointsArr;
+			PointsArr.Reserve(NumPoints);
+			for (int32 idx = 0; idx < NumPoints; ++idx)
+			{
+				PointsArr.Emplace(BBox.Min + FVector(RandStream.FRand(), RandStream.FRand(), RandStream.FRand()) * Extent);
+			}
+
+			SetValue<TArray<FVector>>(Context, PointsArr, &Points);
+		}
+	}
+}
+
+void FRadialScatterPointsDataflowNode::Evaluate(Dataflow::FContext& Context, const FDataflowOutput* Out) const
+{
+	if (Out->IsA<TArray<FVector>>(&Points))
+	{
+		const FVector::FReal RadialStep = GetValue<float>(Context, &Radius) / GetValue<int32>(Context, &RadialSteps);
+		const FVector::FReal AngularStep = 2 * PI / GetValue<int32>(Context, &AngularSteps);
+
+		FRandomStream RandStream(GetValue<float>(Context, &RandomSeed));
+		FVector UpVector(GetValue<FVector>(Context, &Normal));
+		UpVector.Normalize();
+		FVector BasisX, BasisY;
+		UpVector.FindBestAxisVectors(BasisX, BasisY);
+
+		TArray<FVector> PointsArr;
 
+		FVector::FReal Len = RadialStep * .5;
+		for (int32 ii = 0; ii < GetValue<int32>(Context, &RadialSteps); ++ii, Len += RadialStep)
+		{
+			FVector::FReal Angle = FMath::DegreesToRadians(GetValue<float>(Context, &AngleOffset));
+			for (int32 kk = 0; kk < AngularSteps; ++kk, Angle += AngularStep)
+			{
+				FVector RotatingOffset = Len * (FMath::Cos(Angle) * BasisX + FMath::Sin(Angle) * BasisY);
+				PointsArr.Emplace(GetValue<FVector>(Context, &Center) + RotatingOffset + (RandStream.VRand() * RandStream.FRand() * Variability));
+			}
+		}
+
+		SetValue<TArray<FVector>>(Context, PointsArr, &Points);
+	}
+}
+
+void FMakeLiteralFloatDataflowNode::Evaluate(Dataflow::FContext& Context, const FDataflowOutput* Out) const
+{
+	if (Out->IsA<float>(&Float))
+	{
+		SetValue<float>(Context, Value, &Float);
+	}
+}
+
+void FMakeLiteralIntDataflowNode::Evaluate(Dataflow::FContext& Context, const FDataflowOutput* Out) const
+{
+	if (Out->IsA<int32>(&Int))
+	{
+		SetValue<int32>(Context, Value, &Int);
+	}
+}
+
+void FMakeLiteralBoolDataflowNode::Evaluate(Dataflow::FContext& Context, const FDataflowOutput* Out) const
+{
+	if (Out->IsA<bool>(&Bool))
+	{
+		SetValue<bool>(Context, Value, &Bool);
+	}
+}
+
+void FMakeLiteralVectorDataflowNode::Evaluate(Dataflow::FContext& Context, const FDataflowOutput* Out) const
+{
+	if (Out->IsA<FVector>(&Vector))
+	{
+		SetValue<FVector>(Context, Value, &Vector);
+	}
+}
+
+void FIntToStringDataflowNode::Evaluate(Dataflow::FContext& Context, const FDataflowOutput* Out) const
+{
+	if (Out->IsA<FString>(&String))
+	{
+		FString Value = FString::Printf(TEXT("%d"), GetValue<int32>(Context, &Int));
+		SetValue<FString>(Context, Value, &String);
+	}
+}
+
+void FBoolToStringDataflowNode::Evaluate(Dataflow::FContext& Context, const FDataflowOutput* Out) const
+{
+	if (Out->IsA<FString>(&String))
+	{
+		FString Value = FString::Printf(TEXT("%s"), GetValue<bool>(Context, &Bool) ? TEXT("true") : TEXT("false"));
+		SetValue<FString>(Context, Value, &String);
+	}
+}
+
+void FExpandVectorDataflowNode::Evaluate(Dataflow::FContext& Context, const FDataflowOutput* Out) const
+{
+	FVector VectorVal = GetValue<FVector>(Context, &Vector);
+
+	if (Out->IsA<float>(&X))
+	{
+		SetValue<float>(Context, VectorVal.X, &X);
+	}
+	else if (Out->IsA<float>(&Y))
+	{
+		SetValue<float>(Context, VectorVal.Y, &Y);
+	}
+	else if (Out->IsA<float>(&Z))
+	{
+		SetValue<float>(Context, VectorVal.Z, &Z);
+	}
+}
+
+void FIntToFloatDataflowNode::Evaluate(Dataflow::FContext& Context, const FDataflowOutput* Out) const
+{
+	if (Out->IsA<float>(&Float))
+	{
+		float Value = float(GetValue<int32>(Context, &Int));
+		SetValue<float>(Context, Value, &Float);
+	}
+}
+
+float GetMaxVertexMovement(float Grout, float Amplitude, int OctaveNumber, float Persistence)
+{
+	float MaxDisp = Grout;
+	float AmplitudeScaled = Amplitude;
+	for (int32 OctaveIdx = 0; OctaveIdx < OctaveNumber; OctaveIdx++, AmplitudeScaled *= Persistence)
+	{
+		MaxDisp += FMath::Abs(AmplitudeScaled);
+	}
+	return MaxDisp;
+}
+
+void FVoronoiFractureDataflowNode::Evaluate(Dataflow::FContext& Context, const FDataflowOutput* Out) const
+{
+	if (Out->IsA<FManagedArrayCollection>(&Collection))
+	{
+		FManagedArrayCollection CollectionPtr = GetValue<FManagedArrayCollection>(Context, &Collection);
+		if (FGeometryCollection* GeomCollection = CollectionPtr.NewCopy<FGeometryCollection>())
+		{
+			TArray<FVector> Sites = GetValue<TArray<FVector>>(Context, &Points);
+			if (Sites.Num() > 0)
+			{
+				//
+				// Compute BoundingBox for ManagedArrayIn
+				//
+				FBox BoundingBox(ForceInit);
+
+				if (GeomCollection->HasAttribute("Transform", "Transform") &&
+					GeomCollection->HasAttribute("Parent", "Transform") &&
+					GeomCollection->HasAttribute("TransformIndex", "Geometry") &&
+					GeomCollection->HasAttribute("BoundingBox", "Geometry"))
+				{
+					const TManagedArray<FTransform>& Transforms = GeomCollection->GetAttribute<FTransform>("Transform", "Transform");
+					const TManagedArray<int32>& ParentIndices = GeomCollection->GetAttribute<int32>("Parent", "Transform");
+					const TManagedArray<int32>& TransformIndices = GeomCollection->GetAttribute<int32>("TransformIndex", "Geometry");
+					const TManagedArray<FBox>& BoundingBoxes = GeomCollection->GetAttribute<FBox>("BoundingBox", "Geometry");
+
+					TArray<FMatrix> TmpGlobalMatrices;
+					GeometryCollectionAlgo::GlobalMatrices(Transforms, ParentIndices, TmpGlobalMatrices);
+
+					if (TmpGlobalMatrices.Num() > 0)
+					{
+						for (int32 BoxIdx = 0; BoxIdx < BoundingBoxes.Num(); ++BoxIdx)
+						{
+							const int32 TransformIndex = TransformIndices[BoxIdx];
+							BoundingBox += BoundingBoxes[BoxIdx].TransformBy(TmpGlobalMatrices[TransformIndex]);
+						}
+					}
+
+					//
+					// Compute Voronoi Bounds
+					//
+					FBox VoronoiBounds = BoundingBox;
+					VoronoiBounds += FBox(Sites);
+
+					float GroutVal = GetValue<float>(Context, &Grout);
+					float AmplitudeVal = GetValue<float>(Context, &Amplitude);
+					int32 OctaveNumberVal = GetValue<int32>(Context, &OctaveNumber);
+					float PersistenceVal = GetValue<float>(Context, &Persistence);
+
+					VoronoiBounds = VoronoiBounds.ExpandBy(GetMaxVertexMovement(GroutVal, AmplitudeVal, OctaveNumberVal, PersistenceVal) + KINDA_SMALL_NUMBER);
+
+					//
+					// Voronoi Fracture
+					//
+					FNoiseSettings NoiseSettings;
+					NoiseSettings.Amplitude = AmplitudeVal;
+					NoiseSettings.Frequency = GetValue<float>(Context, &Frequency);
+					NoiseSettings.Octaves = OctaveNumberVal;
+					NoiseSettings.PointSpacing = GetValue<float>(Context, &PointSpacing);
+					NoiseSettings.Lacunarity = GetValue<float>(Context, &Lacunarity);
+					NoiseSettings.Persistence = GetValue<float>(Context, &Persistence);;
+
+					FVoronoiDiagram Voronoi(Sites, VoronoiBounds, .1f);
+
+					FPlanarCells VoronoiPlanarCells = FPlanarCells(Sites, Voronoi);
+					VoronoiPlanarCells.InternalSurfaceMaterials.NoiseSettings = NoiseSettings;
+
+					const TArrayView<const int32>& TransformIndicesArray(TransformIndices.GetConstArray());
+
+					float CollisionSampleSpacingVal = GetValue<float>(Context, &CollisionSampleSpacing);
+					float RandomSeedVal = GetValue<float>(Context, &RandomSeed);
+
+					int ResultGeometryIndex = CutMultipleWithPlanarCells(VoronoiPlanarCells, *GeomCollection, TransformIndicesArray, GroutVal, CollisionSampleSpacingVal, RandomSeedVal, FTransform().Identity);
+
+					//					Out->SetValue<FManagedArrayCollection>(FManagedArrayCollection(*GeomCollection), Context);
+					SetValue<FManagedArrayCollection>(Context, (const FManagedArrayCollection&)(*GeomCollection), &Collection);
+				}
+			}
+		}
+	}
+}
+
+void FStringAppendDataflowNode::Evaluate(Dataflow::FContext& Context, const FDataflowOutput* Out) const
+{
+	if (Out->IsA<FString>(&String))
+	{
+		const FString StringOut = GetValue<FString>(Context, &String1) + GetValue<FString>(Context, &String2);
+		SetValue<FString>(Context, StringOut, &String);
+	}
+}
+
+void FRandomFloatDataflowNode::Evaluate(Dataflow::FContext& Context, const FDataflowOutput* Out) const
+{
+	if (Out->IsA<float>(&Float))
+	{
+		if (Deterministic)
+		{
+			float RandomSeedVal = GetValue<float>(Context, &RandomSeed);
+
+			FRandomStream Stream(RandomSeedVal);
+			SetValue<float>(Context, Stream.FRand(), &Float);
+		}
+		else
+		{
+			SetValue<float>(Context, FMath::FRand(), &Float);
+		}
+	}
+}
+
+void FRandomFloatInRangeDataflowNode::Evaluate(Dataflow::FContext& Context, const FDataflowOutput* Out) const
+{
+	if (Out->IsA<float>(&Float))
+	{
+		float MinVal = GetValue<float>(Context, &Min);
+		float MaxVal = GetValue<float>(Context, &Max);
+
+		if (Deterministic)
+		{
+			float RandomSeedVal = GetValue<float>(Context, &RandomSeed);
+
+			FRandomStream Stream(RandomSeedVal);
+			SetValue<float>(Context, Stream.FRandRange(MinVal, MaxVal), &Float);
+		}
+		else
+		{
+			SetValue<float>(Context, FMath::FRandRange(MinVal, MaxVal), &Float);
+		}
+	}
+}
+
+void FRandomUnitVectorDataflowNode::Evaluate(Dataflow::FContext& Context, const FDataflowOutput* Out) const
+{
+	if (Out->IsA<FVector>(&Vector))
+	{
+		if (Deterministic)
+		{
+			float RandomSeedVal = GetValue<float>(Context, &RandomSeed);
+
+			FRandomStream Stream(RandomSeedVal);
+			SetValue<FVector>(Context, Stream.VRand(), &Vector);
+		}
+		else
+		{
+			SetValue<FVector>(Context, FMath::VRand(), &Vector);
+		}
+	}
+}
+
+void FRandomUnitVectorInConeDataflowNode::Evaluate(Dataflow::FContext& Context, const FDataflowOutput* Out) const
+{
+	if (Out->IsA<FVector>(&Vector))
+	{
+		FVector ConeDirectionVal = GetValue<FVector>(Context, &ConeDirection);
+		float ConeHalfAngleVal = GetValue<float>(Context, &ConeHalfAngle);
+
+		if (Deterministic)
+		{
+			float RandomSeedVal = GetValue<float>(Context, &RandomSeed);
+
+			FRandomStream Stream(RandomSeedVal);
+			SetValue<FVector>(Context, Stream.VRandCone(ConeDirectionVal, ConeHalfAngleVal), &Vector);
+		}
+		else
+		{
+			SetValue<FVector>(Context, FMath::VRandCone(ConeDirectionVal, ConeHalfAngleVal), &Vector);
+		}
+	}
+}
+
+void FRadiansToDegreesDataflowNode::Evaluate(Dataflow::FContext& Context, const FDataflowOutput* Out) const
+{
+	if (Out->IsA<float>(&Degrees))
+	{
+		SetValue<float>(Context, FMath::RadiansToDegrees(GetValue<float>(Context, &Radians)), &Degrees);
+	}
+}
+
+void FDegreesToRadiansDataflowNode::Evaluate(Dataflow::FContext& Context, const FDataflowOutput* Out) const
+{
+	if (Out->IsA<float>(&Radians))
+	{
+		SetValue<float>(Context, FMath::DegreesToRadians(GetValue<float>(Context, &Degrees)), &Radians);
+	}
+}
+
+void AddAdditionalAttributesIfRequired(FGeometryCollection* GeometryCollection)
+{
+	FGeometryCollectionClusteringUtility::UpdateHierarchyLevelOfChildren(GeometryCollection, -1);
+}
+
+bool GetValidGeoCenter(FGeometryCollection* Collection, const TManagedArray<int32>& TransformToGeometryIndex, const TArray<FTransform>& Transforms, const TManagedArray<TSet<int32>>& Children, const TManagedArray<FBox>& BoundingBox, int32 TransformIndex, FVector& OutGeoCenter)
+{
+	if (Collection->IsRigid(TransformIndex))
+	{
+		OutGeoCenter = Transforms[TransformIndex].TransformPosition(BoundingBox[TransformToGeometryIndex[TransformIndex]].GetCenter());
+
+		return true;
+	}
+	else if (Collection->SimulationType[TransformIndex] == FGeometryCollection::ESimulationTypes::FST_None) // ie this is embedded geometry
+	{
+		int32 Parent = Collection->Parent[TransformIndex];
+		int32 ParentGeo = Parent != INDEX_NONE ? TransformToGeometryIndex[Parent] : INDEX_NONE;
+		if (ensureMsgf(ParentGeo != INDEX_NONE, TEXT("Embedded geometry should always have a rigid geometry parent!  Geometry collection may be malformed.")))
+		{
+			OutGeoCenter = Transforms[Collection->Parent[TransformIndex]].TransformPosition(BoundingBox[ParentGeo].GetCenter());
+		}
+		else
+		{
+			return false; // no valid value to return
+		}
+
+		return true;
+	}
+	else
+	{
+		FVector AverageCenter;
+		int32 ValidVectors = 0;
+		for (int32 ChildIndex : Children[TransformIndex])
+		{
+
+			if (GetValidGeoCenter(Collection, TransformToGeometryIndex, Transforms, Children, BoundingBox, ChildIndex, OutGeoCenter))
+			{
+				if (ValidVectors == 0)
+				{
+					AverageCenter = OutGeoCenter;
+				}
+				else
+				{
+					AverageCenter += OutGeoCenter;
+				}
+				++ValidVectors;
+			}
+		}
+
+		if (ValidVectors > 0)
+		{
+			OutGeoCenter = AverageCenter / ValidVectors;
+			return true;
+		}
+	}
+	return false;
+}
+
+void FExplodedViewDataflowNode::Evaluate(Dataflow::FContext& Context, const FDataflowOutput* Out) const
+{
+	if (Out->IsA<FManagedArrayCollection>(&Collection))
+	{
+		FManagedArrayCollection CollectionPtr = GetValue<FManagedArrayCollection>(Context, &Collection);
+		if (FGeometryCollection* GeomCollection = CollectionPtr.NewCopy<FGeometryCollection>())
+		{
+			GeomCollection->AddAttribute<FVector3f>("ExplodedVector", FGeometryCollection::TransformGroup, FManagedArrayCollection::FConstructionParameters(FName()));
+			check(GeomCollection->HasAttribute("ExplodedVector", FGeometryCollection::TransformGroup));
+
+			TManagedArray<FVector3f>& ExplodedVectors = GeomCollection->ModifyAttribute<FVector3f>("ExplodedVector", FGeometryCollection::TransformGroup);
+			const TManagedArray<FTransform>& Transform = GeomCollection->GetAttribute<FTransform>("Transform", FGeometryCollection::TransformGroup);
+			const TManagedArray<int32>& TransformToGeometryIndex = GeomCollection->GetAttribute<int32>("TransformToGeometryIndex", FGeometryCollection::TransformGroup);
+			const TManagedArray<FBox>& BoundingBox = GeomCollection->GetAttribute<FBox>("BoundingBox", FGeometryCollection::GeometryGroup);
+
+			// Make sure we have valid "Level"
+			AddAdditionalAttributesIfRequired(GeomCollection);
+
+			const TManagedArray<int32>& Levels = GeomCollection->GetAttribute<int32>("Level", FTransformCollection::TransformGroup);
+			const TManagedArray<int32>& Parent = GeomCollection->GetAttribute<int32>("Parent", FTransformCollection::TransformGroup);
+			const TManagedArray<TSet<int32>>& Children = GeomCollection->GetAttribute<TSet<int32>>("Children", FGeometryCollection::TransformGroup);
+
+			int32 ViewFractureLevel = -1;
+			int32 MaxFractureLevel = ViewFractureLevel;
+			for (int32 Idx = 0, ni = Transform.Num(); Idx < ni; ++Idx)
+			{
+				if (Levels[Idx] > MaxFractureLevel)
+					MaxFractureLevel = Levels[Idx];
+			}
+
+			TArray<FTransform> Transforms;
+			GeometryCollectionAlgo::GlobalMatrices(Transform, GeomCollection->Parent, Transforms);
+
+			TArray<FVector> TransformedCenters;
+			TransformedCenters.SetNumUninitialized(Transforms.Num());
+
+			int32 TransformsCount = 0;
+
+			FVector Center(ForceInitToZero);
+			for (int32 Idx = 0, ni = Transform.Num(); Idx < ni; ++Idx)
+			{
+				ExplodedVectors[Idx] = FVector3f::ZeroVector;
+				FVector GeoCenter;
+
+				if (GetValidGeoCenter(GeomCollection, TransformToGeometryIndex, Transforms, Children, BoundingBox, Idx, GeoCenter))
+				{
+					TransformedCenters[Idx] = GeoCenter;
+					if ((ViewFractureLevel < 0) || Levels[Idx] == ViewFractureLevel)
+					{
+						Center += TransformedCenters[Idx];
+						++TransformsCount;
+					}
+				}
+			}
+
+			Center /= TransformsCount;
+
+			for (int Level = 1; Level <= MaxFractureLevel; Level++)
+			{
+				for (int32 Idx = 0, ni = Transforms.Num(); Idx < ni; ++Idx)
+				{
+					if ((ViewFractureLevel < 0) || Levels[Idx] == ViewFractureLevel)
+					{
+						FVector ScaleVal = GetValue<FVector>(Context, &Scale);
+						float UniformScaleVal = GetValue<float>(Context, &UniformScale);
+
+						FVector ScaleVec = ScaleVal * UniformScaleVal;
+						ExplodedVectors[Idx] = (FVector3f)(TransformedCenters[Idx] - Center) * (FVector3f)ScaleVec;
+					}
+					else
+					{
+						if (Parent[Idx] > -1)
+						{
+							ExplodedVectors[Idx] = ExplodedVectors[Parent[Idx]];
+						}
+					}
+				}
+			}
+
+			SetValue<FManagedArrayCollection>(Context, (const FManagedArrayCollection&)(*GeomCollection), &Collection);
+		}
+	}
+}
+
+void FCreateNonOverlappingConvexHullsDataflowNode::Evaluate(Dataflow::FContext& Context, const FDataflowOutput* Out) const
+{
+	if (Out->IsA<FManagedArrayCollection>(&Collection))
+	{
+		FManagedArrayCollection CollectionPtr = GetValue<FManagedArrayCollection>(Context, &Collection);
+		if (FGeometryCollection* GeomCollection = CollectionPtr.NewCopy<FGeometryCollection>())
+		{
+			float CanRemoveFractionVal = GetValue<float>(Context, &CanRemoveFraction);
+			float CanExceedFractionVal = GetValue<float>(Context, &CanExceedFraction);
+			float SimplificationDistanceThresholdVal = GetValue<float>(Context, &SimplificationDistanceThreshold);
+
+			FGeometryCollectionConvexUtility::FGeometryCollectionConvexData ConvexData = FGeometryCollectionConvexUtility::CreateNonOverlappingConvexHullData(GeomCollection, CanRemoveFractionVal, SimplificationDistanceThresholdVal, CanExceedFractionVal);
+
+			SetValue<FManagedArrayCollection>(Context, (const FManagedArrayCollection&)(*GeomCollection), &Collection);
+		}
+	}
+}
+
+void FPlaneCutterDataflowNode::Evaluate(Dataflow::FContext& Context, const FDataflowOutput* Out) const
+{
+	if (Out->IsA<FManagedArrayCollection>(&Collection))
+	{
+		FManagedArrayCollection CollectionPtr = GetValue<FManagedArrayCollection>(Context, &Collection);
+		if (FGeometryCollection* GeomCollection = CollectionPtr.NewCopy<FGeometryCollection>())
+		{
+			TArray<FPlane> CuttingPlanes;
+			TArray<FTransform> CuttingPlaneTransforms;
+
+			float RandomSeedVal = GetValue<float>(Context, &RandomSeed);
+			FRandomStream RandStream(RandomSeedVal);
+
+			FBox Bounds = GetValue<FBox>(Context, &BoundingBox);
+			const FVector Extent(Bounds.Max - Bounds.Min);
+
+			CuttingPlaneTransforms.Reserve(CuttingPlaneTransforms.Num() + NumPlanes);
+			for (int32 ii = 0; ii < NumPlanes; ++ii)
+			{
+				FVector Position(Bounds.Min + FVector(RandStream.FRand(), RandStream.FRand(), RandStream.FRand()) * Extent);
+				CuttingPlaneTransforms.Emplace(FTransform(FRotator(RandStream.FRand() * 360.0f, RandStream.FRand() * 360.0f, 0.0f), Position));
+			}
+
+			for (const FTransform& Transform : CuttingPlaneTransforms)
+			{
+				CuttingPlanes.Add(FPlane(Transform.GetLocation(), Transform.GetUnitAxis(EAxis::Z)));
+			}
+
+			FInternalSurfaceMaterials InternalSurfaceMaterials;
+			FNoiseSettings NoiseSettings;
+
+			float AmplitudeVal = GetValue<float>(Context, &Amplitude);
+			if (AmplitudeVal > 0.0f)
+			{
+				NoiseSettings.Amplitude = AmplitudeVal;
+				NoiseSettings.Frequency = GetValue<float>(Context, &Frequency);
+				NoiseSettings.Lacunarity = GetValue<float>(Context, &Lacunarity);
+				NoiseSettings.Persistence = GetValue<float>(Context, &Persistence);
+				NoiseSettings.Octaves = GetValue<int32>(Context, &OctaveNumber);
+				NoiseSettings.PointSpacing = GetValue<float>(Context, &PointSpacing);
+
+				InternalSurfaceMaterials.NoiseSettings = NoiseSettings;
+			}
+
+			if (GeomCollection->HasAttribute("TransformIndex", "Geometry"))
+			{
+				const TManagedArray<int32>& TransformIndices = GeomCollection->GetAttribute<int32>("TransformIndex", "Geometry");
+				const TArrayView<const int32>& TransformIndicesArray(TransformIndices.GetConstArray());
+
+				float CollisionSampleSpacingVal = GetValue<float>(Context, &CollisionSampleSpacing);
+				float GroutVal = GetValue<float>(Context, &Grout);
+
+				int ResultGeometryIndex = CutMultipleWithMultiplePlanes(CuttingPlanes, InternalSurfaceMaterials, *GeomCollection, TransformIndicesArray, GroutVal, CollisionSampleSpacingVal, RandomSeedVal, FTransform().Identity);
+
+				SetValue<FManagedArrayCollection>(Context, (const FManagedArrayCollection&)(*GeomCollection), &Collection);
+			}
+		}
+	}
+}
+
+void FHashStringDataflowNode::Evaluate(Dataflow::FContext& Context, const FDataflowOutput* Out) const
+{
+	if (Out->IsA<int32>(&Hash))
+	{
+		SetValue<int32>(Context, GetTypeHash(GetValue<FString>(Context, &String)), &Hash);
+	}
+}
+
+void FHashVectorDataflowNode::Evaluate(Dataflow::FContext& Context, const FDataflowOutput* Out) const
+{
+	if (Out->IsA<int32>(&Hash))
+	{
+		SetValue<int32>(Context, GetTypeHash(GetValue<FVector>(Context, &Vector)), &Hash);
+	}
+}
+
+void FFloatToIntDataflowNode::Evaluate(Dataflow::FContext& Context, const FDataflowOutput* Out) const
+{
+	if (Out->IsA<int32>(&Int))
+	{
+		float FloatVal = GetValue<float>(Context, &Float);
+		if (Function == EFloatToIntFunctionEnum::Dataflow_FloatToInt_Function_Floor)
+		{
+			SetValue<int32>(Context, FMath::FloorToInt32(FloatVal), &Int);
+		}
+		else if (Function == EFloatToIntFunctionEnum::Dataflow_FloatToInt_Function_Ceil)
+		{
+			SetValue<int32>(Context, FMath::CeilToInt32(FloatVal), &Int);
+		}
+		else if (Function == EFloatToIntFunctionEnum::Dataflow_FloatToInt_Function_Round)
+		{
+			SetValue<int32>(Context, FMath::RoundToInt32(FloatVal), &Int);
+		}
+		else if (Function == EFloatToIntFunctionEnum::Dataflow_FloatToInt_Function_Truncate)
+		{
+			SetValue<int32>(Context, int32(FMath::TruncToFloat(FloatVal)), &Int);
+		}
+	}
+}
+
+void FMathConstantsDataflowNode::Evaluate(Dataflow::FContext& Context, const FDataflowOutput* Out) const
+{
+	if (Out->IsA<float>(&Float))
+	{
+		if (Constant == EMathConstantsEnum::Dataflow_MathConstants_Pi)
+		{
+			SetValue<float>(Context, FMathf::Pi, &Float);
+		}
+		else if (Constant == EMathConstantsEnum::Dataflow_MathConstants_HalfPi)
+		{
+			SetValue<float>(Context, FMathf::HalfPi, &Float);
+		}
+		else if (Constant == EMathConstantsEnum::Dataflow_MathConstants_TwoPi)
+		{
+			SetValue<float>(Context, FMathf::TwoPi, &Float);
+		}
+		else if (Constant == EMathConstantsEnum::Dataflow_MathConstants_FourPi)
+		{
+			SetValue<float>(Context, FMathf::FourPi, &Float);
+		}
+		else if (Constant == EMathConstantsEnum::Dataflow_MathConstants_InvPi)
+		{
+			SetValue<float>(Context, FMathf::InvPi, &Float);
+		}
+		else if (Constant == EMathConstantsEnum::Dataflow_MathConstants_InvTwoPi)
+		{
+			SetValue<float>(Context, FMathf::InvTwoPi, &Float);
+		}
+		else if (Constant == EMathConstantsEnum::Dataflow_MathConstants_Sqrt2)
+		{
+			SetValue<float>(Context, FMathf::Sqrt2, &Float);
+		}
+		else if (Constant == EMathConstantsEnum::Dataflow_MathConstants_InvSqrt2)
+		{
+			SetValue<float>(Context, FMathf::InvSqrt2, &Float);
+		}
+		else if (Constant == EMathConstantsEnum::Dataflow_MathConstants_Sqrt3)
+		{
+			SetValue<float>(Context, FMathf::Sqrt3, &Float);
+		}
+		else if (Constant == EMathConstantsEnum::Dataflow_MathConstants_InvSqrt3)
+		{
+			SetValue<float>(Context, FMathf::InvSqrt3, &Float);
+		}
+		else if (Constant == EMathConstantsEnum::Dataflow_FloatToInt_Function_E)
+		{
+			SetValue<float>(Context, 2.71828182845904523536f, &Float);
+		}
+		else if (Constant == EMathConstantsEnum::Dataflow_FloatToInt_Function_Gamma)
+		{
+			SetValue<float>(Context, 0.577215664901532860606512090082f, &Float);
+		}
+		else if (Constant == EMathConstantsEnum::Dataflow_FloatToInt_Function_GoldenRatio)
+		{
+			SetValue<float>(Context, 1.618033988749894f, &Float);
+		}
+		else if (Constant == EMathConstantsEnum::Dataflow_FloatToInt_Function_ZeroTolerance)
+		{
+			SetValue<float>(Context, FMathf::ZeroTolerance, &Float);
+		}
+	}
+}
 
 
diff --git a/Engine/Plugins/Experimental/GeometryCollectionPlugin/Source/GeometryCollectionNodes/Public/Dataflow/GeometryCollectionNodes.h b/Engine/Plugins/Experimental/GeometryCollectionPlugin/Source/GeometryCollectionNodes/Public/Dataflow/GeometryCollectionNodes.h
--- a/Engine/Plugins/Experimental/GeometryCollectionPlugin/Source/GeometryCollectionNodes/Public/Dataflow/GeometryCollectionNodes.h
+++ b/Engine/Plugins/Experimental/GeometryCollectionPlugin/Source/GeometryCollectionNodes/Public/Dataflow/GeometryCollectionNodes.h
@@ -6,7 +6,6 @@
 #include "Dataflow/DataflowEngine.h"
 #include "GeometryCollection/ManagedArrayCollection.h"
 
-
 #include "GeometryCollectionNodes.generated.h"
 
 
@@ -29,8 +28,12 @@ public:
 	}
 
 	virtual void Evaluate(Dataflow::FContext& Context, const FDataflowOutput* Out) const override;
+
 };
 
+/**
+	Description for this node
+*/
 USTRUCT()
 struct FExampleCollectionEditDataflowNode : public FDataflowNode
 {
@@ -40,10 +43,8 @@ struct FExampleCollectionEditDataflowNode : public FDataflowNode
 public:
 	typedef FManagedArrayCollection DataType;
 
-	UPROPERTY(EditAnywhere, Category = "Dataflow");
-	bool Active = true;
-
-	UPROPERTY(EditAnywhere, Category = "Dataflow");
+	/** Description for this parameter */
+	UPROPERTY(EditAnywhere, Category = "Dataflow", meta = (EditCondition = "false", EditConditionHides));
 	float Scale = 1.0;
 
 	UPROPERTY(meta = (DataflowInput, DataflowOutput, DisplayName = "Collection"))
@@ -100,11 +101,1105 @@ public:
 		RegisterOutputConnection(&Collection);
 	}
 
+	virtual void Evaluate(Dataflow::FContext& Context, const FDataflowOutput* Out) const override;
+
+};
+
+
+USTRUCT()
+struct FPrintStringDataflowNode : public FDataflowNode
+{
+	GENERATED_USTRUCT_BODY()
+	DATAFLOW_NODE_DEFINE_INTERNAL(FPrintStringDataflowNode, "PrintString", "Development", "")
+
+public:
+	UPROPERTY(EditAnywhere, Category = "Print");
+	bool PrintToScreen = true;
+
+	UPROPERTY(EditAnywhere, Category = "Print");
+	bool PrintToLog = true;
+
+	UPROPERTY(EditAnywhere, Category = "Print");
+	FColor Color = FColor::White;
+
+	UPROPERTY(EditAnywhere, Category = "Print");
+	float Duration = 2.f;
+
+	UPROPERTY(EditAnywhere, Category = "Print", meta = (DataflowInput, DisplayName = "String"));
+	FString String = FString("");
+
+	FPrintStringDataflowNode(const Dataflow::FNodeParameters& InParam, FGuid InGuid = FGuid::NewGuid())
+		: FDataflowNode(InParam, InGuid)
+	{
+		RegisterInputConnection(&String);
+	}
+
+	virtual void Evaluate(Dataflow::FContext& Context, const FDataflowOutput* Out) const override;
+
+};
+
+
+USTRUCT()
+struct FLogStringDataflowNode : public FDataflowNode
+{
+	GENERATED_USTRUCT_BODY()
+	DATAFLOW_NODE_DEFINE_INTERNAL(FLogStringDataflowNode, "LogString", "Development", "")
+
+public:
+	UPROPERTY(EditAnywhere, Category = "Print");
+	bool PrintToLog = true;
+
+	UPROPERTY(EditAnywhere, Category = "Print", meta = (DataflowInput, DisplayName = "String"));
+	FString String = FString("");
+
+	FLogStringDataflowNode(const Dataflow::FNodeParameters& InParam, FGuid InGuid = FGuid::NewGuid())
+		: FDataflowNode(InParam, InGuid)
+	{
+		RegisterInputConnection(&String);
+	}
+
+	virtual void Evaluate(Dataflow::FContext& Context, const FDataflowOutput* Out) const override;
+
+};
+
+
+USTRUCT()
+struct FMakeLiteralStringDataflowNode : public FDataflowNode
+{
+	GENERATED_USTRUCT_BODY()
+	DATAFLOW_NODE_DEFINE_INTERNAL(FMakeLiteralStringDataflowNode, "MakeLiteralString", "Utilities|String", "")
+
+public:
+	UPROPERTY(EditAnywhere, Category = "String");
+	FString Value = FString("");
+
+	UPROPERTY(meta = (DataflowOutput, DisplayName = "String"))
+	FString String;
+
+	FMakeLiteralStringDataflowNode(const Dataflow::FNodeParameters& InParam, FGuid InGuid = FGuid::NewGuid())
+		: FDataflowNode(InParam, InGuid)
+	{
+		RegisterOutputConnection(&String);
+	}
+
+	virtual void Evaluate(Dataflow::FContext& Context, const FDataflowOutput* Out) const override;
+
+};
+
+USTRUCT()
+struct FBoundingBoxDataflowNode : public FDataflowNode
+{
+	GENERATED_USTRUCT_BODY()
+	DATAFLOW_NODE_DEFINE_INTERNAL(FBoundingBoxDataflowNode, "BoundingBox", "Utilities|Box", "")
+
+public:
+	typedef FManagedArrayCollection DataType;
+
+	UPROPERTY(meta = (DataflowInput, DisplayName = "Collection"))
+	FManagedArrayCollection Collection;
+
+	UPROPERTY(meta = (DataflowOutput, DisplayName = "BoundingBox"))
+	FBox BoundingBox = FBox(ForceInit);
+
+	FBoundingBoxDataflowNode(const Dataflow::FNodeParameters& InParam, FGuid InGuid = FGuid::NewGuid())
+		: FDataflowNode(InParam, InGuid)
+	{
+		RegisterInputConnection(&Collection);
+		RegisterOutputConnection(&BoundingBox);
+	}
+
+	virtual void Evaluate(Dataflow::FContext& Context, const FDataflowOutput* Out) const override;
+
+};
+
+USTRUCT()
+struct FExpandBoundingBoxDataflowNode : public FDataflowNode
+{
+	GENERATED_USTRUCT_BODY()
+	DATAFLOW_NODE_DEFINE_INTERNAL(FExpandBoundingBoxDataflowNode, "ExpandBoundingBox", "Utilities|Box", "")
+
+public:
+	UPROPERTY(meta = (DataflowInput, DisplayName = "BoundingBox"))
+	FBox BoundingBox = FBox(ForceInit);;
+
+	UPROPERTY(meta = (DataflowOutput, DisplayName = "Min"))
+	FVector Min = FVector(0.0);
+
+	UPROPERTY(meta = (DataflowOutput, DisplayName = "Max"))
+	FVector Max = FVector(0.0);
+
+	UPROPERTY(meta = (DataflowOutput, DisplayName = "Center"))
+	FVector Center = FVector(0.0);
+
+	UPROPERTY(meta = (DataflowOutput, DisplayName = "HalfExtents"))
+	FVector HalfExtents = FVector(0.0);
+
+	UPROPERTY(meta = (DataflowOutput, DisplayName = "Volume"))
+	float Volume = 0.0;
+
+	FExpandBoundingBoxDataflowNode(const Dataflow::FNodeParameters& InParam, FGuid InGuid = FGuid::NewGuid())
+		: FDataflowNode(InParam, InGuid)
+	{
+		RegisterInputConnection(&BoundingBox);
+		RegisterOutputConnection(&Min);
+		RegisterOutputConnection(&Max);
+		RegisterOutputConnection(&Center);
+		RegisterOutputConnection(&HalfExtents);
+		RegisterOutputConnection(&Volume);
+	}
+
+	virtual void Evaluate(Dataflow::FContext& Context, const FDataflowOutput* Out) const override;
+
+};
+
+USTRUCT()
+struct FVectorToStringDataflowNode : public FDataflowNode
+{
+	GENERATED_USTRUCT_BODY()
+	DATAFLOW_NODE_DEFINE_INTERNAL(FVectorToStringDataflowNode, "VectorToString", "Utilities|String", "")
+
+public:
+	UPROPERTY(EditAnywhere, Category = "Vector", meta = (DataflowInput, DisplayName = "Vector"))
+	FVector Vector = FVector(0.0);
+
+	UPROPERTY(meta = (DataflowOutput, DisplayName = "String"))
+	FString String = FString("");
+
+	FVectorToStringDataflowNode(const Dataflow::FNodeParameters& InParam, FGuid InGuid = FGuid::NewGuid())
+		: FDataflowNode(InParam, InGuid)
+	{
+		RegisterInputConnection(&Vector);
+		RegisterOutputConnection(&String);
+	}
+
+	virtual void Evaluate(Dataflow::FContext& Context, const FDataflowOutput* Out) const override;
+
+};
+
+USTRUCT()
+struct FFloatToStringDataflowNode : public FDataflowNode
+{
+	GENERATED_USTRUCT_BODY()
+	DATAFLOW_NODE_DEFINE_INTERNAL(FFloatToStringDataflowNode, "FloatToString", "Utilities|String", "")
+
+public:
+	UPROPERTY(EditAnywhere, Category = "Float", meta = (DataflowInput, DisplayName = "Float"))
+	float Float = 0.f;
+
+	UPROPERTY(meta = (DataflowOutput, DisplayName = "String"))
+	FString String = FString("");
+
+	FFloatToStringDataflowNode(const Dataflow::FNodeParameters& InParam, FGuid InGuid = FGuid::NewGuid())
+		: FDataflowNode(InParam, InGuid)
+	{
+		RegisterInputConnection(&Float);
+		RegisterOutputConnection(&String);
+	}
+
+	virtual void Evaluate(Dataflow::FContext& Context, const FDataflowOutput* Out) const override;
+
+};
+
+USTRUCT()
+struct FMakePointsDataflowNode : public FDataflowNode
+{
+	GENERATED_USTRUCT_BODY()
+	DATAFLOW_NODE_DEFINE_INTERNAL(FMakePointsDataflowNode, "MakePoints", "Generators|Point", "")
+
+public:
+	UPROPERTY(EditAnywhere, Category = "Point")
+	TArray<FVector> Point;
+
+	UPROPERTY(meta = (DataflowOutput, DisplayName = "Points"))
+	TArray<FVector> Points;
+
+	FMakePointsDataflowNode(const Dataflow::FNodeParameters& InParam, FGuid InGuid = FGuid::NewGuid())
+		: FDataflowNode(InParam, InGuid)
+	{
+		RegisterOutputConnection(&Points);
+	}
+
+	virtual void Evaluate(Dataflow::FContext& Context, const FDataflowOutput* Out) const override;
+
+};
+
+UENUM(BlueprintType)
+enum class EMakeBoxDataTypeEnum : uint8
+{
+	Dataflow_MakeBox_DataType_MinMax UMETA(DisplayName = "Min/Max"),
+	Dataflow_MakeBox_DataType_CenterSize UMETA(DisplayName = "Center/Size"),
+	//~~~
+	//256th entry
+	Dataflow_Max                UMETA(Hidden)
+};
+
+USTRUCT()
+struct FMakeBoxDataflowNode : public FDataflowNode
+{
+	GENERATED_USTRUCT_BODY()
+	DATAFLOW_NODE_DEFINE_INTERNAL(FMakeBoxDataflowNode, "MakeBox", "Generators|Box", "")
+
+public:
+	UPROPERTY(EditAnywhere, Category = "Box", meta = (DisplayName = "Input Data Type"));
+	EMakeBoxDataTypeEnum DataType = EMakeBoxDataTypeEnum::Dataflow_MakeBox_DataType_MinMax;
+
+	UPROPERTY(EditAnywhere, Category = "Box", meta = (DataflowInput, DisplayName = "Min", EditCondition = "DataType == EMakeBoxDataTypeEnum::Dataflow_MakeBox_DataType_MinMax", EditConditionHides));
+	FVector Min = FVector(0.0);
+
+	UPROPERTY(EditAnywhere, Category = "Box", meta = (DataflowInput, DisplayName = "Max", EditCondition = "DataType == EMakeBoxDataTypeEnum::Dataflow_MakeBox_DataType_MinMax", EditConditionHides));
+	FVector Max = FVector(0.0);
+
+	UPROPERTY(EditAnywhere, Category = "Box", meta = (DataflowInput, DisplayName = "Center", EditCondition = "DataType == EMakeBoxDataTypeEnum::Dataflow_MakeBox_DataType_CenterSize", EditConditionHides));
+	FVector Center = FVector(0.0);
+
+	UPROPERTY(EditAnywhere, Category = "Box", meta = (DataflowInput, DisplayName = "Size", EditCondition = "DataType == EMakeBoxDataTypeEnum::Dataflow_MakeBox_DataType_CenterSize", EditConditionHides));
+	FVector Size = FVector(0.0);
+
+	UPROPERTY(meta = (DataflowOutput, DisplayName = "Box"));
+	FBox Box = FBox(ForceInit);
+
+	FMakeBoxDataflowNode(const Dataflow::FNodeParameters& InParam, FGuid InGuid = FGuid::NewGuid())
+		: FDataflowNode(InParam, InGuid)
+	{
+		RegisterInputConnection(&Min);
+		RegisterInputConnection(&Max);
+		RegisterInputConnection(&Center);
+		RegisterInputConnection(&Size);
+		RegisterOutputConnection(&Box);
+	}
+
+	virtual void Evaluate(Dataflow::FContext& Context, const FDataflowOutput* Out) const override;
+
+};
+
+USTRUCT()
+struct FUniformScatterPointsDataflowNode : public FDataflowNode
+{
+	GENERATED_USTRUCT_BODY()
+		DATAFLOW_NODE_DEFINE_INTERNAL(FUniformScatterPointsDataflowNode, "UniformScatterPoints", "Generators|Point", "")
+
+public:
+	UPROPERTY(EditAnywhere, Category = "Scatter", meta = (DataflowInput, DisplayName = "Min Number Of Points"));
+	int32 MinNumberOfPoints = 20;
+
+	UPROPERTY(EditAnywhere, Category = "Scatter", meta = (DataflowInput, DisplayName = "Max Number Of Points"));
+	int32 MaxNumberOfPoints = 20;
+
+	UPROPERTY(EditAnywhere, Category = "Scatter", meta = (DataflowInput, DisplayName = "Random Seed"));
+	float RandomSeed = -1.f;
+
+	UPROPERTY(meta = (DataflowInput, DisplayName = "BoundingBox"))
+		FBox BoundingBox = FBox(ForceInit);
+
+	UPROPERTY(meta = (DataflowOutput, DisplayName = "Points"))
+		TArray<FVector> Points;
+
+	FUniformScatterPointsDataflowNode(const Dataflow::FNodeParameters& InParam, FGuid InGuid = FGuid::NewGuid())
+		: FDataflowNode(InParam, InGuid)
+	{
+		RegisterInputConnection(&MinNumberOfPoints);
+		RegisterInputConnection(&MaxNumberOfPoints);
+		RegisterInputConnection(&RandomSeed);
+		RegisterInputConnection(&BoundingBox);
+		RegisterOutputConnection(&Points);
+	}
+
+	virtual void Evaluate(Dataflow::FContext& Context, const FDataflowOutput* Out) const override;
+
+};
+
+USTRUCT()
+struct FRadialScatterPointsDataflowNode : public FDataflowNode
+{
+	GENERATED_USTRUCT_BODY()
+		DATAFLOW_NODE_DEFINE_INTERNAL(FRadialScatterPointsDataflowNode, "RadialScatterPoints", "Generators|Point", "")
+
+public:
+	UPROPERTY(EditAnywhere, Category = "Scatter", meta = (DataflowInput, DisplayName = "Center"));
+	FVector Center = FVector(0.0);
+
+	UPROPERTY(EditAnywhere, Category = "Scatter", meta = (DataflowInput, DisplayName = "Normal"));
+	FVector Normal = FVector(0.0, 0.0, 1.0);
+
+	UPROPERTY(EditAnywhere, Category = "Scatter", meta = (DataflowInput, DisplayName = "Radius", UIMin = 0.01f));
+	float Radius = 50.f;
+
+	UPROPERTY(EditAnywhere, Category = "Scatter", meta = (DataflowInput, DisplayName = "Angular Steps", UIMin = 1, UIMax = 50));
+	int32 AngularSteps = 5;
+
+	UPROPERTY(EditAnywhere, Category = "Scatter", meta = (DataflowInput, DisplayName = "Radial Steps", UIMin = 1, UIMax = 50));
+	int32 RadialSteps = 5;
+
+	UPROPERTY(EditAnywhere, Category = "Scatter", meta = (DataflowInput, DisplayName = "Angle Offset"));
+	float AngleOffset = 0.f;
+
+	UPROPERTY(EditAnywhere, Category = "Scatter", meta = (DataflowInput, DisplayName = "Variability", UIMin = 0.f));
+	float Variability = 0.f;
+
+	UPROPERTY(EditAnywhere, Category = "Scatter", meta = (DataflowInput, DisplayName = "Random Seed"));
+	float RandomSeed = -1.f;
+
+	UPROPERTY(meta = (DataflowOutput, DisplayName = "Points"))
+		TArray<FVector> Points;
+
+	FRadialScatterPointsDataflowNode(const Dataflow::FNodeParameters& InParam, FGuid InGuid = FGuid::NewGuid())
+		: FDataflowNode(InParam, InGuid)
+	{
+		RegisterInputConnection(&Center);
+		RegisterInputConnection(&Normal);
+		RegisterInputConnection(&Radius);
+		RegisterInputConnection(&AngularSteps);
+		RegisterInputConnection(&RadialSteps);
+		RegisterInputConnection(&AngleOffset);
+		RegisterInputConnection(&Variability);
+		RegisterInputConnection(&RandomSeed);
+		RegisterOutputConnection(&Points);
+	}
+
+	virtual void Evaluate(Dataflow::FContext& Context, const FDataflowOutput* Out) const override;
+
+};
+
+
+//
+// GridScatterPoints
+//
+
+
+
+USTRUCT()
+struct FMakeLiteralFloatDataflowNode : public FDataflowNode
+{
+	GENERATED_USTRUCT_BODY()
+		DATAFLOW_NODE_DEFINE_INTERNAL(FMakeLiteralFloatDataflowNode, "MakeLiteralFloat", "Math|Float", "")
+
+public:
+	UPROPERTY(EditAnywhere, Category = "Float");
+	float Value = 0.f;
+
+	UPROPERTY(meta = (DataflowOutput, DisplayName = "Float"))
+		float Float;
+
+	FMakeLiteralFloatDataflowNode(const Dataflow::FNodeParameters& InParam, FGuid InGuid = FGuid::NewGuid())
+		: FDataflowNode(InParam, InGuid)
+	{
+		RegisterOutputConnection(&Float);
+	}
+
+	virtual void Evaluate(Dataflow::FContext& Context, const FDataflowOutput* Out) const override;
+
+};
+
+USTRUCT()
+struct FMakeLiteralIntDataflowNode : public FDataflowNode
+{
+	GENERATED_USTRUCT_BODY()
+		DATAFLOW_NODE_DEFINE_INTERNAL(FMakeLiteralIntDataflowNode, "MakeLiteralInt", "Math|Int", "")
+
+public:
+	UPROPERTY(EditAnywhere, Category = "Int");
+	int32 Value = 0;
+
+	UPROPERTY(meta = (DataflowOutput, DisplayName = "Int"))
+		int32 Int;
+
+	FMakeLiteralIntDataflowNode(const Dataflow::FNodeParameters& InParam, FGuid InGuid = FGuid::NewGuid())
+		: FDataflowNode(InParam, InGuid)
+	{
+		RegisterOutputConnection(&Int);
+	}
+
+	virtual void Evaluate(Dataflow::FContext& Context, const FDataflowOutput* Out) const override;
+
+};
+
+USTRUCT()
+struct FMakeLiteralBoolDataflowNode : public FDataflowNode
+{
+	GENERATED_USTRUCT_BODY()
+		DATAFLOW_NODE_DEFINE_INTERNAL(FMakeLiteralBoolDataflowNode, "MakeLiteralBool", "Math|Boolean", "")
+
+public:
+	UPROPERTY(EditAnywhere, Category = "Bool");
+	bool Value = false;
+
+	UPROPERTY(meta = (DataflowOutput, DisplayName = "Bool"))
+		bool Bool;
+
+	FMakeLiteralBoolDataflowNode(const Dataflow::FNodeParameters& InParam, FGuid InGuid = FGuid::NewGuid())
+		: FDataflowNode(InParam, InGuid)
+	{
+		RegisterOutputConnection(&Bool);
+	}
 
 	virtual void Evaluate(Dataflow::FContext& Context, const FDataflowOutput* Out) const override;
 
 };
 
+USTRUCT()
+struct FMakeLiteralVectorDataflowNode : public FDataflowNode
+{
+	GENERATED_USTRUCT_BODY()
+		DATAFLOW_NODE_DEFINE_INTERNAL(FMakeLiteralVectorDataflowNode, "MakeLiteralVector", "Math|Vector", "")
+
+public:
+	UPROPERTY(EditAnywhere, Category = "Vector");
+	FVector Value = FVector(0.0);
+
+	UPROPERTY(meta = (DataflowOutput, DisplayName = "Vector"))
+		FVector Vector;
+
+	FMakeLiteralVectorDataflowNode(const Dataflow::FNodeParameters& InParam, FGuid InGuid = FGuid::NewGuid())
+		: FDataflowNode(InParam, InGuid)
+	{
+		RegisterOutputConnection(&Vector);
+	}
+
+	virtual void Evaluate(Dataflow::FContext& Context, const FDataflowOutput* Out) const override;
+
+};
+
+USTRUCT()
+struct FIntToStringDataflowNode : public FDataflowNode
+{
+	GENERATED_USTRUCT_BODY()
+		DATAFLOW_NODE_DEFINE_INTERNAL(FIntToStringDataflowNode, "IntToString", "Utilities|String", "")
+
+public:
+	UPROPERTY(EditAnywhere, Category = "Int", meta = (DataflowInput, DisplayName = "Int"))
+		int32 Int = 0;
+
+	UPROPERTY(meta = (DataflowOutput, DisplayName = "String"))
+		FString String = FString("");
+
+	FIntToStringDataflowNode(const Dataflow::FNodeParameters& InParam, FGuid InGuid = FGuid::NewGuid())
+		: FDataflowNode(InParam, InGuid)
+	{
+		RegisterInputConnection(&Int);
+		RegisterOutputConnection(&String);
+	}
+
+	virtual void Evaluate(Dataflow::FContext& Context, const FDataflowOutput* Out) const override;
+
+};
+
+USTRUCT()
+struct FBoolToStringDataflowNode : public FDataflowNode
+{
+	GENERATED_USTRUCT_BODY()
+		DATAFLOW_NODE_DEFINE_INTERNAL(FBoolToStringDataflowNode, "BoolToString", "Utilities|String", "")
+
+public:
+	UPROPERTY(EditAnywhere, Category = "Bool", meta = (DataflowInput, DisplayName = "Bool"))
+		bool Bool = false;
+
+	UPROPERTY(meta = (DataflowOutput, DisplayName = "String"))
+		FString String = FString("");
+
+	FBoolToStringDataflowNode(const Dataflow::FNodeParameters& InParam, FGuid InGuid = FGuid::NewGuid())
+		: FDataflowNode(InParam, InGuid)
+	{
+		RegisterInputConnection(&Bool);
+		RegisterOutputConnection(&String);
+	}
+
+	virtual void Evaluate(Dataflow::FContext& Context, const FDataflowOutput* Out) const override;
+
+};
+
+USTRUCT()
+struct FExpandVectorDataflowNode : public FDataflowNode
+{
+	GENERATED_USTRUCT_BODY()
+		DATAFLOW_NODE_DEFINE_INTERNAL(FExpandVectorDataflowNode, "ExpandVector", "Utilities|Vector", "")
+
+public:
+	UPROPERTY(meta = (DataflowInput, DisplayName = "Vector"))
+		FVector Vector = FVector(0.0);
+
+	UPROPERTY(meta = (DataflowOutput, DisplayName = "X"))
+		float X = 0.f;
+
+	UPROPERTY(meta = (DataflowOutput, DisplayName = "Y"))
+		float Y = 0.f;
+
+	UPROPERTY(meta = (DataflowOutput, DisplayName = "Z"))
+		float Z = 0.f;
+
+
+	FExpandVectorDataflowNode(const Dataflow::FNodeParameters& InParam, FGuid InGuid = FGuid::NewGuid())
+		: FDataflowNode(InParam, InGuid)
+	{
+		RegisterInputConnection(&Vector);
+		RegisterOutputConnection(&X);
+		RegisterOutputConnection(&Y);
+		RegisterOutputConnection(&Z);
+	}
+
+	virtual void Evaluate(Dataflow::FContext& Context, const FDataflowOutput* Out) const override;
+
+};
+
+USTRUCT()
+struct FIntToFloatDataflowNode : public FDataflowNode
+{
+	GENERATED_USTRUCT_BODY()
+		DATAFLOW_NODE_DEFINE_INTERNAL(FIntToFloatDataflowNode, "IntToFloat", "Math|Conversions", "")
+
+public:
+	UPROPERTY(EditAnywhere, Category = "Int", meta = (DataflowInput, DisplayName = "Int"))
+		int32 Int = 0;
+
+	UPROPERTY(meta = (DataflowOutput, DisplayName = "Float"))
+		float Float = 0.f;
+
+	FIntToFloatDataflowNode(const Dataflow::FNodeParameters& InParam, FGuid InGuid = FGuid::NewGuid())
+		: FDataflowNode(InParam, InGuid)
+	{
+		RegisterInputConnection(&Int);
+		RegisterOutputConnection(&Float);
+	}
+
+	virtual void Evaluate(Dataflow::FContext& Context, const FDataflowOutput* Out) const override;
+
+};
+
+USTRUCT()
+struct FVoronoiFractureDataflowNode : public FDataflowNode
+{
+	GENERATED_USTRUCT_BODY()
+		DATAFLOW_NODE_DEFINE_INTERNAL(FVoronoiFractureDataflowNode, "VoronoiFracture", "Fracture", "")
+
+public:
+	UPROPERTY(meta = (DataflowInput, DataflowOutput, DisplayName = "Collection"))
+		FManagedArrayCollection Collection;
+
+	UPROPERTY(meta = (DataflowInput, DisplayName = "Points"))
+		TArray<FVector> Points;
+
+	UPROPERTY(EditAnywhere, Category = "Fracture", meta = (DataflowInput, DisplayName = "RandomSeed"));
+	float RandomSeed = -1.f;
+
+	UPROPERTY(EditAnywhere, Category = "Fracture", meta = (DataflowInput, DisplayName = "ChanceToFracture", UIMin = 0.f, UIMax = 1.f));
+	float ChanceToFracture = 1.f;
+
+	UPROPERTY(EditAnywhere, Category = "Fracture");
+	bool GroupFracture = true;
+
+	UPROPERTY(EditAnywhere, Category = "Fracture", meta = (DataflowInput, DisplayName = "Grout", UIMin = 0.f));
+	float Grout = 0.f;
+
+	UPROPERTY(EditAnywhere, Category = "Noise", meta = (DataflowInput, DisplayName = "Amplitude", UIMin = 0.f));
+	float Amplitude = 0.f;
+
+	UPROPERTY(EditAnywhere, Category = "Noise", meta = (DataflowInput, DisplayName = "Frequency", UIMin = 0.00001f));
+	float Frequency = 0.1f;
+
+	UPROPERTY(EditAnywhere, Category = "Noise", meta = (DataflowInput, DisplayName = "Persistence", UIMin = 0.f));
+	float Persistence = 0.5f;
+
+	UPROPERTY(EditAnywhere, Category = "Noise", meta = (DataflowInput, DisplayName = "Lacunarity", UIMin = 0.f));
+	float Lacunarity = 2.f;
+
+	UPROPERTY(EditAnywhere, Category = "Noise", meta = (DataflowInput, DisplayName = "OctaveNumber", UIMin = 0.f));
+	int32 OctaveNumber = 4;
+
+	UPROPERTY(EditAnywhere, Category = "Noise", meta = (DataflowInput, DisplayName = "PointSpacing", UIMin = 0.f));
+	float PointSpacing = 10.f;
+
+	UPROPERTY(EditAnywhere, Category = "Collision");
+	bool AddSamplesForCollision = false;
+
+	UPROPERTY(EditAnywhere, Category = "Collision", meta = (DataflowInput, DisplayName = "CollisionSampleSpacing", UIMin = 0.f));
+	float CollisionSampleSpacing = 50.f;
+
+	FVoronoiFractureDataflowNode(const Dataflow::FNodeParameters& InParam, FGuid InGuid = FGuid::NewGuid())
+		: FDataflowNode(InParam, InGuid)
+	{
+		RegisterInputConnection(&Collection);
+		RegisterInputConnection(&Points);
+		RegisterInputConnection(&RandomSeed);
+		RegisterInputConnection(&ChanceToFracture);
+		RegisterInputConnection(&Grout);
+		RegisterInputConnection(&Amplitude);
+		RegisterInputConnection(&Frequency);
+		RegisterInputConnection(&Persistence);
+		RegisterInputConnection(&Lacunarity);
+		RegisterInputConnection(&OctaveNumber);
+		RegisterInputConnection(&PointSpacing);
+		RegisterInputConnection(&CollisionSampleSpacing);
+		RegisterOutputConnection(&Collection);
+	}
+
+	virtual void Evaluate(Dataflow::FContext& Context, const FDataflowOutput* Out) const override;
+
+};
+
+USTRUCT()
+struct FStringAppendDataflowNode : public FDataflowNode
+{
+	GENERATED_USTRUCT_BODY()
+		DATAFLOW_NODE_DEFINE_INTERNAL(FStringAppendDataflowNode, "StringAppend", "Utilities|String", "")
+
+public:
+	UPROPERTY(EditAnywhere, Category = "String", meta = (DataflowInput, DisplayName = "String1"))
+		FString String1 = FString("");
+
+	UPROPERTY(EditAnywhere, Category = "String", meta = (DataflowInput, DisplayName = "String2"))
+		FString String2 = FString("");
+
+	UPROPERTY(meta = (DataflowOutput, DisplayName = "String"))
+		FString String;
+
+	FStringAppendDataflowNode(const Dataflow::FNodeParameters& InParam, FGuid InGuid = FGuid::NewGuid())
+		: FDataflowNode(InParam, InGuid)
+	{
+		RegisterInputConnection(&String1);
+		RegisterInputConnection(&String2);
+		RegisterOutputConnection(&String);
+	}
+
+	virtual void Evaluate(Dataflow::FContext& Context, const FDataflowOutput* Out) const override;
+
+};
+
+USTRUCT()
+struct FRandomFloatDataflowNode : public FDataflowNode
+{
+	GENERATED_USTRUCT_BODY()
+		DATAFLOW_NODE_DEFINE_INTERNAL(FRandomFloatDataflowNode, "RandomFloat", "Math|Random", "")
+
+public:
+	UPROPERTY(EditAnywhere, Category = "Seed")
+		bool Deterministic = false;
+
+	UPROPERTY(EditAnywhere, Category = "Seed", meta = (DataflowInput, DisplayName = "RandomSeed", EditCondition = "Deterministic"))
+		float RandomSeed = FMath::FRandRange(-1e5, 1e5);
+
+	UPROPERTY(meta = (DataflowOutput, DisplayName = "Float"))
+		float Float = 0.f;
+
+	FRandomFloatDataflowNode(const Dataflow::FNodeParameters& InParam, FGuid InGuid = FGuid::NewGuid())
+		: FDataflowNode(InParam, InGuid)
+	{
+		RegisterInputConnection(&RandomSeed);
+		RegisterOutputConnection(&Float);
+	}
+
+	virtual void Evaluate(Dataflow::FContext& Context, const FDataflowOutput* Out) const override;
+
+};
+
+USTRUCT()
+struct FRandomFloatInRangeDataflowNode : public FDataflowNode
+{
+	GENERATED_USTRUCT_BODY()
+		DATAFLOW_NODE_DEFINE_INTERNAL(FRandomFloatInRangeDataflowNode, "RandomFloatInRange", "Math|Random", "")
+
+public:
+	UPROPERTY(EditAnywhere, Category = "Seed")
+		bool Deterministic = false;
+
+	UPROPERTY(EditAnywhere, Category = "Seed", meta = (DataflowInput, DisplayName = "RandomSeed", EditCondition = "Deterministic"))
+		float RandomSeed = FMath::FRandRange(-1e5, 1e5);
+
+	UPROPERTY(EditAnywhere, Category = "Random", meta = (DataflowInput, DisplayName = "Min"))
+		float Min = 0.f;
+
+	UPROPERTY(EditAnywhere, Category = "Random", meta = (DataflowInput, DisplayName = "Max"))
+		float Max = 1.f;
+
+	UPROPERTY(meta = (DataflowOutput, DisplayName = "Float"))
+		float Float = 0.f;
+
+	FRandomFloatInRangeDataflowNode(const Dataflow::FNodeParameters& InParam, FGuid InGuid = FGuid::NewGuid())
+		: FDataflowNode(InParam, InGuid)
+	{
+		RegisterInputConnection(&RandomSeed);
+		RegisterInputConnection(&Min);
+		RegisterInputConnection(&Max);
+		RegisterOutputConnection(&Float);
+	}
+
+	virtual void Evaluate(Dataflow::FContext& Context, const FDataflowOutput* Out) const override;
+
+};
+
+USTRUCT()
+struct FRandomUnitVectorDataflowNode : public FDataflowNode
+{
+	GENERATED_USTRUCT_BODY()
+		DATAFLOW_NODE_DEFINE_INTERNAL(FRandomUnitVectorDataflowNode, "RandomUnitVector", "Math|Random", "")
+
+public:
+	UPROPERTY(EditAnywhere, Category = "Seed")
+		bool Deterministic = false;
+
+	UPROPERTY(EditAnywhere, Category = "Seed", meta = (DataflowInput, DisplayName = "RandomSeed", EditCondition = "Deterministic"))
+		float RandomSeed = FMath::FRandRange(-1e5, 1e5);
+
+	UPROPERTY(meta = (DataflowOutput, DisplayName = "Vector"))
+		FVector Vector = FVector(0.0);
+
+	FRandomUnitVectorDataflowNode(const Dataflow::FNodeParameters& InParam, FGuid InGuid = FGuid::NewGuid())
+		: FDataflowNode(InParam, InGuid)
+	{
+		RegisterInputConnection(&RandomSeed);
+		RegisterOutputConnection(&Vector);
+	}
+
+	virtual void Evaluate(Dataflow::FContext& Context, const FDataflowOutput* Out) const override;
+
+};
+
+USTRUCT()
+struct FRandomUnitVectorInConeDataflowNode : public FDataflowNode
+{
+	GENERATED_USTRUCT_BODY()
+		DATAFLOW_NODE_DEFINE_INTERNAL(FRandomUnitVectorInConeDataflowNode, "RandomUnitVectorInCone", "Math|Random", "")
+
+public:
+	UPROPERTY(EditAnywhere, Category = "Seed")
+		bool Deterministic = false;
+
+	UPROPERTY(EditAnywhere, Category = "Seed", meta = (DataflowInput, DisplayName = "RandomSeed", EditCondition = "Deterministic"))
+		float RandomSeed = FMath::FRandRange(-1e5, 1e5);
+
+	UPROPERTY(EditAnywhere, Category = "Random", meta = (DataflowInput, DisplayName = "ConeDirection"))
+		FVector ConeDirection = FVector(0.0, 0.0, 1.0);
+
+	UPROPERTY(EditAnywhere, Category = "Random", meta = (DataflowInput, DisplayName = "ConeHalfAngle"))
+		float ConeHalfAngle = PI / 4.f;
+
+	UPROPERTY(meta = (DataflowOutput, DisplayName = "Vector"))
+		FVector Vector = FVector(0.0);
+
+	FRandomUnitVectorInConeDataflowNode(const Dataflow::FNodeParameters& InParam, FGuid InGuid = FGuid::NewGuid())
+		: FDataflowNode(InParam, InGuid)
+	{
+		RegisterInputConnection(&RandomSeed);
+		RegisterInputConnection(&ConeDirection);
+		RegisterInputConnection(&ConeHalfAngle);
+		RegisterOutputConnection(&Vector);
+	}
+
+	virtual void Evaluate(Dataflow::FContext& Context, const FDataflowOutput* Out) const override;
+
+};
+
+USTRUCT()
+struct FRadiansToDegreesDataflowNode : public FDataflowNode
+{
+	GENERATED_USTRUCT_BODY()
+		DATAFLOW_NODE_DEFINE_INTERNAL(FRadiansToDegreesDataflowNode, "RadiansToDegrees", "Math|Trigonometry", "")
+
+public:
+	UPROPERTY(EditAnywhere, Category = "Radians", meta = (DataflowInput, DisplayName = "Radians"))
+		float Radians = 0.f;
+
+	UPROPERTY(meta = (DataflowOutput, DisplayName = "Degrees"))
+		float Degrees = 0.f;
+
+	FRadiansToDegreesDataflowNode(const Dataflow::FNodeParameters& InParam, FGuid InGuid = FGuid::NewGuid())
+		: FDataflowNode(InParam, InGuid)
+	{
+		RegisterInputConnection(&Radians);
+		RegisterOutputConnection(&Degrees);
+	}
+
+	virtual void Evaluate(Dataflow::FContext& Context, const FDataflowOutput* Out) const override;
+
+};
+
+USTRUCT()
+struct FDegreesToRadiansDataflowNode : public FDataflowNode
+{
+	GENERATED_USTRUCT_BODY()
+		DATAFLOW_NODE_DEFINE_INTERNAL(FDegreesToRadiansDataflowNode, "DegreesToRadians", "Math|Trigonometry", "")
+
+public:
+	UPROPERTY(EditAnywhere, Category = "Degrees", meta = (DataflowInput, DisplayName = "Degrees"))
+		float Degrees = 0.f;
+
+	UPROPERTY(meta = (DataflowOutput, DisplayName = "Radians"))
+		float Radians = 0.f;
+
+	FDegreesToRadiansDataflowNode(const Dataflow::FNodeParameters& InParam, FGuid InGuid = FGuid::NewGuid())
+		: FDataflowNode(InParam, InGuid)
+	{
+		RegisterInputConnection(&Degrees);
+		RegisterOutputConnection(&Radians);
+	}
+
+	virtual void Evaluate(Dataflow::FContext& Context, const FDataflowOutput* Out) const override;
+
+};
+
+USTRUCT()
+struct FExplodedViewDataflowNode : public FDataflowNode
+{
+	GENERATED_USTRUCT_BODY()
+		DATAFLOW_NODE_DEFINE_INTERNAL(FExplodedViewDataflowNode, "ExplodedView", "Fracture|Utilities", "")
+
+public:
+	UPROPERTY(meta = (DataflowInput, DataflowOutput, DisplayName = "Collection"))
+		FManagedArrayCollection Collection;
+
+	UPROPERTY(EditAnywhere, Category = "Scale", meta = (DataflowInput, DisplayName = "UniformScale"))
+		float UniformScale = 1.f;
+
+	UPROPERTY(EditAnywhere, Category = "Scale", meta = (DataflowInput, DisplayName = "Scale"))
+		FVector Scale = FVector(1.0);
+
+	FExplodedViewDataflowNode(const Dataflow::FNodeParameters& InParam, FGuid InGuid = FGuid::NewGuid())
+		: FDataflowNode(InParam, InGuid)
+	{
+		RegisterInputConnection(&Collection);
+		RegisterInputConnection(&UniformScale);
+		RegisterInputConnection(&Scale);
+		RegisterOutputConnection(&Collection);
+	}
+
+	virtual void Evaluate(Dataflow::FContext& Context, const FDataflowOutput* Out) const override;
+
+};
+
+USTRUCT()
+struct FCreateNonOverlappingConvexHullsDataflowNode : public FDataflowNode
+{
+	GENERATED_USTRUCT_BODY()
+		DATAFLOW_NODE_DEFINE_INTERNAL(FCreateNonOverlappingConvexHullsDataflowNode, "CreateNonOverlappingConvexHulls", "Fracture|Utilities", "")
+
+public:
+	UPROPERTY(meta = (DataflowInput, DataflowOutput, DisplayName = "Collection"))
+		FManagedArrayCollection Collection;
+
+	UPROPERTY(EditAnywhere, Category = "Convex", meta = (DataflowInput, DisplayName = "CanRemoveFraction", UIMin = 0.01f, UIMax = 1.f))
+		float CanRemoveFraction = 0.5f;
+
+	UPROPERTY(EditAnywhere, Category = "Convex", meta = (DataflowInput, DisplayName = "CanExceedFraction", UIMin = 0.f))
+		float CanExceedFraction = 0.5f;
+
+	UPROPERTY(EditAnywhere, Category = "Convex", meta = (DataflowInput, DisplayName = "SimplificationDistanceThreshold", UIMin = 0.f))
+		float SimplificationDistanceThreshold = 10.f;
+
+	FCreateNonOverlappingConvexHullsDataflowNode(const Dataflow::FNodeParameters& InParam, FGuid InGuid = FGuid::NewGuid())
+		: FDataflowNode(InParam, InGuid)
+	{
+		RegisterInputConnection(&Collection);
+		RegisterInputConnection(&CanRemoveFraction);
+		RegisterInputConnection(&CanExceedFraction);
+		RegisterInputConnection(&SimplificationDistanceThreshold);
+		RegisterOutputConnection(&Collection);
+	}
+
+	virtual void Evaluate(Dataflow::FContext& Context, const FDataflowOutput* Out) const override;
+
+};
+
+USTRUCT()
+struct FPlaneCutterDataflowNode : public FDataflowNode
+{
+	GENERATED_USTRUCT_BODY()
+		DATAFLOW_NODE_DEFINE_INTERNAL(FPlaneCutterDataflowNode, "PlaneCutter", "Fracture", "")
+
+public:
+	UPROPERTY(meta = (DataflowInput, DataflowOutput, DisplayName = "Collection"))
+		FManagedArrayCollection Collection;
+
+	UPROPERTY(meta = (DataflowInput, DisplayName = "BoundingBox"))
+		FBox BoundingBox = FBox(ForceInit);
+
+	UPROPERTY(EditAnywhere, Category = "Fracture", meta = (DataflowInput, DisplayName = "NumPlanes", UIMin = 1))
+		int32 NumPlanes = 1;
+
+	UPROPERTY(EditAnywhere, Category = "Fracture", meta = (DataflowInput, DisplayName = "RandomSeed"))
+		float RandomSeed = -1.f;
+
+	UPROPERTY(EditAnywhere, Category = "Fracture", meta = (DataflowInput, DisplayName = "Grout", UIMin = 0.f))
+		float Grout = 0.f;
+
+	UPROPERTY(EditAnywhere, Category = "Noise", meta = (DataflowInput, DisplayName = "Amplitude", UIMin = 0.f));
+	float Amplitude = 0.f;
+
+	UPROPERTY(EditAnywhere, Category = "Noise", meta = (DataflowInput, DisplayName = "Frequency", UIMin = 0.00001f));
+	float Frequency = 0.1f;
+
+	UPROPERTY(EditAnywhere, Category = "Noise", meta = (DataflowInput, DisplayName = "Persistence", UIMin = 0.f));
+	float Persistence = 0.5f;
+
+	UPROPERTY(EditAnywhere, Category = "Noise", meta = (DataflowInput, DisplayName = "Lacunarity", UIMin = 0.f));
+	float Lacunarity = 2.f;
+
+	UPROPERTY(EditAnywhere, Category = "Noise", meta = (DataflowInput, DisplayName = "OctaveNumber", UIMin = 0.f));
+	int32 OctaveNumber = 4;
+
+	UPROPERTY(EditAnywhere, Category = "Noise", meta = (DataflowInput, DisplayName = "PointSpacing", UIMin = 0.f));
+	float PointSpacing = 10.f;
+
+	UPROPERTY(EditAnywhere, Category = "Collision");
+	bool AddSamplesForCollision = false;
+
+	UPROPERTY(EditAnywhere, Category = "Collision", meta = (DataflowInput, DisplayName = "CollisionSampleSpacing", UIMin = 0.f));
+	float CollisionSampleSpacing = 50.f;
+
+	FPlaneCutterDataflowNode(const Dataflow::FNodeParameters& InParam, FGuid InGuid = FGuid::NewGuid())
+		: FDataflowNode(InParam, InGuid)
+	{
+		RegisterInputConnection(&Collection);
+		RegisterInputConnection(&BoundingBox);
+		RegisterInputConnection(&NumPlanes);
+		RegisterInputConnection(&RandomSeed);
+		RegisterInputConnection(&Grout);
+		RegisterInputConnection(&Amplitude);
+		RegisterInputConnection(&Frequency);
+		RegisterInputConnection(&Persistence);
+		RegisterInputConnection(&Lacunarity);
+		RegisterInputConnection(&OctaveNumber);
+		RegisterInputConnection(&PointSpacing);
+		RegisterInputConnection(&CollisionSampleSpacing);
+		RegisterOutputConnection(&Collection);
+	}
+
+	virtual void Evaluate(Dataflow::FContext& Context, const FDataflowOutput* Out) const override;
+
+};
+
+USTRUCT()
+struct FHashStringDataflowNode : public FDataflowNode
+{
+	GENERATED_USTRUCT_BODY()
+		DATAFLOW_NODE_DEFINE_INTERNAL(FHashStringDataflowNode, "HashString", "Utilities|String", "")
+
+public:
+	UPROPERTY(meta = (DataflowInput, DisplayName = "String"))
+		FString String;
+
+	UPROPERTY(meta = (DataflowOutput, DisplayName = "Hash"))
+		int32 Hash = 0;
+
+
+	FHashStringDataflowNode(const Dataflow::FNodeParameters& InParam, FGuid InGuid = FGuid::NewGuid())
+		: FDataflowNode(InParam, InGuid)
+	{
+		RegisterInputConnection(&String);
+		RegisterOutputConnection(&Hash);
+	}
+
+	virtual void Evaluate(Dataflow::FContext& Context, const FDataflowOutput* Out) const override;
+
+};
+
+USTRUCT()
+struct FHashVectorDataflowNode : public FDataflowNode
+{
+	GENERATED_USTRUCT_BODY()
+		DATAFLOW_NODE_DEFINE_INTERNAL(FHashVectorDataflowNode, "HashVector", "Utilities|Vector", "")
+
+public:
+	UPROPERTY(meta = (DataflowInput, DisplayName = "Vector"))
+		FVector Vector;
+
+	UPROPERTY(meta = (DataflowOutput, DisplayName = "Hash"))
+		int32 Hash = 0;
+
+
+	FHashVectorDataflowNode(const Dataflow::FNodeParameters& InParam, FGuid InGuid = FGuid::NewGuid())
+		: FDataflowNode(InParam, InGuid)
+	{
+		RegisterInputConnection(&Vector);
+		RegisterOutputConnection(&Hash);
+	}
+
+	virtual void Evaluate(Dataflow::FContext& Context, const FDataflowOutput* Out) const override;
+
+};
+
+UENUM(BlueprintType)
+enum class EFloatToIntFunctionEnum : uint8
+{
+	Dataflow_FloatToInt_Function_Floor UMETA(DisplayName = "Floor()"),
+	Dataflow_FloatToInt_Function_Ceil UMETA(DisplayName = "Ceil()"),
+	Dataflow_FloatToInt_Function_Round UMETA(DisplayName = "Round()"),
+	Dataflow_FloatToInt_Function_Truncate UMETA(DisplayName = "Truncate()"),
+	//~~~
+	//256th entry
+	Dataflow_Max                UMETA(Hidden)
+};
+
+USTRUCT()
+struct FFloatToIntDataflowNode : public FDataflowNode
+{
+	GENERATED_USTRUCT_BODY()
+		DATAFLOW_NODE_DEFINE_INTERNAL(FFloatToIntDataflowNode, "FloatToInt", "Math|Conversions", "")
+
+public:
+	UPROPERTY(EditAnywhere, Category = "Float", meta = (DisplayName = "Function"));
+	EFloatToIntFunctionEnum Function = EFloatToIntFunctionEnum::Dataflow_FloatToInt_Function_Round;
+
+	UPROPERTY(EditAnywhere, Category = "Float", meta = (DataflowInput, DisplayName = "Float"))
+		float Float = 0.f;
+
+	UPROPERTY(meta = (DataflowOutput, DisplayName = "Int"))
+		int32 Int = 0;
+
+	FFloatToIntDataflowNode(const Dataflow::FNodeParameters& InParam, FGuid InGuid = FGuid::NewGuid())
+		: FDataflowNode(InParam, InGuid)
+	{
+		RegisterInputConnection(&Float);
+		RegisterOutputConnection(&Int);
+	}
+
+	virtual void Evaluate(Dataflow::FContext& Context, const FDataflowOutput* Out) const override;
+
+};
+
+UENUM(BlueprintType)
+enum class EMathConstantsEnum : uint8
+{
+	Dataflow_MathConstants_Pi UMETA(DisplayName = "Pi"),
+	Dataflow_MathConstants_HalfPi UMETA(DisplayName = "HalfPi"),
+	Dataflow_MathConstants_TwoPi UMETA(DisplayName = "TwoPi"),
+	Dataflow_MathConstants_FourPi UMETA(DisplayName = "FourPi"),
+	Dataflow_MathConstants_InvPi UMETA(DisplayName = "InvPi"),
+	Dataflow_MathConstants_InvTwoPi UMETA(DisplayName = "InvTwoPi"),
+	Dataflow_MathConstants_Sqrt2 UMETA(DisplayName = "Sqrt2"),
+	Dataflow_MathConstants_InvSqrt2 UMETA(DisplayName = "InvSqrt2"),
+	Dataflow_MathConstants_Sqrt3 UMETA(DisplayName = "Sqrt3"),
+	Dataflow_MathConstants_InvSqrt3 UMETA(DisplayName = "InvSqrt3"),
+	Dataflow_FloatToInt_Function_E UMETA(DisplayName = "e"),
+	Dataflow_FloatToInt_Function_Gamma UMETA(DisplayName = "Gamma"),
+	Dataflow_FloatToInt_Function_GoldenRatio UMETA(DisplayName = "GoldenRatio"),
+	Dataflow_FloatToInt_Function_ZeroTolerance UMETA(DisplayName = "ZeroTolerance"),
+	//~~~
+	//256th entry
+	Dataflow_Max                UMETA(Hidden)
+};
+
+USTRUCT()
+struct FMathConstantsDataflowNode : public FDataflowNode
+{
+	GENERATED_USTRUCT_BODY()
+		DATAFLOW_NODE_DEFINE_INTERNAL(FMathConstantsDataflowNode, "MathConstants", "Math|Utilities", "")
+
+public:
+	UPROPERTY(EditAnywhere, Category = "Constants", meta = (DisplayName = "Constant"));
+	EMathConstantsEnum Constant = EMathConstantsEnum::Dataflow_MathConstants_Pi;
+
+	UPROPERTY(meta = (DataflowOutput, DisplayName = "Float"))
+		float Float = 0;
+
+	FMathConstantsDataflowNode(const Dataflow::FNodeParameters& InParam, FGuid InGuid = FGuid::NewGuid())
+		: FDataflowNode(InParam, InGuid)
+	{
+		RegisterOutputConnection(&Float);
+	}
+
+	virtual void Evaluate(Dataflow::FContext& Context, const FDataflowOutput* Out) const override;
+
+};
+
+
 namespace Dataflow
 {
 	void GeometryCollectionEngineAssetNodes();
