Dataflow : Fix Warnings

diff --git a/Engine/Plugins/Experimental/GeometryCollectionPlugin/Source/GeometryCollectionNodes/Public/Dataflow/GeometryCollectionNodes.h b/Engine/Plugins/Experimental/GeometryCollectionPlugin/Source/GeometryCollectionNodes/Public/Dataflow/GeometryCollectionNodes.h
--- a/Engine/Plugins/Experimental/GeometryCollectionPlugin/Source/GeometryCollectionNodes/Public/Dataflow/GeometryCollectionNodes.h
+++ b/Engine/Plugins/Experimental/GeometryCollectionPlugin/Source/GeometryCollectionNodes/Public/Dataflow/GeometryCollectionNodes.h
@@ -480,7 +480,7 @@ public:
 	float Value = 0.f;
 
 	UPROPERTY(meta = (DataflowOutput, DisplayName = "Float"))
-		float Float;
+	float Float = 0.f;
 
 	FMakeLiteralFloatDataflowNode(const Dataflow::FNodeParameters& InParam, FGuid InGuid = FGuid::NewGuid())
 		: FDataflowNode(InParam, InGuid)
@@ -496,14 +496,14 @@ USTRUCT()
 struct FMakeLiteralIntDataflowNode : public FDataflowNode
 {
 	GENERATED_USTRUCT_BODY()
-		DATAFLOW_NODE_DEFINE_INTERNAL(FMakeLiteralIntDataflowNode, "MakeLiteralInt", "Math|Int", "")
+	DATAFLOW_NODE_DEFINE_INTERNAL(FMakeLiteralIntDataflowNode, "MakeLiteralInt", "Math|Int", "")
 
 public:
 	UPROPERTY(EditAnywhere, Category = "Int");
 	int32 Value = 0;
 
 	UPROPERTY(meta = (DataflowOutput, DisplayName = "Int"))
-		int32 Int;
+	int32 Int = 0;
 
 	FMakeLiteralIntDataflowNode(const Dataflow::FNodeParameters& InParam, FGuid InGuid = FGuid::NewGuid())
 		: FDataflowNode(InParam, InGuid)
@@ -519,14 +519,14 @@ USTRUCT()
 struct FMakeLiteralBoolDataflowNode : public FDataflowNode
 {
 	GENERATED_USTRUCT_BODY()
-		DATAFLOW_NODE_DEFINE_INTERNAL(FMakeLiteralBoolDataflowNode, "MakeLiteralBool", "Math|Boolean", "")
+	DATAFLOW_NODE_DEFINE_INTERNAL(FMakeLiteralBoolDataflowNode, "MakeLiteralBool", "Math|Boolean", "")
 
 public:
 	UPROPERTY(EditAnywhere, Category = "Bool");
 	bool Value = false;
 
 	UPROPERTY(meta = (DataflowOutput, DisplayName = "Bool"))
-		bool Bool;
+	bool Bool = false;
 
 	FMakeLiteralBoolDataflowNode(const Dataflow::FNodeParameters& InParam, FGuid InGuid = FGuid::NewGuid())
 		: FDataflowNode(InParam, InGuid)
@@ -549,7 +549,7 @@ public:
 	FVector Value = FVector(0.0);
 
 	UPROPERTY(meta = (DataflowOutput, DisplayName = "Vector"))
-		FVector Vector;
+	FVector Vector = FVector(0.0);
 
 	FMakeLiteralVectorDataflowNode(const Dataflow::FNodeParameters& InParam, FGuid InGuid = FGuid::NewGuid())
 		: FDataflowNode(InParam, InGuid)
@@ -569,10 +569,10 @@ struct FIntToStringDataflowNode : public FDataflowNode
 
 public:
 	UPROPERTY(EditAnywhere, Category = "Int", meta = (DataflowInput, DisplayName = "Int"))
-		int32 Int = 0;
+	int32 Int = 0;
 
 	UPROPERTY(meta = (DataflowOutput, DisplayName = "String"))
-		FString String = FString("");
+	FString String = FString("");
 
 	FIntToStringDataflowNode(const Dataflow::FNodeParameters& InParam, FGuid InGuid = FGuid::NewGuid())
 		: FDataflowNode(InParam, InGuid)
@@ -589,14 +589,14 @@ USTRUCT()
 struct FBoolToStringDataflowNode : public FDataflowNode
 {
 	GENERATED_USTRUCT_BODY()
-		DATAFLOW_NODE_DEFINE_INTERNAL(FBoolToStringDataflowNode, "BoolToString", "Utilities|String", "")
+	DATAFLOW_NODE_DEFINE_INTERNAL(FBoolToStringDataflowNode, "BoolToString", "Utilities|String", "")
 
 public:
 	UPROPERTY(EditAnywhere, Category = "Bool", meta = (DataflowInput, DisplayName = "Bool"))
-		bool Bool = false;
+	bool Bool = false;
 
 	UPROPERTY(meta = (DataflowOutput, DisplayName = "String"))
-		FString String = FString("");
+	FString String = FString("");
 
 	FBoolToStringDataflowNode(const Dataflow::FNodeParameters& InParam, FGuid InGuid = FGuid::NewGuid())
 		: FDataflowNode(InParam, InGuid)
@@ -613,20 +613,20 @@ USTRUCT()
 struct FExpandVectorDataflowNode : public FDataflowNode
 {
 	GENERATED_USTRUCT_BODY()
-		DATAFLOW_NODE_DEFINE_INTERNAL(FExpandVectorDataflowNode, "ExpandVector", "Utilities|Vector", "")
+	DATAFLOW_NODE_DEFINE_INTERNAL(FExpandVectorDataflowNode, "ExpandVector", "Utilities|Vector", "")
 
 public:
 	UPROPERTY(meta = (DataflowInput, DisplayName = "Vector"))
-		FVector Vector = FVector(0.0);
+	FVector Vector = FVector(0.0);
 
 	UPROPERTY(meta = (DataflowOutput, DisplayName = "X"))
-		float X = 0.f;
+	float X = 0.f;
 
 	UPROPERTY(meta = (DataflowOutput, DisplayName = "Y"))
-		float Y = 0.f;
+	float Y = 0.f;
 
 	UPROPERTY(meta = (DataflowOutput, DisplayName = "Z"))
-		float Z = 0.f;
+	float Z = 0.f;
 
 
 	FExpandVectorDataflowNode(const Dataflow::FNodeParameters& InParam, FGuid InGuid = FGuid::NewGuid())
@@ -646,14 +646,14 @@ USTRUCT()
 struct FIntToFloatDataflowNode : public FDataflowNode
 {
 	GENERATED_USTRUCT_BODY()
-		DATAFLOW_NODE_DEFINE_INTERNAL(FIntToFloatDataflowNode, "IntToFloat", "Math|Conversions", "")
+	DATAFLOW_NODE_DEFINE_INTERNAL(FIntToFloatDataflowNode, "IntToFloat", "Math|Conversions", "")
 
 public:
 	UPROPERTY(EditAnywhere, Category = "Int", meta = (DataflowInput, DisplayName = "Int"))
-		int32 Int = 0;
+	int32 Int = 0;
 
 	UPROPERTY(meta = (DataflowOutput, DisplayName = "Float"))
-		float Float = 0.f;
+	float Float = 0.f;
 
 	FIntToFloatDataflowNode(const Dataflow::FNodeParameters& InParam, FGuid InGuid = FGuid::NewGuid())
 		: FDataflowNode(InParam, InGuid)
@@ -670,14 +670,14 @@ USTRUCT()
 struct FVoronoiFractureDataflowNode : public FDataflowNode
 {
 	GENERATED_USTRUCT_BODY()
-		DATAFLOW_NODE_DEFINE_INTERNAL(FVoronoiFractureDataflowNode, "VoronoiFracture", "Fracture", "")
+	DATAFLOW_NODE_DEFINE_INTERNAL(FVoronoiFractureDataflowNode, "VoronoiFracture", "Fracture", "")
 
 public:
 	UPROPERTY(meta = (DataflowInput, DataflowOutput, DisplayName = "Collection"))
-		FManagedArrayCollection Collection;
+	FManagedArrayCollection Collection;
 
 	UPROPERTY(meta = (DataflowInput, DisplayName = "Points"))
-		TArray<FVector> Points;
+	TArray<FVector> Points;
 
 	UPROPERTY(EditAnywhere, Category = "Fracture", meta = (DataflowInput, DisplayName = "RandomSeed"));
 	float RandomSeed = -1.f;
@@ -741,17 +741,17 @@ USTRUCT()
 struct FStringAppendDataflowNode : public FDataflowNode
 {
 	GENERATED_USTRUCT_BODY()
-		DATAFLOW_NODE_DEFINE_INTERNAL(FStringAppendDataflowNode, "StringAppend", "Utilities|String", "")
+	DATAFLOW_NODE_DEFINE_INTERNAL(FStringAppendDataflowNode, "StringAppend", "Utilities|String", "")
 
 public:
 	UPROPERTY(EditAnywhere, Category = "String", meta = (DataflowInput, DisplayName = "String1"))
-		FString String1 = FString("");
+	FString String1 = FString("");
 
 	UPROPERTY(EditAnywhere, Category = "String", meta = (DataflowInput, DisplayName = "String2"))
-		FString String2 = FString("");
+	FString String2 = FString("");
 
 	UPROPERTY(meta = (DataflowOutput, DisplayName = "String"))
-		FString String;
+	FString String = FString("");
 
 	FStringAppendDataflowNode(const Dataflow::FNodeParameters& InParam, FGuid InGuid = FGuid::NewGuid())
 		: FDataflowNode(InParam, InGuid)
@@ -773,17 +773,18 @@ struct FRandomFloatDataflowNode : public FDataflowNode
 
 public:
 	UPROPERTY(EditAnywhere, Category = "Seed")
-		bool Deterministic = false;
+	bool Deterministic = false;
 
 	UPROPERTY(EditAnywhere, Category = "Seed", meta = (DataflowInput, DisplayName = "RandomSeed", EditCondition = "Deterministic"))
-		float RandomSeed = FMath::FRandRange(-1e5, 1e5);
+	float RandomSeed = 0.f;
 
 	UPROPERTY(meta = (DataflowOutput, DisplayName = "Float"))
-		float Float = 0.f;
+	float Float = 0.f;
 
 	FRandomFloatDataflowNode(const Dataflow::FNodeParameters& InParam, FGuid InGuid = FGuid::NewGuid())
 		: FDataflowNode(InParam, InGuid)
 	{
+		RandomSeed = FMath::FRandRange(-1e5, 1e5);
 		RegisterInputConnection(&RandomSeed);
 		RegisterOutputConnection(&Float);
 	}
@@ -800,23 +801,24 @@ struct FRandomFloatInRangeDataflowNode : public FDataflowNode
 
 public:
 	UPROPERTY(EditAnywhere, Category = "Seed")
-		bool Deterministic = false;
+	bool Deterministic = false;
 
 	UPROPERTY(EditAnywhere, Category = "Seed", meta = (DataflowInput, DisplayName = "RandomSeed", EditCondition = "Deterministic"))
-		float RandomSeed = FMath::FRandRange(-1e5, 1e5);
-
+	float RandomSeed = 0.f;
+	
 	UPROPERTY(EditAnywhere, Category = "Random", meta = (DataflowInput, DisplayName = "Min"))
-		float Min = 0.f;
+	float Min = 0.f;
 
 	UPROPERTY(EditAnywhere, Category = "Random", meta = (DataflowInput, DisplayName = "Max"))
-		float Max = 1.f;
+	float Max = 1.f;
 
 	UPROPERTY(meta = (DataflowOutput, DisplayName = "Float"))
-		float Float = 0.f;
+	float Float = 0.f;
 
 	FRandomFloatInRangeDataflowNode(const Dataflow::FNodeParameters& InParam, FGuid InGuid = FGuid::NewGuid())
 		: FDataflowNode(InParam, InGuid)
 	{
+		RandomSeed = FMath::FRandRange(-1e5, 1e5);
 		RegisterInputConnection(&RandomSeed);
 		RegisterInputConnection(&Min);
 		RegisterInputConnection(&Max);
@@ -831,21 +833,22 @@ USTRUCT()
 struct FRandomUnitVectorDataflowNode : public FDataflowNode
 {
 	GENERATED_USTRUCT_BODY()
-		DATAFLOW_NODE_DEFINE_INTERNAL(FRandomUnitVectorDataflowNode, "RandomUnitVector", "Math|Random", "")
+	DATAFLOW_NODE_DEFINE_INTERNAL(FRandomUnitVectorDataflowNode, "RandomUnitVector", "Math|Random", "")
 
 public:
 	UPROPERTY(EditAnywhere, Category = "Seed")
-		bool Deterministic = false;
+	bool Deterministic = false;
 
 	UPROPERTY(EditAnywhere, Category = "Seed", meta = (DataflowInput, DisplayName = "RandomSeed", EditCondition = "Deterministic"))
-		float RandomSeed = FMath::FRandRange(-1e5, 1e5);
-
+	float RandomSeed = 0.f;
+	
 	UPROPERTY(meta = (DataflowOutput, DisplayName = "Vector"))
-		FVector Vector = FVector(0.0);
+	FVector Vector = FVector(0.0);
 
 	FRandomUnitVectorDataflowNode(const Dataflow::FNodeParameters& InParam, FGuid InGuid = FGuid::NewGuid())
 		: FDataflowNode(InParam, InGuid)
 	{
+		RandomSeed = FMath::FRandRange(-1e5, 1e5);
 		RegisterInputConnection(&RandomSeed);
 		RegisterOutputConnection(&Vector);
 	}
@@ -858,27 +861,28 @@ USTRUCT()
 struct FRandomUnitVectorInConeDataflowNode : public FDataflowNode
 {
 	GENERATED_USTRUCT_BODY()
-		DATAFLOW_NODE_DEFINE_INTERNAL(FRandomUnitVectorInConeDataflowNode, "RandomUnitVectorInCone", "Math|Random", "")
+	DATAFLOW_NODE_DEFINE_INTERNAL(FRandomUnitVectorInConeDataflowNode, "RandomUnitVectorInCone", "Math|Random", "")
 
 public:
 	UPROPERTY(EditAnywhere, Category = "Seed")
-		bool Deterministic = false;
+	bool Deterministic = false;
 
 	UPROPERTY(EditAnywhere, Category = "Seed", meta = (DataflowInput, DisplayName = "RandomSeed", EditCondition = "Deterministic"))
-		float RandomSeed = FMath::FRandRange(-1e5, 1e5);
+	float RandomSeed = 0.f;
 
 	UPROPERTY(EditAnywhere, Category = "Random", meta = (DataflowInput, DisplayName = "ConeDirection"))
-		FVector ConeDirection = FVector(0.0, 0.0, 1.0);
+	FVector ConeDirection = FVector(0.0, 0.0, 1.0);
 
 	UPROPERTY(EditAnywhere, Category = "Random", meta = (DataflowInput, DisplayName = "ConeHalfAngle"))
-		float ConeHalfAngle = PI / 4.f;
+	float ConeHalfAngle = PI / 4.f;
 
 	UPROPERTY(meta = (DataflowOutput, DisplayName = "Vector"))
-		FVector Vector = FVector(0.0);
+	FVector Vector = FVector(0.0);
 
 	FRandomUnitVectorInConeDataflowNode(const Dataflow::FNodeParameters& InParam, FGuid InGuid = FGuid::NewGuid())
 		: FDataflowNode(InParam, InGuid)
 	{
+		RandomSeed = FMath::FRandRange(-1e5, 1e5);
 		RegisterInputConnection(&RandomSeed);
 		RegisterInputConnection(&ConeDirection);
 		RegisterInputConnection(&ConeHalfAngle);
@@ -893,14 +897,14 @@ USTRUCT()
 struct FRadiansToDegreesDataflowNode : public FDataflowNode
 {
 	GENERATED_USTRUCT_BODY()
-		DATAFLOW_NODE_DEFINE_INTERNAL(FRadiansToDegreesDataflowNode, "RadiansToDegrees", "Math|Trigonometry", "")
+	DATAFLOW_NODE_DEFINE_INTERNAL(FRadiansToDegreesDataflowNode, "RadiansToDegrees", "Math|Trigonometry", "")
 
 public:
 	UPROPERTY(EditAnywhere, Category = "Radians", meta = (DataflowInput, DisplayName = "Radians"))
-		float Radians = 0.f;
+	float Radians = 0.f;
 
 	UPROPERTY(meta = (DataflowOutput, DisplayName = "Degrees"))
-		float Degrees = 0.f;
+	float Degrees = 0.f;
 
 	FRadiansToDegreesDataflowNode(const Dataflow::FNodeParameters& InParam, FGuid InGuid = FGuid::NewGuid())
 		: FDataflowNode(InParam, InGuid)
@@ -917,14 +921,14 @@ USTRUCT()
 struct FDegreesToRadiansDataflowNode : public FDataflowNode
 {
 	GENERATED_USTRUCT_BODY()
-		DATAFLOW_NODE_DEFINE_INTERNAL(FDegreesToRadiansDataflowNode, "DegreesToRadians", "Math|Trigonometry", "")
+	DATAFLOW_NODE_DEFINE_INTERNAL(FDegreesToRadiansDataflowNode, "DegreesToRadians", "Math|Trigonometry", "")
 
 public:
 	UPROPERTY(EditAnywhere, Category = "Degrees", meta = (DataflowInput, DisplayName = "Degrees"))
-		float Degrees = 0.f;
+	float Degrees = 0.f;
 
 	UPROPERTY(meta = (DataflowOutput, DisplayName = "Radians"))
-		float Radians = 0.f;
+	float Radians = 0.f;
 
 	FDegreesToRadiansDataflowNode(const Dataflow::FNodeParameters& InParam, FGuid InGuid = FGuid::NewGuid())
 		: FDataflowNode(InParam, InGuid)
@@ -941,17 +945,17 @@ USTRUCT()
 struct FExplodedViewDataflowNode : public FDataflowNode
 {
 	GENERATED_USTRUCT_BODY()
-		DATAFLOW_NODE_DEFINE_INTERNAL(FExplodedViewDataflowNode, "ExplodedView", "Fracture|Utilities", "")
+	DATAFLOW_NODE_DEFINE_INTERNAL(FExplodedViewDataflowNode, "ExplodedView", "Fracture|Utilities", "")
 
 public:
 	UPROPERTY(meta = (DataflowInput, DataflowOutput, DisplayName = "Collection"))
-		FManagedArrayCollection Collection;
+	FManagedArrayCollection Collection;
 
 	UPROPERTY(EditAnywhere, Category = "Scale", meta = (DataflowInput, DisplayName = "UniformScale"))
-		float UniformScale = 1.f;
+	float UniformScale = 1.f;
 
 	UPROPERTY(EditAnywhere, Category = "Scale", meta = (DataflowInput, DisplayName = "Scale"))
-		FVector Scale = FVector(1.0);
+	FVector Scale = FVector(1.0);
 
 	FExplodedViewDataflowNode(const Dataflow::FNodeParameters& InParam, FGuid InGuid = FGuid::NewGuid())
 		: FDataflowNode(InParam, InGuid)
@@ -974,20 +978,20 @@ USTRUCT()
 struct FCreateNonOverlappingConvexHullsDataflowNode : public FDataflowNode
 {
 	GENERATED_USTRUCT_BODY()
-		DATAFLOW_NODE_DEFINE_INTERNAL(FCreateNonOverlappingConvexHullsDataflowNode, "CreateNonOverlappingConvexHulls", "Fracture|Utilities", "")
+	DATAFLOW_NODE_DEFINE_INTERNAL(FCreateNonOverlappingConvexHullsDataflowNode, "CreateNonOverlappingConvexHulls", "Fracture|Utilities", "")
 
 public:
 	UPROPERTY(meta = (DataflowInput, DataflowOutput, DisplayName = "Collection"))
-		FManagedArrayCollection Collection;
+	FManagedArrayCollection Collection;
 
 	UPROPERTY(EditAnywhere, Category = "Convex", meta = (DataflowInput, DisplayName = "CanRemoveFraction", UIMin = 0.01f, UIMax = 1.f))
-		float CanRemoveFraction = 0.5f;
+	float CanRemoveFraction = 0.5f;
 
 	UPROPERTY(EditAnywhere, Category = "Convex", meta = (DataflowInput, DisplayName = "CanExceedFraction", UIMin = 0.f))
-		float CanExceedFraction = 0.5f;
+	float CanExceedFraction = 0.5f;
 
 	UPROPERTY(EditAnywhere, Category = "Convex", meta = (DataflowInput, DisplayName = "SimplificationDistanceThreshold", UIMin = 0.f))
-		float SimplificationDistanceThreshold = 10.f;
+	float SimplificationDistanceThreshold = 10.f;
 
 	FCreateNonOverlappingConvexHullsDataflowNode(const Dataflow::FNodeParameters& InParam, FGuid InGuid = FGuid::NewGuid())
 		: FDataflowNode(InParam, InGuid)
@@ -1007,23 +1011,23 @@ USTRUCT()
 struct FPlaneCutterDataflowNode : public FDataflowNode
 {
 	GENERATED_USTRUCT_BODY()
-		DATAFLOW_NODE_DEFINE_INTERNAL(FPlaneCutterDataflowNode, "PlaneCutter", "Fracture", "")
+	DATAFLOW_NODE_DEFINE_INTERNAL(FPlaneCutterDataflowNode, "PlaneCutter", "Fracture", "")
 
 public:
 	UPROPERTY(meta = (DataflowInput, DataflowOutput, DisplayName = "Collection"))
-		FManagedArrayCollection Collection;
+	FManagedArrayCollection Collection;
 
 	UPROPERTY(meta = (DataflowInput, DisplayName = "BoundingBox"))
-		FBox BoundingBox = FBox(ForceInit);
+	FBox BoundingBox = FBox(ForceInit);
 
 	UPROPERTY(EditAnywhere, Category = "Fracture", meta = (DataflowInput, DisplayName = "NumPlanes", UIMin = 1))
-		int32 NumPlanes = 1;
+	int32 NumPlanes = 1;
 
 	UPROPERTY(EditAnywhere, Category = "Fracture", meta = (DataflowInput, DisplayName = "RandomSeed"))
-		float RandomSeed = -1.f;
+	float RandomSeed = -1.f;
 
 	UPROPERTY(EditAnywhere, Category = "Fracture", meta = (DataflowInput, DisplayName = "Grout", UIMin = 0.f))
-		float Grout = 0.f;
+	float Grout = 0.f;
 
 	UPROPERTY(EditAnywhere, Category = "Noise", meta = (DataflowInput, DisplayName = "Amplitude", UIMin = 0.f));
 	float Amplitude = 0.f;
@@ -1079,10 +1083,10 @@ struct FHashStringDataflowNode : public FDataflowNode
 
 public:
 	UPROPERTY(meta = (DataflowInput, DisplayName = "String"))
-		FString String;
+	FString String = FString("");
 
 	UPROPERTY(meta = (DataflowOutput, DisplayName = "Hash"))
-		int32 Hash = 0;
+	int32 Hash = 0;
 
 
 	FHashStringDataflowNode(const Dataflow::FNodeParameters& InParam, FGuid InGuid = FGuid::NewGuid())
@@ -1100,14 +1104,14 @@ USTRUCT()
 struct FHashVectorDataflowNode : public FDataflowNode
 {
 	GENERATED_USTRUCT_BODY()
-		DATAFLOW_NODE_DEFINE_INTERNAL(FHashVectorDataflowNode, "HashVector", "Utilities|Vector", "")
+	DATAFLOW_NODE_DEFINE_INTERNAL(FHashVectorDataflowNode, "HashVector", "Utilities|Vector", "")
 
 public:
 	UPROPERTY(meta = (DataflowInput, DisplayName = "Vector"))
-		FVector Vector;
+	FVector Vector = FVector(0.0);
 
 	UPROPERTY(meta = (DataflowOutput, DisplayName = "Hash"))
-		int32 Hash = 0;
+	int32 Hash = 0;
 
 
 	FHashVectorDataflowNode(const Dataflow::FNodeParameters& InParam, FGuid InGuid = FGuid::NewGuid())
@@ -1144,10 +1148,10 @@ public:
 	EFloatToIntFunctionEnum Function = EFloatToIntFunctionEnum::Dataflow_FloatToInt_Function_Round;
 
 	UPROPERTY(EditAnywhere, Category = "Float", meta = (DataflowInput, DisplayName = "Float"))
-		float Float = 0.f;
+	float Float = 0.f;
 
 	UPROPERTY(meta = (DataflowOutput, DisplayName = "Int"))
-		int32 Int = 0;
+	int32 Int = 0;
 
 	FFloatToIntDataflowNode(const Dataflow::FNodeParameters& InParam, FGuid InGuid = FGuid::NewGuid())
 		: FDataflowNode(InParam, InGuid)
@@ -1186,14 +1190,14 @@ USTRUCT()
 struct FMathConstantsDataflowNode : public FDataflowNode
 {
 	GENERATED_USTRUCT_BODY()
-		DATAFLOW_NODE_DEFINE_INTERNAL(FMathConstantsDataflowNode, "MathConstants", "Math|Utilities", "")
+	DATAFLOW_NODE_DEFINE_INTERNAL(FMathConstantsDataflowNode, "MathConstants", "Math|Utilities", "")
 
 public:
 	UPROPERTY(EditAnywhere, Category = "Constants", meta = (DisplayName = "Constant"));
 	EMathConstantsEnum Constant = EMathConstantsEnum::Dataflow_MathConstants_Pi;
 
 	UPROPERTY(meta = (DataflowOutput, DisplayName = "Float"))
-		float Float = 0;
+	float Float = 0;
 
 	FMathConstantsDataflowNode(const Dataflow::FNodeParameters& InParam, FGuid InGuid = FGuid::NewGuid())
 		: FDataflowNode(InParam, InGuid)
