ChaosFlesh- Tetarahedron generation node

diff --git a/Engine/Plugins/Experimental/GeometryCollectionPlugin/Source/GeometryCollectionNodes/Private/Dataflow/GeometryCollectionNodesPlugin.cpp b/Engine/Plugins/Experimental/GeometryCollectionPlugin/Source/GeometryCollectionNodes/Private/Dataflow/GeometryCollectionNodesPlugin.cpp
--- a/Engine/Plugins/Experimental/GeometryCollectionPlugin/Source/GeometryCollectionNodes/Private/Dataflow/GeometryCollectionNodesPlugin.cpp
+++ b/Engine/Plugins/Experimental/GeometryCollectionPlugin/Source/GeometryCollectionNodes/Private/Dataflow/GeometryCollectionNodesPlugin.cpp
@@ -8,7 +8,6 @@
 #include "Dataflow/GeometryCollectionNodes.h"
 #include "Dataflow/GeometryCollectionProcessingNodes.h"
 #include "Dataflow/GeometryCollectionSkeletalMeshNodes.h"
-#include "Dataflow/GeometryCollectionTetrahedralNodes.h"
 
 #define LOCTEXT_NAMESPACE "DataflowNodes"
 
@@ -18,7 +17,6 @@ void IGeometryCollectionNodesPlugin::StartupModule()
 	Dataflow::GeometryCollectionEngineAssetNodes();
 	Dataflow::GeometryCollectionProcessingNodes();
 	Dataflow::GeometryCollectionSkeletalMeshNodes();
-	Dataflow::GeometryCollectionTetrahedralNodes();
 }
 
 void IGeometryCollectionNodesPlugin::ShutdownModule()
diff --git a/Engine/Plugins/Experimental/GeometryCollectionPlugin/Source/GeometryCollectionNodes/Private/Dataflow/GeometryCollectionTetrahedralNodes.cpp b/Engine/Plugins/Experimental/GeometryCollectionPlugin/Source/GeometryCollectionNodes/Private/Dataflow/GeometryCollectionTetrahedralNodes.cpp
deleted file mode 100644
diff --git a/Engine/Plugins/Experimental/GeometryCollectionPlugin/Source/GeometryCollectionNodes/Public/Dataflow/GeometryCollectionTetrahedralNodes.h b/Engine/Plugins/Experimental/GeometryCollectionPlugin/Source/GeometryCollectionNodes/Public/Dataflow/GeometryCollectionTetrahedralNodes.h
deleted file mode 100644
