Fix GeometryCollection slate icons missing from cooked editors

diff --git a/Engine/Plugins/Experimental/GeometryCollectionPlugin/Source/GeometryCollectionEditor/Private/GeometryCollection/GeometryCollectionEditorStyle.h b/Engine/Plugins/Experimental/GeometryCollectionPlugin/Source/GeometryCollectionEditor/Private/GeometryCollection/GeometryCollectionEditorStyle.h
--- a/Engine/Plugins/Experimental/GeometryCollectionPlugin/Source/GeometryCollectionEditor/Private/GeometryCollection/GeometryCollectionEditorStyle.h
+++ b/Engine/Plugins/Experimental/GeometryCollectionPlugin/Source/GeometryCollectionEditor/Private/GeometryCollection/GeometryCollectionEditorStyle.h
@@ -12,9 +12,7 @@ public:
 		const FVector2D Icon16x16(16.f, 16.f);
 		const FVector2D Icon64x64(64.f, 64.f);
 
-#if !IS_MONOLITHIC
 		SetContentRoot(FPaths::EnginePluginsDir() / TEXT("Experimental/GeometryCollectionPlugin/Resources"));
-#endif
 
 		Set("ClassIcon.GeometryCollection", new FSlateImageBrush(RootToContentDir(TEXT("GeometryCollection_16x.png")), Icon16x16));
 		Set("ClassThumbnail.GeometryCollection", new FSlateImageBrush(RootToContentDir(TEXT("GeometryCollection_64x.png")), Icon64x64));
