Chaos : fix potential crash with generating geometry collection thumbnails

diff --git a/Engine/Plugins/Experimental/GeometryCollectionPlugin/Source/GeometryCollectionEditor/Private/GeometryCollection/GeometryCollectionThumbnailScene.cpp b/Engine/Plugins/Experimental/GeometryCollectionPlugin/Source/GeometryCollectionEditor/Private/GeometryCollection/GeometryCollectionThumbnailScene.cpp
--- a/Engine/Plugins/Experimental/GeometryCollectionPlugin/Source/GeometryCollectionEditor/Private/GeometryCollection/GeometryCollectionThumbnailScene.cpp
+++ b/Engine/Plugins/Experimental/GeometryCollectionPlugin/Source/GeometryCollectionEditor/Private/GeometryCollection/GeometryCollectionThumbnailScene.cpp
@@ -24,19 +24,22 @@ FGeometryCollectionThumbnailScene::FGeometryCollectionThumbnailScene()
 
 void FGeometryCollectionThumbnailScene::SetGeometryCollection(UGeometryCollection* GeometryCollection)
 {
-	PreviewActor->GetGeometryCollectionComponent()->SetRestCollection(GeometryCollection);
-
-	if (GeometryCollection)
+	if (PreviewActor->GetGeometryCollectionComponent())
 	{
-		FTransform MeshTransform = FTransform::Identity;
+		PreviewActor->GetGeometryCollectionComponent()->SetRestCollection(GeometryCollection);
+
+		if (GeometryCollection)
+		{
+			FTransform MeshTransform = FTransform::Identity;
 
-		PreviewActor->SetActorLocation(FVector(0, 0, 0), false);
-		PreviewActor->GetGeometryCollectionComponent()->UpdateBounds();
+			PreviewActor->SetActorLocation(FVector(0, 0, 0), false);
+			PreviewActor->GetGeometryCollectionComponent()->UpdateBounds();
 
-		// Center the mesh at the world origin then offset to put it on top of the plane
-		const float BoundsZOffset = GetBoundsZOffset(PreviewActor->GetGeometryCollectionComponent()->Bounds);
-		PreviewActor->SetActorLocation(-PreviewActor->GetGeometryCollectionComponent()->Bounds.Origin + FVector(0, 0, BoundsZOffset), false);
-		PreviewActor->GetGeometryCollectionComponent()->RecreateRenderState_Concurrent();
+			// Center the mesh at the world origin then offset to put it on top of the plane
+			const float BoundsZOffset = GetBoundsZOffset(PreviewActor->GetGeometryCollectionComponent()->Bounds);
+			PreviewActor->SetActorLocation(-PreviewActor->GetGeometryCollectionComponent()->Bounds.Origin + FVector(0, 0, BoundsZOffset), false);
+			PreviewActor->GetGeometryCollectionComponent()->RecreateRenderState_Concurrent();
+		}
 	}
 }
 
