GeometryCollectionEditor: Persistent dataflow context with invalidation.

diff --git a/Engine/Plugins/Experimental/GeometryCollectionPlugin/Source/GeometryCollectionEditor/Private/GeometryCollection/GeometryCollectionEditorToolkit.cpp b/Engine/Plugins/Experimental/GeometryCollectionPlugin/Source/GeometryCollectionEditor/Private/GeometryCollection/GeometryCollectionEditorToolkit.cpp
--- a/Engine/Plugins/Experimental/GeometryCollectionPlugin/Source/GeometryCollectionEditor/Private/GeometryCollection/GeometryCollectionEditorToolkit.cpp
+++ b/Engine/Plugins/Experimental/GeometryCollectionPlugin/Source/GeometryCollectionEditor/Private/GeometryCollection/GeometryCollectionEditorToolkit.cpp
@@ -23,10 +23,14 @@
 #include "IStructureDetailsView.h"
 
 #define LOCTEXT_NAMESPACE "GeometryCollectionEditorToolkit"
-
 //DEFINE_LOG_CATEGORY_STATIC(FGeometryCollectionEditorToolkitLog, Log, All);
 
 
+int32 bDataflowAssetEditorGeometryCollectionLiveEvaluationEnableCVar = 1;
+FAutoConsoleVariableRef CVarDataflowAssetEditorGeometryCollectionEnableLiveEvaluation(TEXT("p.Dataflow.AssetEditor.GeometryCollection.LiveEvaluation.Enable"),
+	bDataflowAssetEditorGeometryCollectionLiveEvaluationEnableCVar,
+	TEXT("Enable live evaluation of specified output on the GeometryCollectionAsset within the GeometryCollection Dataflow Editor.[def:true]"));
+
 const FName FGeometryCollectionEditorToolkit::GraphCanvasTabId(TEXT("GeometryCollectionEditor_GraphCanvas"));
 const FName FGeometryCollectionEditorToolkit::AssetDetailsTabId(TEXT("GeometryCollectionEditor_AssetDetails"));
 const FName FGeometryCollectionEditorToolkit::NodeDetailsTabId(TEXT("GeometryCollectionEditor_NodeDetails"));
@@ -50,6 +54,9 @@ void FGeometryCollectionEditorToolkit::InitGeometryCollectionAssetEditor(const E
 		AssetDetailsEditor = CreateAssetDetailsEditorWidget(GeometryCollection);
 		GraphEditor = CreateGraphEditorWidget(Dataflow, NodeDetailsEditor);
 
+		Context = TSharedPtr< Dataflow::FEngineContext>(new Dataflow::FEngineContext(GeometryCollection, Dataflow, FPlatformTime::Cycles64(), FString("UFleshAsset")));
+		LastNodeTimestamp = Context->GetTimestamp();
+
 		const TSharedRef<FTabManager::FLayout> StandaloneDefaultLayout = FTabManager::NewLayout("GeometryCollectionDataflowEditor_Layout.V1")
 			->AddArea
 			(
@@ -88,6 +95,33 @@ void FGeometryCollectionEditorToolkit::InitGeometryCollectionAssetEditor(const E
 	}
 }
 
+void FGeometryCollectionEditorToolkit::OnPropertyValueChanged(const FPropertyChangedEvent& PropertyChangedEvent)
+{
+	FDataflowEditorCommands::OnPropertyValueChanged(this->GetDataflow(), Context, LastNodeTimestamp, PropertyChangedEvent);
+}
+
+void FGeometryCollectionEditorToolkit::Tick(float DeltaTime)
+{
+	if (bDataflowAssetEditorGeometryCollectionLiveEvaluationEnableCVar)
+	{
+		if (Dataflow && GeometryCollection)
+		{
+			if (!Context)
+			{
+				Context = TSharedPtr< Dataflow::FEngineContext>(new Dataflow::FEngineContext(GeometryCollection, Dataflow, Dataflow::FTimestamp::Invalid, FString("UGeometryCollection")));
+				LastNodeTimestamp = Dataflow::FTimestamp::Invalid;
+			}
+
+			FDataflowEditorCommands::EvaluateNode(*Context.Get(), LastNodeTimestamp, Dataflow, nullptr, nullptr, GeometryCollection->Terminal);
+		}
+	}
+}
+
+TStatId FGeometryCollectionEditorToolkit::GetStatId() const
+{
+	RETURN_QUICK_DECLARE_CYCLE_STAT(FGeometryCollectionEditorToolkit, STATGROUP_Tickables);
+}
+
 TSharedRef<SGraphEditor> FGeometryCollectionEditorToolkit::CreateGraphEditorWidget(UDataflow* DataflowToEdit, TSharedPtr<IStructureDetailsView> InNodeDetailsEditor)
 {
 	ensure(DataflowToEdit);
@@ -96,7 +130,16 @@ TSharedRef<SGraphEditor> FGeometryCollectionEditorToolkit::CreateGraphEditorWidg
 
 	FDataflowEditorCommands::FGraphEvaluationCallback Evaluate = [&](FDataflowNode* Node, FDataflowOutput* Out)
 	{
-		FEngineContext(GeometryCollection, Dataflow, FPlatformTime::Cycles64(), FString("UGeometryCollection")).Evaluate(Node, Out);
+		if (Dataflow && GeometryCollection)
+		{
+			if (!Context)
+			{
+				Context = TSharedPtr< Dataflow::FEngineContext>(new FEngineContext(GeometryCollection, Dataflow, Dataflow::FTimestamp::Invalid, FString("UGeometryCollection")));
+				LastNodeTimestamp = Dataflow::FTimestamp::Invalid;
+			}
+
+			FDataflowEditorCommands::EvaluateNode(*Context.Get(), LastNodeTimestamp, Dataflow, Node, Out);
+		}
 	};
 
 	return SNew(SDataflowGraphEditor, GeometryCollection)
@@ -132,6 +175,7 @@ TSharedPtr<IStructureDetailsView> FGeometryCollectionEditorToolkit::CreateNodeDe
 	}
 	TSharedPtr<IStructureDetailsView> DetailsView = PropertyEditorModule.CreateStructureDetailView(DetailsViewArgs, StructureViewArgs, nullptr);
 	DetailsView->GetDetailsView()->SetObject(ObjectToEdit);
+	DetailsView->GetOnFinishedChangingPropertiesDelegate().AddSP(this, &FGeometryCollectionEditorToolkit::OnPropertyValueChanged);
 
 	return DetailsView;
 
diff --git a/Engine/Plugins/Experimental/GeometryCollectionPlugin/Source/GeometryCollectionEditor/Public/GeometryCollection/GeometryCollectionEditorToolkit.h b/Engine/Plugins/Experimental/GeometryCollectionPlugin/Source/GeometryCollectionEditor/Public/GeometryCollection/GeometryCollectionEditorToolkit.h
--- a/Engine/Plugins/Experimental/GeometryCollectionPlugin/Source/GeometryCollectionEditor/Public/GeometryCollection/GeometryCollectionEditorToolkit.h
+++ b/Engine/Plugins/Experimental/GeometryCollectionPlugin/Source/GeometryCollectionEditor/Public/GeometryCollection/GeometryCollectionEditorToolkit.h
@@ -7,8 +7,11 @@
 #include "Dataflow/DataflowGraphEditor.h"
 #include "GeometryCollection/GeometryCollectionObject.h"
 #include "GraphEditor.h"
-#include "Toolkits/SimpleAssetEditor.h"
 #include "Misc/NotifyHook.h"
+#include "TickableEditorObject.h"
+#include "Toolkits/SimpleAssetEditor.h"
+
+
 
 class IDetailsView;
 class FTabManager;
@@ -16,10 +19,11 @@ class IToolkitHost;
 class UDataflow;
 
 
-class FGeometryCollectionEditorToolkit : public FAssetEditorToolkit, public FNotifyHook, public FGCObject
+class FGeometryCollectionEditorToolkit : public FAssetEditorToolkit, public FTickableEditorObject, public FNotifyHook, public FGCObject
 {
 public:
 	void InitGeometryCollectionAssetEditor(const EToolkitMode::Type Mode, const TSharedPtr<IToolkitHost>& InitToolkitHost, UObject* ObjectToEdit);
+	void OnPropertyValueChanged(const FPropertyChangedEvent& PropertyChangedEvent);
 
 	// IToolkit Interface
 	virtual FName GetToolkitFName() const override;
@@ -52,6 +56,12 @@ public:
 	TSharedPtr<SGraphEditor> GetGraphEditor() { return GraphEditor; }
 	const TSharedPtr<SGraphEditor> GetGraphEditor() const { return GraphEditor; }
 
+	//~ Begin FTickableEditorObject interface
+	virtual void Tick(float DeltaTime) override;
+	virtual bool IsTickable() const override { return true; }
+	virtual TStatId GetStatId() const override;
+	//~ End FTickableEditorObject Interface
+
 private:
 	static const FName GraphCanvasTabId;
 	TSharedPtr<SGraphEditor> GraphEditor;
@@ -67,4 +77,7 @@ private:
 
 	UDataflow* Dataflow = nullptr;
 	UGeometryCollection* GeometryCollection = nullptr;
+
+	TSharedPtr<Dataflow::FEngineContext> Context;
+	Dataflow::FTimestamp LastNodeTimestamp = Dataflow::FTimestamp::Invalid;
 };
diff --git a/Engine/Source/Runtime/Experimental/GeometryCollectionEngine/Public/GeometryCollection/GeometryCollectionObject.h b/Engine/Source/Runtime/Experimental/GeometryCollectionEngine/Public/GeometryCollection/GeometryCollectionObject.h
--- a/Engine/Source/Runtime/Experimental/GeometryCollectionEngine/Public/GeometryCollection/GeometryCollectionObject.h
+++ b/Engine/Source/Runtime/Experimental/GeometryCollectionEngine/Public/GeometryCollection/GeometryCollectionObject.h
@@ -654,9 +654,13 @@ public:
 	//
 	// Dataflow
 	//
-	UPROPERTY(EditAnywhere, Category = "Procedural")
+	UPROPERTY(EditAnywhere, Category = "Dataflow")
 	TObjectPtr<UDataflow> Dataflow;
 
+	UPROPERTY(EditAnywhere, Category = "Dataflow")
+	FString Terminal;
+
+
 private:
 #if WITH_EDITOR
 	void CreateSimulationDataImp(bool bCopyFromDDC);
