Chaos : Dataflow : add geometry collection asset specific nodes

diff --git a/Engine/Plugins/Experimental/GeometryCollectionPlugin/Source/GeometryCollectionNodes/Private/Dataflow/GeometryCollectionNodes.cpp b/Engine/Plugins/Experimental/GeometryCollectionPlugin/Source/GeometryCollectionNodes/Private/Dataflow/GeometryCollectionNodes.cpp
--- a/Engine/Plugins/Experimental/GeometryCollectionPlugin/Source/GeometryCollectionNodes/Private/Dataflow/GeometryCollectionNodes.cpp
+++ b/Engine/Plugins/Experimental/GeometryCollectionPlugin/Source/GeometryCollectionNodes/Private/Dataflow/GeometryCollectionNodes.cpp
@@ -28,7 +28,7 @@ namespace Dataflow
 {
 	static const FLinearColor CDefaultNodeBodyTintColor = FLinearColor(0.f, 0.f, 0.f, 0.5f);
 
-	void GeometryCollectionEngineAssetNodes()
+	void GeometryCollectionEngineNodes()
 	{
 		DATAFLOW_NODE_REGISTER_CREATION_FACTORY(FGetCollectionAssetDataflowNode);
 		DATAFLOW_NODE_REGISTER_CREATION_FACTORY(FExampleCollectionEditDataflowNode);
diff --git a/Engine/Plugins/Experimental/GeometryCollectionPlugin/Source/GeometryCollectionNodes/Private/Dataflow/GeometryCollectionNodesPlugin.cpp b/Engine/Plugins/Experimental/GeometryCollectionPlugin/Source/GeometryCollectionNodes/Private/Dataflow/GeometryCollectionNodesPlugin.cpp
--- a/Engine/Plugins/Experimental/GeometryCollectionPlugin/Source/GeometryCollectionNodes/Private/Dataflow/GeometryCollectionNodesPlugin.cpp
+++ b/Engine/Plugins/Experimental/GeometryCollectionPlugin/Source/GeometryCollectionNodes/Private/Dataflow/GeometryCollectionNodesPlugin.cpp
@@ -6,6 +6,7 @@
 
 #include "Dataflow/DataflowNodeFactory.h"
 #include "Dataflow/GeometryCollectionNodes.h"
+#include "Dataflow/Nodes/GeometryCollectionAssetNodes.h"
 #include "Dataflow/GeometryCollectionProcessingNodes.h"
 #include "Dataflow/GeometryCollectionSkeletalMeshNodes.h"
 
@@ -14,6 +15,7 @@
 
 void IGeometryCollectionNodesPlugin::StartupModule()
 {
+	Dataflow::GeometryCollectionEngineNodes();
 	Dataflow::GeometryCollectionEngineAssetNodes();
 	Dataflow::GeometryCollectionProcessingNodes();
 	Dataflow::GeometryCollectionSkeletalMeshNodes();
diff --git a/Engine/Plugins/Experimental/GeometryCollectionPlugin/Source/GeometryCollectionNodes/Private/Dataflow/Nodes/GeometryCollectionAssetNodes.cpp b/Engine/Plugins/Experimental/GeometryCollectionPlugin/Source/GeometryCollectionNodes/Private/Dataflow/Nodes/GeometryCollectionAssetNodes.cpp
new file mode 100644
--- /dev/null
+++ b/Engine/Plugins/Experimental/GeometryCollectionPlugin/Source/GeometryCollectionNodes/Private/Dataflow/Nodes/GeometryCollectionAssetNodes.cpp
@@ -0,0 +1,124 @@
+// Copyright Epic Games, Inc. All Rights Reserved.
+
+#include "Dataflow/Nodes/GeometryCollectionAssetNodes.h"
+#include "Dataflow/DataflowCore.h"
+
+#include "GeometryCollection/GeometryCollection.h"
+#include "GeometryCollection/GeometryCollectionEngineConversion.h"
+#include "GeometryCollection/GeometryCollectionObject.h"
+
+namespace Dataflow
+{
+	void GeometryCollectionEngineAssetNodes()
+	{
+		DATAFLOW_NODE_REGISTER_CREATION_FACTORY(FSetGeometryCollectionAssetDataflowNode);
+		DATAFLOW_NODE_REGISTER_CREATION_FACTORY(FGetGeometryCollectionAssetDataflowNode);
+		DATAFLOW_NODE_REGISTER_CREATION_FACTORY(FGetGeometryCollectionSourcesDataflowNode);
+		DATAFLOW_NODE_REGISTER_CREATION_FACTORY(FCreateGeometryCollectionFromSourcesDataflowNode);
+	}
+}
+
+
+// ===========================================================================================================================
+
+FSetGeometryCollectionAssetDataflowNode::FSetGeometryCollectionAssetDataflowNode(const Dataflow::FNodeParameters& InParam, FGuid InGuid)
+	: FFinalDataflowNode(InParam, InGuid)
+{
+	RegisterInputConnection(&Collection);
+	RegisterInputConnection(&Materials);
+}
+
+void FSetGeometryCollectionAssetDataflowNode::EvaluateFinal(Dataflow::FContext& Context) const
+{
+	using FGeometryCollectionPtr = TSharedPtr<FGeometryCollection, ESPMode::ThreadSafe>;
+	using FMaterialArray = TArray<TObjectPtr<UMaterial>>;
+
+	// todo(chaos) : should certainly have some form of errors handling  
+	if (const Dataflow::FEngineContext* EngineContext = Context.AsType<Dataflow::FEngineContext>())
+	{
+		if (UGeometryCollection* CollectionAsset = Cast<UGeometryCollection>(EngineContext->Owner))
+		{
+			if (FGeometryCollectionPtr GeometryCollection = CollectionAsset->GetGeometryCollection())
+			{
+				const FManagedArrayCollection& InCollection = GetValue<FManagedArrayCollection>(Context, &Collection);
+				const FMaterialArray& InMaterials = GetValue<FMaterialArray>(Context, &Materials);
+
+				CollectionAsset->ResetFrom(InCollection, InMaterials);
+			}
+		}
+	}
+}
+
+// ===========================================================================================================================
+
+FGetGeometryCollectionAssetDataflowNode::FGetGeometryCollectionAssetDataflowNode(const Dataflow::FNodeParameters& InParam, FGuid InGuid)
+	: FDataflowNode(InParam, InGuid)
+{
+	RegisterOutputConnection(&Asset);
+}
+
+void FGetGeometryCollectionAssetDataflowNode::Evaluate(Dataflow::FContext& Context, const FDataflowOutput* Out) const
+{
+	ensure(Out->IsA(&Asset));
+	if (const Dataflow::FEngineContext* EngineContext = Context.AsType<Dataflow::FEngineContext>())
+	{
+		if (const TObjectPtr<UGeometryCollection> CollectionAsset = Cast<UGeometryCollection>(EngineContext->Owner))
+		{
+			SetValue(Context, CollectionAsset, &Asset);
+		}
+	}
+}
+
+// ===========================================================================================================================
+
+FGetGeometryCollectionSourcesDataflowNode::FGetGeometryCollectionSourcesDataflowNode(const Dataflow::FNodeParameters& InParam, FGuid InGuid)
+		: FDataflowNode(InParam, InGuid)
+{
+	RegisterInputConnection(&Asset);
+	RegisterOutputConnection(&Sources);
+}
+
+void FGetGeometryCollectionSourcesDataflowNode::Evaluate(Dataflow::FContext& Context, const FDataflowOutput* Out) const
+{
+	ensure(Out->IsA(&Sources));
+
+	TArray<FGeometryCollectionSource> OutSources;
+	
+	if (const TObjectPtr<UGeometryCollection> InAsset = GetValue(Context, &Asset))
+	{
+		OutSources = InAsset->GeometrySource; 
+	}
+
+	SetValue(Context, OutSources, &Sources);
+}
+
+// ===========================================================================================================================
+
+FCreateGeometryCollectionFromSourcesDataflowNode::FCreateGeometryCollectionFromSourcesDataflowNode(const Dataflow::FNodeParameters& InParam, FGuid InGuid)
+		: FDataflowNode(InParam, InGuid)
+{
+	RegisterInputConnection(&Sources);
+	RegisterOutputConnection(&Collection);
+	RegisterOutputConnection(&Materials);
+}
+
+void FCreateGeometryCollectionFromSourcesDataflowNode::Evaluate(Dataflow::FContext& Context, const FDataflowOutput* Out) const
+{
+	ensure(Out->IsA(&Collection) || Out->IsA(&Materials));
+	
+	using FGeometryCollectionSourceArray = TArray<FGeometryCollectionSource>;
+	const FGeometryCollectionSourceArray& InSources = GetValue<FGeometryCollectionSourceArray>(Context, &Sources);
+
+	FGeometryCollection OutCollection;
+	TArray<TObjectPtr<UMaterial>> OutMaterials;
+	for (const FGeometryCollectionSource& Source: InSources)
+	{
+		// todo: change AppendGeometryCollectionSource to take a FManagedArrayCollection so we could move the collection when assigning it to the output
+		FGeometryCollectionEngineConversion::AppendGeometryCollectionSource(Source, OutCollection, OutMaterials);
+	}
+
+	// we have to make a copy since we have generated a FGeometryCollection which is inherited from FManagedArrayCollection
+
+	SetValue(Context, static_cast<const FManagedArrayCollection&>(OutCollection), &Collection);
+	SetValue(Context, MoveTemp(OutMaterials), &Materials);
+}
diff --git a/Engine/Plugins/Experimental/GeometryCollectionPlugin/Source/GeometryCollectionNodes/Public/Dataflow/GeometryCollectionNodes.h b/Engine/Plugins/Experimental/GeometryCollectionPlugin/Source/GeometryCollectionNodes/Public/Dataflow/GeometryCollectionNodes.h
--- a/Engine/Plugins/Experimental/GeometryCollectionPlugin/Source/GeometryCollectionNodes/Public/Dataflow/GeometryCollectionNodes.h
+++ b/Engine/Plugins/Experimental/GeometryCollectionPlugin/Source/GeometryCollectionNodes/Public/Dataflow/GeometryCollectionNodes.h
@@ -1241,7 +1241,7 @@ public:
 
 namespace Dataflow
 {
-	void GeometryCollectionEngineAssetNodes();
+	void GeometryCollectionEngineNodes();
 
 
 }
diff --git a/Engine/Plugins/Experimental/GeometryCollectionPlugin/Source/GeometryCollectionNodes/Public/Dataflow/Nodes/GeometryCollectionAssetNodes.h b/Engine/Plugins/Experimental/GeometryCollectionPlugin/Source/GeometryCollectionNodes/Public/Dataflow/Nodes/GeometryCollectionAssetNodes.h
new file mode 100644
--- /dev/null
+++ b/Engine/Plugins/Experimental/GeometryCollectionPlugin/Source/GeometryCollectionNodes/Public/Dataflow/Nodes/GeometryCollectionAssetNodes.h
@@ -0,0 +1,145 @@
+// Copyright Epic Games, Inc. All Rights Reserved.
+
+#pragma once 
+
+#include "CoreMinimal.h"
+#include "Dataflow/DataflowEngine.h"
+#include "GeometryCollection/ManagedArrayCollection.h"
+#include "GeometryCollection/GeometryCollectionObject.h"
+
+#include "GeometryCollectionAssetNodes.generated.h"
+
+class UGeometryCollection;
+
+/**
+ * This node is a final node with no output  
+ * it is the only one that can have the right to wroite to the asset 
+ */
+USTRUCT()
+struct FFinalDataflowNode : public FDataflowNode
+{
+	GENERATED_USTRUCT_BODY()
+	DATAFLOW_NODE_DEFINE_INTERNAL(FFinalDataflowNode, "FinalDataflowNode", "DataFlow", "")
+	
+public:
+	FFinalDataflowNode(const Dataflow::FNodeParameters& InParam, FGuid InGuid = FGuid::NewGuid())
+		: FDataflowNode(InParam, InGuid)
+	{
+	}
+
+protected:
+	/** No output evaluation - secific to this type of node */
+	virtual void EvaluateFinal(Dataflow::FContext& Context) const { ensure(false); }
+	
+private:
+	virtual void Evaluate(Dataflow::FContext& Context, const FDataflowOutput* Out) const override
+	{
+		EvaluateFinal(Context);
+	};
+};
+
+
+/**
+ * Set the geometry collection asset   
+ * inputs:
+ *	- Collection : collection attributes that compose the asset
+ *	- Materials : array of materials for the geometry
+ */
+USTRUCT()
+struct FSetGeometryCollectionAssetDataflowNode : public FFinalDataflowNode
+{
+	GENERATED_USTRUCT_BODY()
+	DATAFLOW_NODE_DEFINE_INTERNAL(FSetGeometryCollectionAssetDataflowNode, "SetGeometryCollectionAsset", "GeometryCollection", "")
+
+public:
+	FSetGeometryCollectionAssetDataflowNode(const Dataflow::FNodeParameters& InParam, FGuid InGuid = FGuid::NewGuid());
+	virtual void EvaluateFinal(Dataflow::FContext& Context) const override;
+
+	/** Attribute collection to reset the asset with */ 
+	UPROPERTY(meta = (DataflowInput, DisplayName = "Collection"))
+	FManagedArrayCollection Collection;
+
+	/** Materials array to rest the asset with */
+	UPROPERTY(meta = (DataflowInput, DisplayName = "Materials"))
+	TArray<TObjectPtr<UMaterial>> Materials;
+};
+
+/**
+ * Get Current asset    
+ * outputs:
+ *	- GeometrySources : array of geometry source ( Asset, transform and material )
+ */
+USTRUCT()
+struct FGetGeometryCollectionAssetDataflowNode : public FDataflowNode
+{
+	GENERATED_USTRUCT_BODY()
+	DATAFLOW_NODE_DEFINE_INTERNAL(FGetGeometryCollectionAssetDataflowNode, "GetGeometryCollectionAsset", "GeometryCollection", "")
+
+public:
+	FGetGeometryCollectionAssetDataflowNode(const Dataflow::FNodeParameters& InParam, FGuid InGuid = FGuid::NewGuid());
+	virtual void Evaluate(Dataflow::FContext& Context, const FDataflowOutput* Out) const override;
+
+	/** Asset this data flow graph instance is assigned to */
+	UPROPERTY(meta = (DataflowOutput, DisplayName = "Asset"))
+	TObjectPtr<UGeometryCollection> Asset;
+};
+
+/**
+ * Get geometry source    
+ * outputs:
+ *	- GeometrySources : array of geometry source ( Asset, transform and material )
+ */
+USTRUCT()
+struct FGetGeometryCollectionSourcesDataflowNode : public FDataflowNode
+{
+	GENERATED_USTRUCT_BODY()
+	DATAFLOW_NODE_DEFINE_INTERNAL(FGetGeometryCollectionSourcesDataflowNode, "GetGeometryCollectionSources", "GeometryCollection", "")
+
+public:
+	FGetGeometryCollectionSourcesDataflowNode(const Dataflow::FNodeParameters& InParam, FGuid InGuid = FGuid::NewGuid());
+	virtual void Evaluate(Dataflow::FContext& Context, const FDataflowOutput* Out) const override;
+	
+	/** Asset to get geometry sources from */
+	UPROPERTY(meta = (DataflowInput, DisplayName = "Asset"))
+	TObjectPtr<UGeometryCollection> Asset;
+	
+	/** array of geometry sources */
+	UPROPERTY(meta = (DataflowOutput, DisplayName = "Sources"))
+	TArray<FGeometryCollectionSource> Sources;
+};
+
+/**
+ * create a geometry collection from a set of geometry sources    
+ * inputs:
+ *	- GeometrySources : array of geometry source ( Asset, transform and material )
+ *	outputs:
+ *	- Collection : Geometry collection newly created
+ *	- Materials : Array of materials used by the geometry collection 
+ */
+USTRUCT()
+struct FCreateGeometryCollectionFromSourcesDataflowNode : public FDataflowNode
+{
+	GENERATED_USTRUCT_BODY()
+	DATAFLOW_NODE_DEFINE_INTERNAL(FCreateGeometryCollectionFromSourcesDataflowNode, "CreateGeometryCollectionFromSources", "GeometryCollection", "")
+
+public:
+	FCreateGeometryCollectionFromSourcesDataflowNode(const Dataflow::FNodeParameters& InParam, FGuid InGuid = FGuid::NewGuid());
+	virtual void Evaluate(Dataflow::FContext& Context, const FDataflowOutput* Out) const override;
+	
+	/** array of geometry sources */
+	UPROPERTY(meta = (DataflowInput, DisplayName = "Sources"))
+	TArray<FGeometryCollectionSource> Sources;
+
+	/** Geometry collection newly created */
+	UPROPERTY(meta = (DataflowOutput, DisplayName = "Collection"))
+	FManagedArrayCollection Collection;
+
+	/** Geometry collection newly created */
+	UPROPERTY(meta = (DataflowOutput, DisplayName = "Materials"))
+	TArray<TObjectPtr<UMaterial>> Materials;
+};
+
+namespace Dataflow
+{
+	void GeometryCollectionEngineAssetNodes();
+}
\ No newline at end of file
diff --git a/Engine/Source/Runtime/Experimental/GeometryCollectionEngine/Private/GeometryCollection/GeometryCollectionEngineConversion.cpp b/Engine/Source/Runtime/Experimental/GeometryCollectionEngine/Private/GeometryCollection/GeometryCollectionEngineConversion.cpp
--- a/Engine/Source/Runtime/Experimental/GeometryCollectionEngine/Private/GeometryCollection/GeometryCollectionEngineConversion.cpp
+++ b/Engine/Source/Runtime/Experimental/GeometryCollectionEngine/Private/GeometryCollection/GeometryCollectionEngineConversion.cpp
@@ -1265,3 +1265,40 @@ int32 FGeometryCollectionEngineConversion::AppendSkeletalMeshMaterials(const USk
 	return MaterialStart;
 }
 
+void FGeometryCollectionEngineConversion::AppendGeometryCollectionSource(const FGeometryCollectionSource& GeometryCollectionSource, FGeometryCollection& GeometryCollectionInOut, TArray<UMaterial*>& MaterialsInOut)
+{
+	if (const UObject* SourceObject = GeometryCollectionSource.SourceGeometryObject.ResolveObject())
+	{
+		const int32 StartMaterialIndex = MaterialsInOut.Num();
+
+		MaterialsInOut.Append(GeometryCollectionSource.SourceMaterial);
+
+		if (const UStaticMesh* SourceStaticMesh = Cast<UStaticMesh>(SourceObject))
+		{
+			AppendStaticMesh(
+				SourceStaticMesh,
+				StartMaterialIndex,
+				GeometryCollectionSource.LocalTransform,
+				&GeometryCollectionInOut
+				);
+		}
+		else if (const USkeletalMesh* SourceSkeletalMesh = Cast<USkeletalMesh>(SourceObject))
+		{
+			AppendSkeletalMesh(
+				SourceSkeletalMesh,
+				StartMaterialIndex,
+				GeometryCollectionSource.LocalTransform,
+				&GeometryCollectionInOut
+				);
+		}
+		else if (const UGeometryCollection* SourceGeometryCollection = Cast<UGeometryCollection>(SourceObject))
+		{
+			AppendGeometryCollection(
+				SourceGeometryCollection->GetGeometryCollection().Get(),
+				StartMaterialIndex,
+				GeometryCollectionSource.LocalTransform,
+				&GeometryCollectionInOut
+				);
+		}
+	}
+}
\ No newline at end of file
diff --git a/Engine/Source/Runtime/Experimental/GeometryCollectionEngine/Private/GeometryCollection/GeometryCollectionObject.cpp b/Engine/Source/Runtime/Experimental/GeometryCollectionEngine/Private/GeometryCollection/GeometryCollectionObject.cpp
--- a/Engine/Source/Runtime/Experimental/GeometryCollectionEngine/Private/GeometryCollection/GeometryCollectionObject.cpp
+++ b/Engine/Source/Runtime/Experimental/GeometryCollectionEngine/Private/GeometryCollection/GeometryCollectionObject.cpp
@@ -376,6 +376,22 @@ void UGeometryCollection::Reset()
 	}
 }
 
+void UGeometryCollection::ResetFrom(const FManagedArrayCollection& InCollection, const TArray<UMaterial*>& InMaterials)
+{
+	if (GeometryCollection.IsValid())
+	{
+		Reset();
+
+		InCollection.CopyTo(GeometryCollection.Get());
+
+		// todo(Chaos) : we could certainly run a "dependent attribute update method here instead of having to known about convex specifically 
+		UpdateConvexGeometry();
+				
+		Materials.Append(InMaterials);
+		InitializeMaterials();
+	}
+}
+
 /** AppendGeometry */
 int32 UGeometryCollection::AppendGeometry(const UGeometryCollection & Element, bool ReindexAllMaterials, const FTransform& TransformRoot)
 {
diff --git a/Engine/Source/Runtime/Experimental/GeometryCollectionEngine/Public/GeometryCollection/GeometryCollectionEngineConversion.h b/Engine/Source/Runtime/Experimental/GeometryCollectionEngine/Public/GeometryCollection/GeometryCollectionEngineConversion.h
--- a/Engine/Source/Runtime/Experimental/GeometryCollectionEngine/Public/GeometryCollection/GeometryCollectionEngineConversion.h
+++ b/Engine/Source/Runtime/Experimental/GeometryCollectionEngine/Public/GeometryCollection/GeometryCollectionEngineConversion.h
@@ -1,6 +1,7 @@
 // Copyright Epic Games, Inc. All Rights Reserved.
 
 #pragma once
+#include "GeometryCollectionObject.h"
 #include "Engine/World.h"
 
 class UStaticMesh;
@@ -157,4 +158,11 @@ public:
 	*/
 	static void AppendGeometryCollection(const UGeometryCollection* SourceGeometryCollection, const UGeometryCollectionComponent* GeometryCollectionComponent, const FTransform& GeometryCollectionTransform, UGeometryCollection* TargetGeometryCollectionObject, bool ReindexMaterials = true);
 
+	/**
+	*  Appends a GeometryCollectionSource to a GeometryCollection
+	*  @param GeometryCollectionSource : geometry collectiun source object ( from UGeometryCollection collection asset ) 
+	*  @param GeometryCollectionInOut : GeometryCollection to append to 
+	*  @param MaterialsInOut  : array of materials to append to
+	*/
+	static void AppendGeometryCollectionSource(const FGeometryCollectionSource& GeometryCollectionSource, FGeometryCollection& GeometryCollectionInOut, TArray<UMaterial*>& MaterialsInOut);
 };
\ No newline at end of file
diff --git a/Engine/Source/Runtime/Experimental/GeometryCollectionEngine/Public/GeometryCollection/GeometryCollectionObject.h b/Engine/Source/Runtime/Experimental/GeometryCollectionEngine/Public/GeometryCollection/GeometryCollectionObject.h
--- a/Engine/Source/Runtime/Experimental/GeometryCollectionEngine/Public/GeometryCollection/GeometryCollectionObject.h
+++ b/Engine/Source/Runtime/Experimental/GeometryCollectionEngine/Public/GeometryCollection/GeometryCollectionObject.h
@@ -374,6 +374,9 @@ public:
 
 	/** Return collection to initial (ie. empty) state. */
 	void Reset();
+
+	/** Reset the collection from another set of attributes and materials. */
+	void ResetFrom(const FManagedArrayCollection& InCollection, const TArray<UMaterial*>& InMaterials);
 	
 	int32 AppendGeometry(const UGeometryCollection & Element, bool ReindexAllMaterials = false, const FTransform& TransformRoot = FTransform::Identity);
 	int32 NumElements(const FName& Group) const;
