Chaos : fix compilation issues where geometry source property cannot be found when not in editor

diff --git a/Engine/Plugins/Experimental/GeometryCollectionPlugin/Source/GeometryCollectionNodes/Private/Dataflow/Nodes/GeometryCollectionAssetNodes.cpp b/Engine/Plugins/Experimental/GeometryCollectionPlugin/Source/GeometryCollectionNodes/Private/Dataflow/Nodes/GeometryCollectionAssetNodes.cpp
--- a/Engine/Plugins/Experimental/GeometryCollectionPlugin/Source/GeometryCollectionNodes/Private/Dataflow/Nodes/GeometryCollectionAssetNodes.cpp
+++ b/Engine/Plugins/Experimental/GeometryCollectionPlugin/Source/GeometryCollectionNodes/Private/Dataflow/Nodes/GeometryCollectionAssetNodes.cpp
@@ -86,7 +86,12 @@ void FGetGeometryCollectionSourcesDataflowNode::Evaluate(Dataflow::FContext& Con
 	
 	if (const TObjectPtr<UGeometryCollection> InAsset = GetValue(Context, &Asset))
 	{
+#if WITH_EDITORONLY_DATA
 		OutSources = InAsset->GeometrySource; 
+#else
+		ensureMsgf(false, TEXT("FGetGeometryCollectionSourcesDataflowNode - GeometrySource is only available in editor, returning an empty array"));
+#endif
+
 	}
 
 	SetValue(Context, OutSources, &Sources);
