Dataflow : Terminal Node

diff --git a/Engine/Plugins/Experimental/ChaosFlesh/Source/ChaosFleshNodes/Public/Dataflow/ChaosFleshEngineAssetNodes.h b/Engine/Plugins/Experimental/ChaosFlesh/Source/ChaosFleshNodes/Public/Dataflow/ChaosFleshEngineAssetNodes.h
--- a/Engine/Plugins/Experimental/ChaosFlesh/Source/ChaosFleshNodes/Public/Dataflow/ChaosFleshEngineAssetNodes.h
+++ b/Engine/Plugins/Experimental/ChaosFlesh/Source/ChaosFleshNodes/Public/Dataflow/ChaosFleshEngineAssetNodes.h
@@ -62,7 +62,7 @@ public:
 
 
 USTRUCT()
-struct FSetFleshAssetDataflowNode : public FDataflowNode
+struct FSetFleshAssetDataflowNode : public FDataflowTerminalNode
 {
 	GENERATED_USTRUCT_BODY()
 	DATAFLOW_NODE_DEFINE_INTERNAL(FSetFleshAssetDataflowNode, "SetFleshAsset", "Flesh", "")
@@ -73,7 +73,7 @@ public:
 	FManagedArrayCollection Input;
 
 	FSetFleshAssetDataflowNode(const Dataflow::FNodeParameters& InParam, FGuid InGuid = FGuid::NewGuid())
-		: FDataflowNode(InParam, InGuid)
+	: FDataflowTerminalNode(InParam, InGuid)
 	{
 		RegisterInputConnection(&Input);
 	}
diff --git a/Engine/Plugins/Experimental/GeometryCollectionPlugin/Source/GeometryCollectionNodes/Private/Dataflow/Nodes/GeometryCollectionAssetNodes.cpp b/Engine/Plugins/Experimental/GeometryCollectionPlugin/Source/GeometryCollectionNodes/Private/Dataflow/Nodes/GeometryCollectionAssetNodes.cpp
--- a/Engine/Plugins/Experimental/GeometryCollectionPlugin/Source/GeometryCollectionNodes/Private/Dataflow/Nodes/GeometryCollectionAssetNodes.cpp
+++ b/Engine/Plugins/Experimental/GeometryCollectionPlugin/Source/GeometryCollectionNodes/Private/Dataflow/Nodes/GeometryCollectionAssetNodes.cpp
@@ -22,13 +22,13 @@ namespace Dataflow
 // ===========================================================================================================================
 
 FSetGeometryCollectionAssetDataflowNode::FSetGeometryCollectionAssetDataflowNode(const Dataflow::FNodeParameters& InParam, FGuid InGuid)
-	: FFinalDataflowNode(InParam, InGuid)
+	: FDataflowTerminalNode(InParam, InGuid)
 {
 	RegisterInputConnection(&Collection);
 	RegisterInputConnection(&Materials);
 }
 
-void FSetGeometryCollectionAssetDataflowNode::EvaluateFinal(Dataflow::FContext& Context) const
+void FSetGeometryCollectionAssetDataflowNode::Evaluate(Dataflow::FContext& Context) const
 {
 	using FGeometryCollectionPtr = TSharedPtr<FGeometryCollection, ESPMode::ThreadSafe>;
 	using FMaterialArray = TArray<TObjectPtr<UMaterial>>;
diff --git a/Engine/Plugins/Experimental/GeometryCollectionPlugin/Source/GeometryCollectionNodes/Public/Dataflow/GeometryCollectionNodes.h b/Engine/Plugins/Experimental/GeometryCollectionPlugin/Source/GeometryCollectionNodes/Public/Dataflow/GeometryCollectionNodes.h
--- a/Engine/Plugins/Experimental/GeometryCollectionPlugin/Source/GeometryCollectionNodes/Public/Dataflow/GeometryCollectionNodes.h
+++ b/Engine/Plugins/Experimental/GeometryCollectionPlugin/Source/GeometryCollectionNodes/Public/Dataflow/GeometryCollectionNodes.h
@@ -64,7 +64,7 @@ public:
 };
 
 USTRUCT()
-struct FSetCollectionAssetDataflowNode : public FDataflowNode
+struct FSetCollectionAssetDataflowNode : public FDataflowTerminalNode
 {
 	GENERATED_USTRUCT_BODY()
 	DATAFLOW_NODE_DEFINE_INTERNAL(FSetCollectionAssetDataflowNode, "SetCollectionAsset", "GeometryCollection", "")
@@ -76,7 +76,7 @@ public:
 	FManagedArrayCollection Collection;
 
 	FSetCollectionAssetDataflowNode(const Dataflow::FNodeParameters& InParam, FGuid InGuid = FGuid::NewGuid())
-		: FDataflowNode(InParam, InGuid)
+		: FDataflowTerminalNode(InParam, InGuid)
 	{
 		RegisterInputConnection(&Collection);
 	}
diff --git a/Engine/Plugins/Experimental/GeometryCollectionPlugin/Source/GeometryCollectionNodes/Public/Dataflow/Nodes/GeometryCollectionAssetNodes.h b/Engine/Plugins/Experimental/GeometryCollectionPlugin/Source/GeometryCollectionNodes/Public/Dataflow/Nodes/GeometryCollectionAssetNodes.h
--- a/Engine/Plugins/Experimental/GeometryCollectionPlugin/Source/GeometryCollectionNodes/Public/Dataflow/Nodes/GeometryCollectionAssetNodes.h
+++ b/Engine/Plugins/Experimental/GeometryCollectionPlugin/Source/GeometryCollectionNodes/Public/Dataflow/Nodes/GeometryCollectionAssetNodes.h
@@ -11,34 +11,6 @@
 
 class UGeometryCollection;
 
-/**
- * This node is a final node with no output  
- * it is the only one that can have the right to wroite to the asset 
- */
-USTRUCT()
-struct FFinalDataflowNode : public FDataflowNode
-{
-	GENERATED_USTRUCT_BODY()
-	DATAFLOW_NODE_DEFINE_INTERNAL(FFinalDataflowNode, "FinalDataflowNode", "DataFlow", "")
-	
-public:
-	FFinalDataflowNode(const Dataflow::FNodeParameters& InParam, FGuid InGuid = FGuid::NewGuid())
-		: FDataflowNode(InParam, InGuid)
-	{
-	}
-
-protected:
-	/** No output evaluation - secific to this type of node */
-	virtual void EvaluateFinal(Dataflow::FContext& Context) const { ensure(false); }
-	
-private:
-	virtual void Evaluate(Dataflow::FContext& Context, const FDataflowOutput* Out) const override
-	{
-		EvaluateFinal(Context);
-	};
-};
-
-
 /**
  * Set the geometry collection asset   
  * inputs:
@@ -46,14 +18,14 @@ private:
  *	- Materials : array of materials for the geometry
  */
 USTRUCT()
-struct FSetGeometryCollectionAssetDataflowNode : public FFinalDataflowNode
+struct FSetGeometryCollectionAssetDataflowNode : public FDataflowTerminalNode
 {
 	GENERATED_USTRUCT_BODY()
 	DATAFLOW_NODE_DEFINE_INTERNAL(FSetGeometryCollectionAssetDataflowNode, "SetGeometryCollectionAsset", "GeometryCollection", "")
 
 public:
 	FSetGeometryCollectionAssetDataflowNode(const Dataflow::FNodeParameters& InParam, FGuid InGuid = FGuid::NewGuid());
-	virtual void EvaluateFinal(Dataflow::FContext& Context) const override;
+	virtual void Evaluate(Dataflow::FContext& Context) const override;
 
 	/** Attribute collection to reset the asset with */ 
 	UPROPERTY(meta = (DataflowInput, DisplayName = "Collection"))
diff --git a/Engine/Source/Runtime/Experimental/Dataflow/Core/Private/Dataflow/DataflowGraph.cpp b/Engine/Source/Runtime/Experimental/Dataflow/Core/Private/Dataflow/DataflowGraph.cpp
--- a/Engine/Source/Runtime/Experimental/Dataflow/Core/Private/Dataflow/DataflowGraph.cpp
+++ b/Engine/Source/Runtime/Experimental/Dataflow/Core/Private/Dataflow/DataflowGraph.cpp
@@ -49,6 +49,10 @@ namespace Dataflow
 			}
 		}
 		Nodes.Remove(Node);
+		if (Node->IsA(FDataflowTerminalNode::StaticType()))
+		{
+			TerminalNodes.Remove(Node);
+		}
 	}
 
 	void FGraph::ClearConnections(FDataflowConnection* Connection)
diff --git a/Engine/Source/Runtime/Experimental/Dataflow/Core/Public/Dataflow/DataflowCore.h b/Engine/Source/Runtime/Experimental/Dataflow/Core/Public/Dataflow/DataflowCore.h
--- a/Engine/Source/Runtime/Experimental/Dataflow/Core/Public/Dataflow/DataflowCore.h
+++ b/Engine/Source/Runtime/Experimental/Dataflow/Core/Public/Dataflow/DataflowCore.h
@@ -9,3 +9,4 @@
 #include "Dataflow/DataflowNodeFactory.h"
 #include "Dataflow/DataflowNodeColorsRegistry.h"
 #include "Dataflow/DataflowNodeParameters.h"
+#include "Dataflow/DataflowTerminalNode.h"
diff --git a/Engine/Source/Runtime/Experimental/Dataflow/Core/Public/Dataflow/DataflowGraph.h b/Engine/Source/Runtime/Experimental/Dataflow/Core/Public/Dataflow/DataflowGraph.h
--- a/Engine/Source/Runtime/Experimental/Dataflow/Core/Public/Dataflow/DataflowGraph.h
+++ b/Engine/Source/Runtime/Experimental/Dataflow/Core/Public/Dataflow/DataflowGraph.h
@@ -5,6 +5,7 @@
 #include "CoreMinimal.h"
 #include "Chaos/ChaosArchive.h"
 #include "Dataflow/DataflowNode.h"
+#include "Dataflow/DataflowTerminalNode.h"
 #include "Serialization/Archive.h"
 
 struct FDataflowConnection;
@@ -49,13 +50,15 @@ namespace Dataflow
 
 		FGuid  Guid;
 		TArray< TSharedPtr<FDataflowNode> > Nodes;
+		TArray< TSharedPtr<FDataflowNode> > TerminalNodes;
 		TArray< FLink > Connections;
 		TSet< FName > DisabledNodes;
 	public:
 		FGraph(FGuid InGuid = FGuid::NewGuid());
 		virtual ~FGraph() {}
 
-		const TArray< TSharedPtr<FDataflowNode> >& GetNodes() const {return Nodes;}
+		const TArray< TSharedPtr<FDataflowNode> >& GetTerminalNodes() const { return TerminalNodes; }
+		const TArray< TSharedPtr<FDataflowNode> >& GetNodes() const { return Nodes; }
 		TArray< TSharedPtr<FDataflowNode> >& GetNodes() { return Nodes; }
 		int NumNodes() { return Nodes.Num(); }
 
@@ -63,6 +66,10 @@ namespace Dataflow
 		{
 			TSharedPtr<T> NewNode(InNode);
 			Nodes.AddUnique(NewNode);
+			if (NewNode->IsA(FDataflowTerminalNode::StaticType()))
+			{
+				TerminalNodes.AddUnique(NewNode);
+			}
 			return NewNode;
 		}
 
@@ -70,6 +77,10 @@ namespace Dataflow
 		{
 			TSharedPtr<T> NewNode(InNode.Release());
 			Nodes.AddUnique(NewNode);
+			if (NewNode->IsA(FDataflowTerminalNode::StaticType()))
+			{
+				TerminalNodes.AddUnique(NewNode);
+			}
 			return NewNode;
 		}
 
diff --git a/Engine/Source/Runtime/Experimental/Dataflow/Core/Public/Dataflow/DataflowNode.h b/Engine/Source/Runtime/Experimental/Dataflow/Core/Public/Dataflow/DataflowNode.h
--- a/Engine/Source/Runtime/Experimental/Dataflow/Core/Public/Dataflow/DataflowNode.h
+++ b/Engine/Source/Runtime/Experimental/Dataflow/Core/Public/Dataflow/DataflowNode.h
@@ -85,7 +85,7 @@ struct DATAFLOWCORE_API FDataflowNode
 	TArray<Dataflow::FPin> GetPins() const;
 
 
-	void AddInput(FDataflowInput* InPtr);
+	virtual void AddInput(FDataflowInput* InPtr);
 	TArray< FDataflowInput* > GetInputs() const;
 	void ClearInputs();
 
@@ -94,7 +94,7 @@ struct DATAFLOWCORE_API FDataflowNode
 	const FDataflowInput* FindInput(const void* Reference) const;
 
 
-	void AddOutput(FDataflowOutput* InPtr);
+	virtual void AddOutput(FDataflowOutput* InPtr);
 	int NumOutputs() const;
 	TArray< FDataflowOutput* > GetOutputs() const;
 	void ClearOutputs();
@@ -115,8 +115,8 @@ struct DATAFLOWCORE_API FDataflowNode
 	virtual const UScriptStruct* TypedScriptStruct() const { return nullptr; }
 
 	/** Register the Input and Outputs after the creation in the factory */
-	void RegisterInputConnection(const void* Property);
-	void RegisterOutputConnection(const void* Property, const void* Passthrough = nullptr);
+	virtual void RegisterInputConnection(const void* Property);
+	virtual void RegisterOutputConnection(const void* Property, const void* Passthrough = nullptr);
 
 	//
 	// Evaluation
diff --git a/Engine/Source/Runtime/Experimental/Dataflow/Core/Public/Dataflow/DataflowTerminalNode.h b/Engine/Source/Runtime/Experimental/Dataflow/Core/Public/Dataflow/DataflowTerminalNode.h
new file mode 100644
--- /dev/null
+++ b/Engine/Source/Runtime/Experimental/Dataflow/Core/Public/Dataflow/DataflowTerminalNode.h
@@ -0,0 +1,63 @@
+// Copyright Epic Games, Inc. All Rights Reserved.
+
+#pragma once
+
+#include "ChaosLog.h"
+#include "CoreMinimal.h"
+#include "Dataflow/DataflowConnection.h"
+#include "Dataflow/DataflowNodeParameters.h"
+#include "Dataflow/DataflowNode.h"
+
+#include "DataflowTerminalNode.generated.h"
+
+struct FDataflowInput;
+struct FDataflowOutput;
+
+/**
+* FDataflowTerminalNode
+*		Base class for terminal nodes within the Dataflow graph. 
+* 
+*		Terminal Nodes allow for non-const access to UObjects as
+*       edges in the graph. They are used to push data out to
+*       asset or the world from the calling client. Terminals
+*       may not have outputs, they are only leaf nodes in the 
+*       evaluation graph. 
+*/
+USTRUCT()
+struct DATAFLOWCORE_API FDataflowTerminalNode : public FDataflowNode
+{
+	GENERATED_USTRUCT_BODY()
+
+	FDataflowTerminalNode()
+		: Super() { }
+
+	FDataflowTerminalNode(const Dataflow::FNodeParameters& Param, FGuid InGuid = FGuid::NewGuid())
+		: Super(Param,InGuid) {
+	}
+
+	virtual ~FDataflowTerminalNode() { }
+
+	//
+	// Connections
+	//
+
+	virtual void AddOutput(FDataflowOutput* InPtr) override { ensure(false); }
+
+	//
+	//  Struct Support
+	//
+
+	virtual void RegisterOutputConnection(const void* Property, const void* Passthrough = nullptr) { ensure(false); }
+
+
+	virtual void Evaluate(Dataflow::FContext& Context) const { ensure(false); }
+
+private:
+	virtual void Evaluate(Dataflow::FContext& Context, const FDataflowOutput* Out) const override
+	{
+		Evaluate(Context);
+	};
+
+};
+
+
