Dataflow : Add ability to evaluate a terminal node from blueprint
- also cleaned up some comments to play nice with the tooltips

diff --git a/Engine/Plugins/Experimental/GeometryCollectionPlugin/Source/GeometryCollectionNodes/Public/Dataflow/Nodes/GeometryCollectionAssetNodes.h b/Engine/Plugins/Experimental/GeometryCollectionPlugin/Source/GeometryCollectionNodes/Public/Dataflow/Nodes/GeometryCollectionAssetNodes.h
--- a/Engine/Plugins/Experimental/GeometryCollectionPlugin/Source/GeometryCollectionNodes/Public/Dataflow/Nodes/GeometryCollectionAssetNodes.h
+++ b/Engine/Plugins/Experimental/GeometryCollectionPlugin/Source/GeometryCollectionNodes/Public/Dataflow/Nodes/GeometryCollectionAssetNodes.h
@@ -13,10 +13,7 @@ class UGeometryCollection;
 
 /**
  * Set the geometry collection asset   
- * inputs:
- *	- Collection : collection attributes that compose the asset
- *	- Materials : array of materials for the geometry
- */
+  */
 USTRUCT()
 struct FSetGeometryCollectionAssetDataflowNode : public FDataflowTerminalNode
 {
@@ -27,19 +24,18 @@ public:
 	FSetGeometryCollectionAssetDataflowNode(const Dataflow::FNodeParameters& InParam, FGuid InGuid = FGuid::NewGuid());
 	virtual void Evaluate(Dataflow::FContext& Context) const override;
 
-	/** Attribute collection to reset the asset with */ 
+	/** Attribute collection to use for this asset */ 
 	UPROPERTY(meta = (DataflowInput, DisplayName = "Collection"))
 	FManagedArrayCollection Collection;
 
-	/** Materials array to rest the asset with */
+	/** Materials array to use for this asset */
 	UPROPERTY(meta = (DataflowInput, DisplayName = "Materials"))
 	TArray<TObjectPtr<UMaterial>> Materials;
 };
 
 /**
- * Get Current asset    
- * outputs:
- *	- GeometrySources : array of geometry source ( Asset, transform and material )
+ * Get Current geometry collection asset 
+ * Note : Use with caution as this may get replaced in a near future for a more generic getAsset node
  */
 USTRUCT()
 struct FGetGeometryCollectionAssetDataflowNode : public FDataflowNode
@@ -57,9 +53,8 @@ public:
 };
 
 /**
- * Get geometry source    
- * outputs:
- *	- GeometrySources : array of geometry source ( Asset, transform and material )
+ * Get the list of the original mesh information used to create a specific geometryc collection asset
+ * each entry contains a mesh, a transform and a list of override materials
  */
 USTRUCT()
 struct FGetGeometryCollectionSourcesDataflowNode : public FDataflowNode
@@ -82,11 +77,6 @@ public:
 
 /**
  * create a geometry collection from a set of geometry sources    
- * inputs:
- *	- GeometrySources : array of geometry source ( Asset, transform and material )
- *	outputs:
- *	- Collection : Geometry collection newly created
- *	- Materials : Array of materials used by the geometry collection 
  */
 USTRUCT()
 struct FCreateGeometryCollectionFromSourcesDataflowNode : public FDataflowNode
diff --git a/Engine/Source/Runtime/Experimental/Dataflow/Core/Public/Dataflow/DataflowGraph.h b/Engine/Source/Runtime/Experimental/Dataflow/Core/Public/Dataflow/DataflowGraph.h
--- a/Engine/Source/Runtime/Experimental/Dataflow/Core/Public/Dataflow/DataflowGraph.h
+++ b/Engine/Source/Runtime/Experimental/Dataflow/Core/Public/Dataflow/DataflowGraph.h
@@ -132,6 +132,18 @@ namespace Dataflow
 			return TSharedPtr<const FDataflowNode>(nullptr);
 		}
 
+		TSharedPtr<FDataflowNode> FindTerminalNode(FName InName)
+		{
+			for (TSharedPtr<FDataflowNode> Node : TerminalNodes)
+			{
+				if (Node->GetName().IsEqual(InName))
+				{
+					return Node;
+				}
+			}
+			return TSharedPtr<FDataflowNode>(nullptr);
+		}
+
 		void RemoveNode(TSharedPtr<FDataflowNode> Node);
 
 		void ClearConnections(FDataflowConnection* ConnectionBase);
diff --git a/Engine/Source/Runtime/Experimental/Dataflow/Engine/Private/Dataflow/DataflowObject.cpp b/Engine/Source/Runtime/Experimental/Dataflow/Engine/Private/Dataflow/DataflowObject.cpp
--- a/Engine/Source/Runtime/Experimental/Dataflow/Engine/Private/Dataflow/DataflowObject.cpp
+++ b/Engine/Source/Runtime/Experimental/Dataflow/Engine/Private/Dataflow/DataflowObject.cpp
@@ -4,12 +4,15 @@
 #include "Dataflow/DataflowCore.h"
 #include "Dataflow/DataflowEdNode.h"
 #include "Dataflow/DataflowNodeParameters.h"
+#include "Dataflow/DataflowObjectInterface.h"
 
 #include UE_INLINE_GENERATED_CPP_BY_NAME(DataflowObject)
 
 
 #define LOCTEXT_NAMESPACE "UDataflow"
 
+DEFINE_LOG_CATEGORY_STATIC(DATAFLOWOBJECT_LOG, Error, All);
+
 FDataflowAssetEdit::FDataflowAssetEdit(UDataflow* InAsset, FPostEditFunctionCallback InCallback)
 	: PostEditCallback(InCallback)
 	, Asset(InAsset)
@@ -35,6 +38,24 @@ UDataflow::UDataflow(const FObjectInitializer& ObjectInitializer)
 	, Dataflow(new Dataflow::FGraph())
 {}
 
+void UDataflow::EvaluateTerminalNodeByName(FName NodeName, UObject* Asset)
+{
+	if (Dataflow)
+	{
+		TSharedPtr<FDataflowNode> TerminalNode = Dataflow->FindTerminalNode(NodeName);
+		if (TerminalNode)
+		{
+			const float EvalTime = FGameTime::GetTimeSinceAppStart().GetRealTimeSeconds();
+			Dataflow::FEngineContext Context(Asset, this, EvalTime);
+			// Todo , we should be able to cast to a final node type and call EvaluateFinal
+			TerminalNode->Evaluate(Context, nullptr);
+		}
+		else
+		{
+			UE_LOG(DATAFLOWOBJECT_LOG, Warning, TEXT("UDataflow::EvaluateTerminalNodeByName() : Could no find terminal node %s"), *NodeName.ToString());
+		}
+	}
+}
 
 void UDataflow::PostEditCallback()
 {
diff --git a/Engine/Source/Runtime/Experimental/Dataflow/Engine/Public/Dataflow/DataflowObject.h b/Engine/Source/Runtime/Experimental/Dataflow/Engine/Public/Dataflow/DataflowObject.h
--- a/Engine/Source/Runtime/Experimental/Dataflow/Engine/Public/Dataflow/DataflowObject.h
+++ b/Engine/Source/Runtime/Experimental/Dataflow/Engine/Public/Dataflow/DataflowObject.h
@@ -66,12 +66,20 @@ class DATAFLOWENGINE_API UDataflow : public UEdGraph
 public:
 	UDataflow(const FObjectInitializer& ObjectInitializer = FObjectInitializer::Get());
 
+	/** 
+	* Find all the node of a speific type and evaluate them using a specific UObject
+	*/
+	UFUNCTION(BlueprintCallable, Category = "Graph")
+	void EvaluateTerminalNodeByName(FName NodeName, UObject* Asset);
+
+public:
 	UPROPERTY(EditAnywhere, Category = "Evaluation")
 	bool bActive = true;
 
 	UPROPERTY(EditAnywhere, Category = "Evaluation", AdvancedDisplay )
 	TArray<TObjectPtr<UObject>> Targets;
 
+public:
 	/** UObject Interface */
 #if WITH_EDITOR
 	virtual void PostEditChangeProperty(struct FPropertyChangedEvent& PropertyChangedEvent) override;
