Dataflow - Cluster Flatten dataflow node and RemoveOnBreak are now implemented.

diff --git a/Engine/Plugins/Experimental/GeometryCollectionPlugin/Source/GeometryCollectionNodes/Private/Dataflow/GeometryCollectionNodes.cpp b/Engine/Plugins/Experimental/GeometryCollectionPlugin/Source/GeometryCollectionNodes/Private/Dataflow/GeometryCollectionNodes.cpp
--- a/Engine/Plugins/Experimental/GeometryCollectionPlugin/Source/GeometryCollectionNodes/Private/Dataflow/GeometryCollectionNodes.cpp
+++ b/Engine/Plugins/Experimental/GeometryCollectionPlugin/Source/GeometryCollectionNodes/Private/Dataflow/GeometryCollectionNodes.cpp
@@ -9,6 +9,7 @@
 #include "GeometryCollection/ManagedArrayCollection.h"
 #include "GeometryCollection/GeometryCollection.h"
 #include "GeometryCollection/GeometryCollectionEngineUtility.h"
+#include "GeometryCollection/GeometryCollectionEngineRemoval.h"
 #include "GeometryCollection/GeometryCollectionEngineConversion.h"
 #include "Logging/LogMacros.h"
 #include "Templates/SharedPointer.h"
@@ -502,6 +503,74 @@ float GetMaxVertexMovement(float Grout, float Amplitude, int OctaveNumber, float
 	return MaxDisp;
 }
 
+void FClusterFlattenDataflowNode::Evaluate(Dataflow::FContext& Context, const FDataflowOutput* Out) const
+{
+	if (Out->IsA<FManagedArrayCollection>(&Collection))
+	{
+		FManagedArrayCollection CollectionPtr = GetValue<FManagedArrayCollection>(Context, &Collection);
+		if (FGeometryCollection* GeomCollection = CollectionPtr.NewCopy<FGeometryCollection>())
+		{
+			FGeometryCollectionClusteringUtility::UpdateHierarchyLevelOfChildren(GeomCollection, -1);
+			
+			const TManagedArray<int32>& Levels = GeomCollection->GetAttribute<int32>("Level", FGeometryCollection::TransformGroup);
+
+			// Populate Selected Bones in an Array
+			// @todo(harsha) Implement with Selection
+			// For every bone in selected array: [ClusterIndex]
+			int32 ClusterIndex = 0;
+			TArray<int32> LeafBones;
+			FGeometryCollectionClusteringUtility::GetLeafBones(GeomCollection, ClusterIndex, true, LeafBones);
+			FGeometryCollectionClusteringUtility::ClusterBonesUnderExistingNode(GeomCollection, ClusterIndex, LeafBones);
+			FGeometryCollectionClusteringUtility::RemoveDanglingClusters(GeomCollection);
+			// End for
+
+			FGeometryCollectionClusteringUtility::UpdateHierarchyLevelOfChildren(GeomCollection, -1);
+			SetValue<FManagedArrayCollection>(Context, (const FManagedArrayCollection&)(*GeomCollection), &Collection);
+		}
+	}
+}
+
+void FRemoveOnBreakDataflowNode::Evaluate(Dataflow::FContext& Context, const FDataflowOutput* Out) const
+{
+	if (Out->IsA<FManagedArrayCollection>(&Collection))
+	{
+		FManagedArrayCollection CollectionPtr = GetValue<FManagedArrayCollection>(Context, &Collection);
+		if (FGeometryCollection* GeometryCollection = CollectionPtr.NewCopy<FGeometryCollection>())
+		{
+			if (!GeometryCollection->HasAttribute("RemoveOnBreak", FGeometryCollection::TransformGroup))
+	 		{
+	 			TManagedArray<FVector4f>& NewRemoveOnBreak = GeometryCollection->AddAttribute<FVector4f>("RemoveOnBreak", FGeometryCollection::TransformGroup);
+	 			NewRemoveOnBreak.Fill(FRemoveOnBreakData::DisabledPackedData);
+	 		}
+
+			TManagedArray<FVector4f>& RemoveOnBreak = GeometryCollection->ModifyAttribute<FVector4f>("RemoveOnBreak", FGeometryCollection::TransformGroup);
+
+			const FVector2f PostBreakTimerData = GetValue<FVector2f>(Context, &PostBreakTimer);
+			const FVector2f RemovalTimerData = GetValue<FVector2f>(Context, &RemovalTimer);
+			const bool ClusterCrumblingData = GetValue<bool>(Context, &ClusterCrumbling);
+
+			RemoveOnBreak.Fill(FVector4f{PostBreakTimerData.X, PostBreakTimerData.Y, RemovalTimerData.X, RemovalTimerData.Y});
+
+			// @todo(harsha) Implement with Selection
+	 		// const FRemoveOnBreakData RemoveOnBreakData(true, PostBreakTimerData, ClusterCrumblingData, RemovalTimerData);
+	 		// for (int32 Index : SelectedBones)
+	 		// {
+				// // if root bone, then do not set 
+				// if (GeometryCollection->Parent[Index] == INDEX_NONE)
+				// {
+				// 	RemoveOnBreak[Index] = FRemoveOnBreakData::DisabledPackedData;
+				// }
+				// else
+				// {
+				// 	RemoveOnBreak[Index] = RemoveOnBreakData.GetPackedData();
+				// }
+	 		// }
+			
+	 		SetValue<FManagedArrayCollection>(Context, (const FManagedArrayCollection&)(*GeometryCollection), &Collection);
+		}
+	}
+}
+
 void FVoronoiFractureDataflowNode::Evaluate(Dataflow::FContext& Context, const FDataflowOutput* Out) const
 {
 	if (Out->IsA<FManagedArrayCollection>(&Collection))
diff --git a/Engine/Plugins/Experimental/GeometryCollectionPlugin/Source/GeometryCollectionNodes/Public/Dataflow/GeometryCollectionNodes.h b/Engine/Plugins/Experimental/GeometryCollectionPlugin/Source/GeometryCollectionNodes/Public/Dataflow/GeometryCollectionNodes.h
--- a/Engine/Plugins/Experimental/GeometryCollectionPlugin/Source/GeometryCollectionNodes/Public/Dataflow/GeometryCollectionNodes.h
+++ b/Engine/Plugins/Experimental/GeometryCollectionPlugin/Source/GeometryCollectionNodes/Public/Dataflow/GeometryCollectionNodes.h
@@ -695,6 +695,64 @@ public:
 
 };
 
+USTRUCT()
+struct FClusterFlattenDataflowNode : public FDataflowNode
+{
+	GENERATED_USTRUCT_BODY()
+	DATAFLOW_NODE_DEFINE_INTERNAL(FClusterFlattenDataflowNode, "Flatten", "Fracture|Cluster", "")
+
+public:
+	// @todo(harsha) Support Selections
+
+	UPROPERTY(meta = (DataflowInput, DataflowOutput, DisplayName = "Collection"))
+	FManagedArrayCollection Collection;
+
+	FClusterFlattenDataflowNode(const Dataflow::FNodeParameters& InParam, FGuid InGuid = FGuid::NewGuid())
+		: FDataflowNode(InParam, InGuid)
+	{
+		RegisterInputConnection(&Collection);
+		RegisterOutputConnection(&Collection);
+	}
+
+	virtual void Evaluate(Dataflow::FContext& Context, const FDataflowOutput* Out) const override;
+
+};
+
+USTRUCT()
+struct FRemoveOnBreakDataflowNode : public FDataflowNode
+{
+	GENERATED_USTRUCT_BODY()
+	DATAFLOW_NODE_DEFINE_INTERNAL(FRemoveOnBreakDataflowNode, "RemoveOnBreak", "Fracture|Utilities", "")
+
+public:
+	// @todo(harsha) Support Selections
+
+	UPROPERTY(meta = (DataflowInput, DataflowOutput, DisplayName = "Collection"))
+	FManagedArrayCollection Collection;
+
+	UPROPERTY(meta = (DataflowInput, DisplayName = "PostBreakTimer"))
+	FVector2f PostBreakTimer{0.0, 0.0};
+
+	UPROPERTY(meta = (DataflowInput, DisplayName = "RemovalTimer"))
+	FVector2f RemovalTimer{0.0, 1.0};
+
+	UPROPERTY(meta = (DataflowInput, DisplayName = "ClusterCrumbling"))
+	bool ClusterCrumbling = false;
+
+	FRemoveOnBreakDataflowNode(const Dataflow::FNodeParameters& InParam, FGuid InGuid = FGuid::NewGuid())
+		: FDataflowNode(InParam, InGuid)
+	{
+		RegisterInputConnection(&Collection);
+		RegisterOutputConnection(&Collection);
+		RegisterInputConnection(&PostBreakTimer);
+		RegisterInputConnection(&RemovalTimer);
+		RegisterInputConnection(&ClusterCrumbling);
+	}
+
+	virtual void Evaluate(Dataflow::FContext& Context, const FDataflowOutput* Out) const override;
+
+};
+
 USTRUCT()
 struct FVoronoiFractureDataflowNode : public FDataflowNode
 {
