Dataflow : Fix evaluation of geometry collection terminal node

diff --git a/Engine/Plugins/Experimental/GeometryCollectionPlugin/Source/GeometryCollectionNodes/Private/Dataflow/GeometryCollectionNodes.cpp b/Engine/Plugins/Experimental/GeometryCollectionPlugin/Source/GeometryCollectionNodes/Private/Dataflow/GeometryCollectionNodes.cpp
--- a/Engine/Plugins/Experimental/GeometryCollectionPlugin/Source/GeometryCollectionNodes/Private/Dataflow/GeometryCollectionNodes.cpp
+++ b/Engine/Plugins/Experimental/GeometryCollectionPlugin/Source/GeometryCollectionNodes/Private/Dataflow/GeometryCollectionNodes.cpp
@@ -37,6 +37,9 @@ namespace Dataflow
 		DATAFLOW_NODE_REGISTER_CREATION_FACTORY(FAppendCollectionAssetsDataflowNode);
 		DATAFLOW_NODE_REGISTER_CREATION_FACTORY(FResetGeometryCollectionDataflowNode);
 
+		DATAFLOW_NODE_REGISTER_CREATION_FACTORY(FClusterFlattenDataflowNode);
+		DATAFLOW_NODE_REGISTER_CREATION_FACTORY(FRemoveOnBreakDataflowNode);
+
 		DATAFLOW_NODE_REGISTER_CREATION_FACTORY(FPrintStringDataflowNode);
 		DATAFLOW_NODE_REGISTER_CREATION_FACTORY(FLogStringDataflowNode);
 		DATAFLOW_NODE_REGISTER_CREATION_FACTORY(FMakeLiteralStringDataflowNode);
diff --git a/Engine/Plugins/Experimental/GeometryCollectionPlugin/Source/GeometryCollectionNodes/Private/Dataflow/Nodes/GeometryCollectionAssetNodes.cpp b/Engine/Plugins/Experimental/GeometryCollectionPlugin/Source/GeometryCollectionNodes/Private/Dataflow/Nodes/GeometryCollectionAssetNodes.cpp
--- a/Engine/Plugins/Experimental/GeometryCollectionPlugin/Source/GeometryCollectionNodes/Private/Dataflow/Nodes/GeometryCollectionAssetNodes.cpp
+++ b/Engine/Plugins/Experimental/GeometryCollectionPlugin/Source/GeometryCollectionNodes/Private/Dataflow/Nodes/GeometryCollectionAssetNodes.cpp
@@ -37,7 +37,7 @@ void FGeometryCollectionTerminalDataflowNode::SetAssetValue(TObjectPtr<UObject>
 	}
 }
 
-void FGeometryCollectionTerminalDataflowNode::Evaluate(Dataflow::FContext& Context, const FDataflowOutput* Out) const
+void FGeometryCollectionTerminalDataflowNode::Evaluate(Dataflow::FContext& Context) const
 {
 	using FMaterialArray = TArray<TObjectPtr<UMaterial>>;
 
diff --git a/Engine/Plugins/Experimental/GeometryCollectionPlugin/Source/GeometryCollectionNodes/Public/Dataflow/Nodes/GeometryCollectionAssetNodes.h b/Engine/Plugins/Experimental/GeometryCollectionPlugin/Source/GeometryCollectionNodes/Public/Dataflow/Nodes/GeometryCollectionAssetNodes.h
--- a/Engine/Plugins/Experimental/GeometryCollectionPlugin/Source/GeometryCollectionNodes/Public/Dataflow/Nodes/GeometryCollectionAssetNodes.h
+++ b/Engine/Plugins/Experimental/GeometryCollectionPlugin/Source/GeometryCollectionNodes/Public/Dataflow/Nodes/GeometryCollectionAssetNodes.h
@@ -36,8 +36,8 @@ public:
 		RegisterOutputConnection(&Materials, &Materials);
 	}
 
+	virtual void Evaluate(Dataflow::FContext& Context) const override;
 	virtual void SetAssetValue(TObjectPtr<UObject> Asset, Dataflow::FContext& Context) const override;
-	virtual void Evaluate(Dataflow::FContext& Context, const FDataflowOutput* Out) const override;
 };
 
 
diff --git a/Engine/Source/Runtime/Experimental/Dataflow/Engine/Private/Dataflow/DataflowObject.cpp b/Engine/Source/Runtime/Experimental/Dataflow/Engine/Private/Dataflow/DataflowObject.cpp
--- a/Engine/Source/Runtime/Experimental/Dataflow/Engine/Private/Dataflow/DataflowObject.cpp
+++ b/Engine/Source/Runtime/Experimental/Dataflow/Engine/Private/Dataflow/DataflowObject.cpp
@@ -42,11 +42,15 @@ void UDataflow::EvaluateTerminalNodeByName(FName NodeName, UObject* Asset)
 {
 	if (Dataflow)
 	{
-		TSharedPtr<FDataflowNode> TerminalNode = Dataflow->FindTerminalNode(NodeName);
-		if (TerminalNode)
+		TSharedPtr<FDataflowNode> Node = Dataflow->FindTerminalNode(NodeName);
+		if (const FDataflowTerminalNode* TerminalNode = Node->AsType<const FDataflowTerminalNode>())
 		{
 			Dataflow::FEngineContext Context(Asset, this, FPlatformTime::Cycles64());
-			TerminalNode->Evaluate(Context, nullptr);
+			TerminalNode->Evaluate(Context);
+			if (Asset)
+			{
+				TerminalNode->SetAssetValue(Asset, Context);
+			}
 		}
 		else
 		{
