Dataflow : make remove on break properties editable in thedetail panel

diff --git a/Engine/Plugins/Experimental/GeometryCollectionPlugin/Source/GeometryCollectionNodes/Public/Dataflow/GeometryCollectionNodes.h b/Engine/Plugins/Experimental/GeometryCollectionPlugin/Source/GeometryCollectionNodes/Public/Dataflow/GeometryCollectionNodes.h
--- a/Engine/Plugins/Experimental/GeometryCollectionPlugin/Source/GeometryCollectionNodes/Public/Dataflow/GeometryCollectionNodes.h
+++ b/Engine/Plugins/Experimental/GeometryCollectionPlugin/Source/GeometryCollectionNodes/Public/Dataflow/GeometryCollectionNodes.h
@@ -730,13 +730,13 @@ public:
 	UPROPERTY(meta = (DataflowInput, DataflowOutput, DisplayName = "Collection"))
 	FManagedArrayCollection Collection;
 
-	UPROPERTY(meta = (DataflowInput, DisplayName = "PostBreakTimer"))
+	UPROPERTY(EditAnywhere, Category = "Removal", meta = (DataflowInput, DisplayName = "PostBreakTimer"))
 	FVector2f PostBreakTimer{0.0, 0.0};
 
-	UPROPERTY(meta = (DataflowInput, DisplayName = "RemovalTimer"))
+	UPROPERTY(EditAnywhere, Category = "Removal", meta = (DataflowInput, DisplayName = "RemovalTimer"))
 	FVector2f RemovalTimer{0.0, 1.0};
 
-	UPROPERTY(meta = (DataflowInput, DisplayName = "ClusterCrumbling"))
+	UPROPERTY(EditAnywhere, Category = "Removal", meta = (DataflowInput, DisplayName = "ClusterCrumbling"))
 	bool ClusterCrumbling = false;
 
 	FRemoveOnBreakDataflowNode(const Dataflow::FNodeParameters& InParam, FGuid InGuid = FGuid::NewGuid())
