Dataflow : Remove unused toolkits

diff --git a/Engine/Plugins/Experimental/ChaosFlesh/Source/ChaosFleshEditor/Private/ChaosFlesh/ChaosFleshEditorPlugin.cpp b/Engine/Plugins/Experimental/ChaosFlesh/Source/ChaosFleshEditor/Private/ChaosFlesh/ChaosFleshEditorPlugin.cpp
--- a/Engine/Plugins/Experimental/ChaosFlesh/Source/ChaosFleshEditor/Private/ChaosFlesh/ChaosFleshEditorPlugin.cpp
+++ b/Engine/Plugins/Experimental/ChaosFlesh/Source/ChaosFleshEditor/Private/ChaosFlesh/ChaosFleshEditorPlugin.cpp
@@ -11,7 +11,6 @@
 #include "ChaosFlesh/Asset/AssetTypeActions_FleshAsset.h"
 #include "ChaosFlesh/Asset/AssetTypeActions_ChaosDeformableSolverAsset.h"
 #include "ChaosFlesh/Asset/FleshAssetThumbnailRenderer.h"
-#include "ChaosFlesh/Editor/FleshEditorToolkit.h"
 #include "CoreMinimal.h"
 #include "Dataflow/DataflowEditorToolkit.h"
 #include "Features/IModularFeatures.h"
diff --git a/Engine/Plugins/Experimental/ChaosFlesh/Source/ChaosFleshEditor/Private/ChaosFlesh/Editor/FleshEditorToolkit.cpp b/Engine/Plugins/Experimental/ChaosFlesh/Source/ChaosFleshEditor/Private/ChaosFlesh/Editor/FleshEditorToolkit.cpp
deleted file mode 100644
diff --git a/Engine/Plugins/Experimental/ChaosFlesh/Source/ChaosFleshEditor/Private/ChaosFlesh/Editor/FleshEditorToolkit.h b/Engine/Plugins/Experimental/ChaosFlesh/Source/ChaosFleshEditor/Private/ChaosFlesh/Editor/FleshEditorToolkit.h
deleted file mode 100644
diff --git a/Engine/Plugins/Experimental/ChaosFlesh/Source/ChaosFleshEditor/Private/ChaosFlesh/Editor/FleshEditorViewport.cpp b/Engine/Plugins/Experimental/ChaosFlesh/Source/ChaosFleshEditor/Private/ChaosFlesh/Editor/FleshEditorViewport.cpp
deleted file mode 100644
diff --git a/Engine/Plugins/Experimental/ChaosFlesh/Source/ChaosFleshEditor/Private/ChaosFlesh/Editor/FleshEditorViewport.h b/Engine/Plugins/Experimental/ChaosFlesh/Source/ChaosFleshEditor/Private/ChaosFlesh/Editor/FleshEditorViewport.h
deleted file mode 100644
diff --git a/Engine/Plugins/Experimental/GeometryCollectionPlugin/Source/GeometryCollectionEditor/Private/GeometryCollection/GeometryCollectionEditorPlugin.cpp b/Engine/Plugins/Experimental/GeometryCollectionPlugin/Source/GeometryCollectionEditor/Private/GeometryCollection/GeometryCollectionEditorPlugin.cpp
--- a/Engine/Plugins/Experimental/GeometryCollectionPlugin/Source/GeometryCollectionEditor/Private/GeometryCollection/GeometryCollectionEditorPlugin.cpp
+++ b/Engine/Plugins/Experimental/GeometryCollectionPlugin/Source/GeometryCollectionEditor/Private/GeometryCollection/GeometryCollectionEditorPlugin.cpp
@@ -19,7 +19,6 @@
 #include "GeometryCollection/GeometryCollectionThumbnailRenderer.h"
 #include "GeometryCollection/GeometryCollectionSelectRigidBodyEdMode.h"
 #include "GeometryCollection/GeometryCollectionSelectionCommands.h"
-#include "GeometryCollection/GeometryCollectionEditorToolkit.h"
 #include "GeometryCollection/GeometryCollectionFactory.h"
 #include "HAL/ConsoleManager.h"
 #include "Features/IModularFeatures.h"
@@ -335,11 +334,3 @@ void IGeometryCollectionEditorPlugin::RegisterMenus()
 		Section.AddMenuEntry(FGeometryCollectionSelectionCommands::Get().SelectInverseGeometry);
 	}
 }
-
-
-TSharedRef<FAssetEditorToolkit> IGeometryCollectionEditorPlugin::CreateGeometryCollectionAssetEditor(const EToolkitMode::Type Mode, const TSharedPtr<IToolkitHost>& InitToolkitHost, UObject* GeometryCollectionAsset)
-{
-	TSharedPtr<FGeometryCollectionEditorToolkit> NewGeometryCollectionAssetEditor = MakeShared<FGeometryCollectionEditorToolkit>();
-	NewGeometryCollectionAssetEditor->InitGeometryCollectionAssetEditor(Mode, InitToolkitHost, GeometryCollectionAsset);
-	return StaticCastSharedPtr<FAssetEditorToolkit>(NewGeometryCollectionAssetEditor).ToSharedRef();
-}
\ No newline at end of file
diff --git a/Engine/Plugins/Experimental/GeometryCollectionPlugin/Source/GeometryCollectionEditor/Private/GeometryCollection/GeometryCollectionEditorToolkit.cpp b/Engine/Plugins/Experimental/GeometryCollectionPlugin/Source/GeometryCollectionEditor/Private/GeometryCollection/GeometryCollectionEditorToolkit.cpp
deleted file mode 100644
diff --git a/Engine/Plugins/Experimental/GeometryCollectionPlugin/Source/GeometryCollectionEditor/Public/GeometryCollection/GeometryCollectionEditorPlugin.h b/Engine/Plugins/Experimental/GeometryCollectionPlugin/Source/GeometryCollectionEditor/Public/GeometryCollection/GeometryCollectionEditorPlugin.h
--- a/Engine/Plugins/Experimental/GeometryCollectionPlugin/Source/GeometryCollectionEditor/Public/GeometryCollection/GeometryCollectionEditorPlugin.h
+++ b/Engine/Plugins/Experimental/GeometryCollectionPlugin/Source/GeometryCollectionEditor/Public/GeometryCollection/GeometryCollectionEditorPlugin.h
@@ -27,13 +27,6 @@ public:
 	virtual void StartupModule();
 	virtual void ShutdownModule();
 
-	/**
-	*  Create a custom asset editor.
-	*
-	*/
-	TSharedRef<FAssetEditorToolkit> CreateGeometryCollectionAssetEditor(const EToolkitMode::Type Mode, const TSharedPtr<IToolkitHost>& InitToolkitHost, UObject* GeometryCollectionAsset);
-
-
 	/**
 	 * Singleton-like access to this module's interface.  This is just for convenience!
 	 * Beware of calling this during the shutdown phase, though.  Your module might have been unloaded already.
diff --git a/Engine/Plugins/Experimental/GeometryCollectionPlugin/Source/GeometryCollectionEditor/Public/GeometryCollection/GeometryCollectionEditorToolkit.h b/Engine/Plugins/Experimental/GeometryCollectionPlugin/Source/GeometryCollectionEditor/Public/GeometryCollection/GeometryCollectionEditorToolkit.h
deleted file mode 100644
