Dataflow:
- Added TransformSelection data type
- Implemented these nodes: GetArrayElementDataflowNode, GetNumArrayElementsDataflowNode, GetCentroidsDataflowNode, PointsToMeshDataflowNode, BoxToMeshDataflowNode, MeshInfoDataflowNode,
                                           MeshToCollectionDataflowNode, StaticMeshToMeshDataflowNode, TransformDataflowNode, MeshAppendDataflowNode, MeshBooleanDataflowNode, MeshCopyToPointsDataflowNode,
                                           CompareIntDataflowNode, BranchDataflowNode, GetMeshDataDataflowNode, GetSchemaDataflowNode, CollectionTransformSelectionAllDataflowNode,
                                           CollectionTransformSelectionSetOperationDataflowNode, CollectionTransformSelectionInfoDataflowNode, CollectionTransformSelectionNoneDataflowNode,
                                           CollectionTransformSelectionInvertDataflowNode, CollectionTransformSelectionRandomDataflowNode, AutoClusterDataflowNode, CollectionTransformSelectionRootDataflowNode,
                                           CollectionTransformSelectionCustomDataflowNode, CollectionTransformSelectionParentDataflowNode, CollectionTransformSelectionByPercentageDataflowNode,
                                           CollectionTransformSelectionChildrenDataflowNode, CollectionTransformSelectionSiblingsDataflowNode, CollectionTransformSelectionLevelDataflowNode,
                                           CollectionTransformSelectionContactDataflowNode, CollectionTransformSelectionLeafDataflowNode, CollectionTransformSelectionClusterDataflowNode,
                                           CollectionTransformSelectionBySizeDataflowNode, CollectionTransformSelectionByVolumeDataflowNode, SetAnchorStateDataflowNode, FProximityDataflowNode
- Renamed Passthrough metadata tag to DataflowPassthrough
- Added Tooltip display for nodes (both for hovering over the node and hovering over the node name in Action menu) and for pins, the tooltip gets build from the /** ... */ comments from the header file

diff --git a/Engine/Plugins/Experimental/ChaosEditor/ChaosEditor.uplugin b/Engine/Plugins/Experimental/ChaosEditor/ChaosEditor.uplugin
--- a/Engine/Plugins/Experimental/ChaosEditor/ChaosEditor.uplugin
+++ b/Engine/Plugins/Experimental/ChaosEditor/ChaosEditor.uplugin
@@ -39,6 +39,10 @@
 		{
 			"Name": "EditorScriptingUtilities",
 			"Enabled": true
+		},
+		{
+			"Name": "Fracture",
+			"Enabled": true
 		}
 	]
 }
\ No newline at end of file
diff --git a/Engine/Plugins/Experimental/ChaosEditor/Source/FractureEditor/FractureEditor.Build.cs b/Engine/Plugins/Experimental/ChaosEditor/Source/FractureEditor/FractureEditor.Build.cs
--- a/Engine/Plugins/Experimental/ChaosEditor/Source/FractureEditor/FractureEditor.Build.cs
+++ b/Engine/Plugins/Experimental/ChaosEditor/Source/FractureEditor/FractureEditor.Build.cs
@@ -44,6 +44,7 @@ public class FractureEditor : ModuleRules
 				"Chaos",
 				"ToolWidgets",
 				"DeveloperSettings",
+				"FractureEngine"
 
 				// ... add private dependencies that you statically link with here ...	
 			}
diff --git a/Engine/Plugins/Experimental/ChaosEditor/Source/FractureEditor/Private/FractureToolAutoCluster.cpp b/Engine/Plugins/Experimental/ChaosEditor/Source/FractureEditor/Private/FractureToolAutoCluster.cpp
--- a/Engine/Plugins/Experimental/ChaosEditor/Source/FractureEditor/Private/FractureToolAutoCluster.cpp
+++ b/Engine/Plugins/Experimental/ChaosEditor/Source/FractureEditor/Private/FractureToolAutoCluster.cpp
@@ -11,6 +11,7 @@
 #include "GeometryCollection/GeometryCollectionProximityUtility.h"
 #include "GeometryCollection/GeometryCollectionAlgo.h"
 #include "Math/NumericLimits.h"
+#include "FractureEngineClustering.h"
 
 #include UE_INLINE_GENERATED_CPP_BY_NAME(FractureToolAutoCluster)
 
@@ -47,6 +48,7 @@ FSlateIcon UFractureToolAutoCluster::GetToolIcon() const
 	return FSlateIcon("FractureEditorStyle", "FractureEditor.AutoCluster");
 }
 
+
 TArray<UObject*> UFractureToolAutoCluster::GetSettingsObjects() const 
 {
 	TArray<UObject*> AllSettings; 
@@ -57,7 +59,7 @@ TArray<UObject*> UFractureToolAutoCluster::GetSettingsObjects() const
 void UFractureToolAutoCluster::DrawBox(FPrimitiveDrawInterface* PDI, FVector Center, float SideLength)
 {
 	float SideBy2 = SideLength/2.;
-	
+
 	FVector X(1., 0, 0);
 	FVector Y(0, 1., 0);
 	FVector Z(0, 0, 1.);
@@ -73,10 +75,10 @@ void UFractureToolAutoCluster::DrawBox(FPrimitiveDrawInterface* PDI, FVector Cen
 			}
 		}
 	}
-	
+
 	Verts.Swap(0,1);
 	Verts.Swap(4,5);
-	
+
 	for (int i=0; i<4; i+=1)
 	{
 		PDI->DrawLine(Verts[i], Verts[(i+1)%4], FLinearColor(255, 0, 0), SDPG_Foreground);
@@ -91,9 +93,9 @@ void UFractureToolAutoCluster::Render(const FSceneView* View, FViewport* Viewpor
 	{
 		return;
 	}
-	
+
 	TArray<FFractureToolContext> FractureContexts = GetFractureToolContexts();
-	
+
 	for (FFractureToolContext& FractureContext : FractureContexts)
 	{
 		DrawBox(PDI, FractureContext.GetTransform().GetTranslation(), AutoClusterSettings->SiteSize);
@@ -119,424 +121,30 @@ void UFractureToolAutoCluster::Execute(TWeakPtr<FFractureEditorModeToolkit> InTo
 			FGeometryCollectionEdit Edit(Context.GetGeometryCollectionComponent(), GeometryCollection::EEditUpdate::RestPhysicsDynamic);
 
 			Context.ConvertSelectionToClusterNodes();
-			FGeometryCollection* GeometryCollection = Context.GetGeometryCollection().Get();
-
-			int32 StartTransformCount = GeometryCollection->Transform.Num();
-
-			for (const int32 ClusterIndex : Context.GetSelection())
+			if (FGeometryCollection* GeometryCollection = Context.GetGeometryCollection().Get())
 			{
-				FVoronoiPartitioner VoronoiPartition(GeometryCollection, ClusterIndex);
-				int32 NumChildren = GeometryCollection->Children[ClusterIndex].Num();
-				int32 SiteCountToUse = 1;
-				if (AutoClusterSettings->ClusterSizeMethod == EClusterSizeMethod::ByNumber)
-				{
-					SiteCountToUse = AutoClusterSettings->SiteCount;
-				}
-				else if (AutoClusterSettings->ClusterSizeMethod == EClusterSizeMethod::ByFractionOfInput)
-				{
-					SiteCountToUse = FMath::Max(2, NumChildren * AutoClusterSettings->SiteCountFraction);
-				}
-				else if (AutoClusterSettings->ClusterSizeMethod == EClusterSizeMethod::BySize)
-				{
-					float TotalVolume = GeometryCollection->GetBoundingBox().GetBox().GetVolume();
-					float DesiredVolume = FMath::Pow(AutoClusterSettings->SiteSize, 3);
-					SiteCountToUse = FMath::Max(1, TotalVolume / DesiredVolume);
-				}
-				
-				VoronoiPartition.KMeansPartition(SiteCountToUse);
-				if (VoronoiPartition.GetPartitionCount() == 0)
-				{
-					continue;
-				}
-
-				bool bNeedProximity = AutoClusterSettings->bEnforceConnectivity || AutoClusterSettings->bAvoidIsolated;
-				if (bNeedProximity)
-				{
-					FGeometryCollectionProximityUtility ProximityUtility(GeometryCollection);
-					ProximityUtility.UpdateProximity();
-				}
+				int32 StartTransformCount = GeometryCollection->Transform.Num();
 
-				if (AutoClusterSettings->bEnforceConnectivity)
+				for (const int32 ClusterIndex : Context.GetSelection())
 				{
-					VoronoiPartition.SplitDisconnectedPartitions(GeometryCollection);
+					FFractureEngineClustering::AutoCluster(*GeometryCollection,
+						ClusterIndex,
+						(EFractureEngineClusterSizeMethod)AutoClusterSettings->ClusterSizeMethod,
+						AutoClusterSettings->SiteCount,
+						AutoClusterSettings->SiteCountFraction,
+						AutoClusterSettings->SiteSize,
+						AutoClusterSettings->bEnforceConnectivity,
+						AutoClusterSettings->bAvoidIsolated);
 				}
 
-				if (AutoClusterSettings->bAvoidIsolated)
-				{
-					// attempt to remove isolated via merging (may not succeed, as it only merges if there is a cluster in proximity)
-					VoronoiPartition.MergeSingleElementPartitions(GeometryCollection);
-				}
-
-				int32 NonEmptyPartitionCount = VoronoiPartition.GetNonEmptyPartitionCount();
-				bool bHasEmptyClusters = NonEmptyPartitionCount > 0;
-
-				if (AutoClusterSettings->bAvoidIsolated && NonEmptyPartitionCount == 1)
-				{
-					continue;
-				}
-					
-				int32 PartitionCount = VoronoiPartition.GetPartitionCount();
-				int32 NewClusterIndexStart = GeometryCollection->AddElements(PartitionCount, FGeometryCollection::TransformGroup);
-
-				for (int32 Index = 0; Index < PartitionCount; ++Index)
-				{
-					TArray<int32> NewCluster = VoronoiPartition.GetPartition(Index);
-					if (AutoClusterSettings->bAvoidIsolated && NewCluster.Num() == 1)
-					{
-						bHasEmptyClusters = true;
-						NewCluster.Reset();
-					}
-
-					int32 NewClusterIndex = NewClusterIndexStart + Index;
-					GeometryCollection->Parent[NewClusterIndex] = ClusterIndex;
-					GeometryCollection->Children[ClusterIndex].Add(NewClusterIndex);
-					GeometryCollection->BoneName[NewClusterIndex] = "ClusterBone";
-					GeometryCollection->Children[NewClusterIndex] = TSet<int32>(NewCluster);
-					GeometryCollection->SimulationType[NewClusterIndex] = FGeometryCollection::ESimulationTypes::FST_Clustered;
-					GeometryCollection->Transform[NewClusterIndex] = FTransform::Identity;
-					GeometryCollectionAlgo::ParentTransforms(GeometryCollection, NewClusterIndex, NewCluster);
-				}
-				if (bHasEmptyClusters)
-				{
-					FGeometryCollectionClusteringUtility::RemoveDanglingClusters(GeometryCollection);
-				}
-				FGeometryCollectionClusteringUtility::UpdateHierarchyLevelOfChildren(GeometryCollection, ClusterIndex);
-				FGeometryCollectionClusteringUtility::RecursivelyUpdateChildBoneNames(ClusterIndex, GeometryCollection->Children, GeometryCollection->BoneName);
-				FGeometryCollectionClusteringUtility::ValidateResults(GeometryCollection);
+				Context.GenerateGuids(StartTransformCount);
 			}
-
-			Context.GenerateGuids(StartTransformCount);
-
+			
 			Refresh(Context, Toolkit);
 		}
 		SetOutlinerComponents(Contexts, Toolkit);
 	}
 }
-
-
-FVoronoiPartitioner::FVoronoiPartitioner(const FGeometryCollection* GeometryCollection, int32 ClusterIndex)
-{
-	// Collect children of cluster
-	const TManagedArray<TSet<int32>>& Children = GeometryCollection->Children;
-	TransformIndices = Children[ClusterIndex].Array();
-
-	GenerateCentroids(GeometryCollection);
-}
-
-void FVoronoiPartitioner::KMeansPartition(int32 InPartitionCount)
-{
-	PartitionCount = InPartitionCount;
-
-	InitializePartitions();
-	for (int32 Iteration = 0; Iteration < MaxKMeansIterations; Iteration++)
-	{
-		// Refinement is complete when no nodes change partition.
-		if (!Refine())
-		{
-			break;
-		}
-	}
-}
-
-void FVoronoiPartitioner::MergeSingleElementPartitions(FGeometryCollection* GeometryCollection)
-{
-	if (Connectivity.IsEmpty())
-	{
-		GenerateConnectivity(GeometryCollection);
-	}
-	for (int32 ElIdx = 0; ElIdx < TransformIndices.Num(); ElIdx++)
-	{
-		int32 Partition = Partitions[ElIdx];
-		if (PartitionSize[Partition] == 1)
-		{
-			// Find the smallest neighboring partition to merge to
-			// (to help keep partition sizes balanced)
-			int32 SmallestNbrPartition = -1;
-			int32 SmallestNbrSize = TransformIndices.Num()+1;
-			for (int32 NbrEl : Connectivity[ElIdx])
-			{
-				int32 NbrPartition = Partitions[NbrEl];
-				if (NbrPartition == Partition)
-				{
-					continue;
-				}
-				int32 NbrSize = PartitionSize[NbrPartition];
-				if (NbrSize > 0 && (SmallestNbrPartition == -1 || NbrSize < SmallestNbrSize))
-				{
-					SmallestNbrPartition = NbrPartition;
-					SmallestNbrSize = NbrSize;
-				}
-			}
-			if (SmallestNbrPartition != -1)
-			{
-				Partitions[ElIdx] = SmallestNbrPartition;
-				PartitionSize[Partition]--;
-				PartitionSize[SmallestNbrPartition]++;
-			}
-		}
-	}
-}
-
-void FVoronoiPartitioner::SplitDisconnectedPartitions(FGeometryCollection* GeometryCollection)
-{
-	Visited.Init(false, TransformIndices.Num());
-	if (Connectivity.IsEmpty())
-	{
-		GenerateConnectivity(GeometryCollection);
-	}
-	
-	for (int32 PartitionIndex = 0; PartitionIndex < PartitionCount; ++PartitionIndex)
-	{
-		int32 FirstElementInPartitionIndex = Partitions.Find(PartitionIndex);
-		if (FirstElementInPartitionIndex > INDEX_NONE)
-		{
-			MarkVisited(FirstElementInPartitionIndex, PartitionIndex);
-
-			// Place unvisited partition members in a new partition.
-			bool bFoundUnattached = false;
-			for (int32 Index = 0; Index < TransformIndices.Num(); ++Index)
-			{
-				if (Partitions[Index] == PartitionIndex)
-				{
-					if (!Visited[Index])
-					{
-						if (!bFoundUnattached)
-						{
-							bFoundUnattached = true;
-							PartitionCount++;
-							PartitionSize.Add(0);
-						}
-						Partitions[Index] = PartitionCount - 1;
-						PartitionSize[PartitionCount - 1]++;
-						PartitionSize[PartitionIndex]--;
-					}
-				}
-			}
-		}
-	}
-}
-
-void FVoronoiPartitioner::MarkVisited(int32 Index, int32 PartitionIndex)
-{
-	Visited[Index] = true;
-	for (int32 Adjacent : Connectivity[Index])
-	{
-		if (!Visited[Adjacent] && Partitions[Adjacent] == PartitionIndex)
-		{
-			MarkVisited(Adjacent, PartitionIndex);
-		}
-	}
-}
-
-TArray<int32> FVoronoiPartitioner::GetPartition(int32 PartitionIndex) const
-{
-	TArray<int32> PartitionContents;
-	PartitionContents.Reserve(TransformIndices.Num());
-	for (int32 Index = 0; Index < TransformIndices.Num(); ++Index)
-	{
-		if (Partitions[Index] == PartitionIndex)
-		{
-			PartitionContents.Add(TransformIndices[Index]);
-		}
-	}
-
-	return PartitionContents;
-}
-
-void FVoronoiPartitioner::GenerateConnectivity(const FGeometryCollection* GeometryCollection)
-{
-	
-	Connectivity.SetNum(TransformIndices.Num());
-
-	if (!ensure(GeometryCollection->HasAttribute("Level", FGeometryCollection::TransformGroup)))
-	{
-		return;
-	}
-	const TManagedArray<int32>& Levels = GeometryCollection->GetAttribute<int32>("Level", FGeometryCollection::TransformGroup);
-
-	for (int32 Index = 0; Index < TransformIndices.Num(); ++Index)
-	{
-		CollectConnections(GeometryCollection, TransformIndices[Index], Levels[TransformIndices[Index]], Connectivity[Index]);
-	}
-}
-
-void FVoronoiPartitioner::CollectConnections(const FGeometryCollection* GeometryCollection, int32 Index, int32 OperatingLevel, TSet<int32>& OutConnections) const
-{
-	const TManagedArray<TSet<int32>>& Children = GeometryCollection->Children;
-	if (GeometryCollection->IsRigid(Index))
-	{
-		const TManagedArray<TSet<int32>>& Proximity = GeometryCollection->GetAttribute<TSet<int32>>("Proximity", FGeometryCollection::GeometryGroup);
-		const TManagedArray<int32>& GeometryToTransformIndex = GeometryCollection->TransformIndex;
-		const TManagedArray<int32>& TransformToGeometryIndex = GeometryCollection->TransformToGeometryIndex;
-
-		for (int32 Neighbor : Proximity[TransformToGeometryIndex[Index]])
-		{
-			int32 NeighborTransformIndex = FGeometryCollectionClusteringUtility::GetParentOfBoneAtSpecifiedLevel(GeometryCollection, GeometryToTransformIndex[Neighbor], OperatingLevel);
-			int32 ContextIndex = TransformIndices.Find(NeighborTransformIndex);
-			if (ContextIndex > INDEX_NONE)
-			{
-				OutConnections.Add(ContextIndex);
-			}
-		}
-	}
-	else
-	{
-		for (int32 ChildIndex : Children[Index])
-		{
-			CollectConnections(GeometryCollection, ChildIndex, OperatingLevel, OutConnections);
-		}
-	}
-}
-
-void FVoronoiPartitioner::GenerateCentroids(const FGeometryCollection* GeometryCollection)
-{
-	Centroids.SetNum(TransformIndices.Num());
-	ParallelFor(TransformIndices.Num(), [this, GeometryCollection](int32 Index)
-		{
-			Centroids[Index] = GenerateCentroid(GeometryCollection, TransformIndices[Index]);
-		});
-}
-
-FVector FVoronoiPartitioner::GenerateCentroid(const FGeometryCollection* GeometryCollection, int32 TransformIndex) const
-{
-	FBox Bounds = GenerateBounds(GeometryCollection, TransformIndex);
-	return Bounds.GetCenter();
-}
-
-FBox FVoronoiPartitioner::GenerateBounds(const FGeometryCollection* GeometryCollection, int32 TransformIndex) const
-{
-	check(GeometryCollection->IsRigid(TransformIndex) || GeometryCollection->IsClustered(TransformIndex));
-
-	// Return the bounds of all the vertices contained by this branch.
-
-	if (GeometryCollection->IsRigid(TransformIndex))
-	{
-		const TManagedArray<FTransform>& Transforms = GeometryCollection->Transform;
-		const TManagedArray<int32>& Parents = GeometryCollection->Parent;
-		FTransform GlobalTransform = GeometryCollectionAlgo::GlobalMatrix(Transforms, Parents, TransformIndex);
-		
-		int32 GeometryIndex = GeometryCollection->TransformToGeometryIndex[TransformIndex];
-		int32 VertexStart = GeometryCollection->VertexStart[GeometryIndex];
-		int32 VertexCount = GeometryCollection->VertexCount[GeometryIndex];
-		const TManagedArray<FVector3f>& GCVertices = GeometryCollection->Vertex;
-
-		TArray<FVector> Vertices;
-		Vertices.SetNum(VertexCount);
-		for (int32 VertexOffset = 0; VertexOffset < VertexCount; ++VertexOffset)
-		{
-			Vertices[VertexOffset] = GlobalTransform.TransformPosition((FVector)GCVertices[VertexStart + VertexOffset]);
-		}
-
-		return FBox(Vertices);
-	}
-	else // recurse the cluster
-	{
-		const TArray<int32> Children = GeometryCollection->Children[TransformIndex].Array();
-
-		// Empty clusters are problematic
-		if (Children.Num() == 0)
-		{
-			return FBox();
-		}
-
-		FBox Bounds = GenerateBounds(GeometryCollection, Children[0]);
-		for (int32 ChildIndex = 1; ChildIndex < Children.Num(); ++ChildIndex)
-		{
-			Bounds += GenerateBounds(GeometryCollection, Children[ChildIndex]);
-		}
-
-		return Bounds;
-	}
-}
-
-void FVoronoiPartitioner::InitializePartitions()
-{
-	check(PartitionCount > 0);
-
-	PartitionCount = FMath::Min(PartitionCount, TransformIndices.Num());
-
-	// Set initial partition centers as selects from the vertex set
-	PartitionCenters.SetNum(PartitionCount);
-	int32 TransformStride = FMath::Max(1,FMath::FloorToInt(float(TransformIndices.Num()) / float(PartitionCount)));
-	for (int32 PartitionIndex = 0; PartitionIndex < PartitionCount; ++PartitionIndex)
-	{
-		PartitionCenters[PartitionIndex] = Centroids[TransformStride * PartitionIndex];
-	}
-
-	// At beginning, all nodes belong to first partition.
-	Partitions.Init(0, TransformIndices.Num());
-	PartitionSize.Init(0, PartitionCount);
-
-	if (PartitionCount > 0)
-	{
-		PartitionSize[0] = TransformIndices.Num();
-	}
-}
-
-bool FVoronoiPartitioner::Refine()
-{
-	// Assign each transform to its closest partition center.
-	bool bUnchanged = true;
-	for (int32 Index = 0; Index < Centroids.Num(); ++Index)
-	{
-		int32 ClosestPartition = FindClosestPartitionCenter(Centroids[Index]);
-		if (ClosestPartition != Partitions[Index])
-		{
-			bUnchanged = false;
-			PartitionSize[Partitions[Index]]--;
-			PartitionSize[ClosestPartition]++;
-			Partitions[Index] = ClosestPartition;
-		}
-	}
-
-	if (bUnchanged)
-	{
-		return false;
-	}
-
-	// Recalculate partition centers
-	for (int32 PartitionIndex = 0; PartitionIndex < PartitionCount; ++PartitionIndex)
-	{
-		if (PartitionSize[PartitionIndex] > 0)
-		{
-			PartitionCenters[PartitionIndex] = FVector(0.0f, 0.0f, 0.0f);
-		}
-	}
-
-	for (int32 Index = 0; Index < Partitions.Num(); ++Index)
-	{
-		PartitionCenters[Partitions[Index]] += Centroids[Index];
-	}
-
-	for (int32 PartitionIndex = 0; PartitionIndex < PartitionCount; ++PartitionIndex)
-	{
-		if (PartitionSize[PartitionIndex] > 0)
-		{
-			PartitionCenters[PartitionIndex] /= float(PartitionSize[PartitionIndex]);
-		}
-	}
-
-	return true;
-}
-
-int32 FVoronoiPartitioner::FindClosestPartitionCenter(const FVector& Location) const
-{
-	int32 ClosestPartition = INDEX_NONE;
-	float SmallestDistSquared = TNumericLimits<float>::Max();
-
-	for (int32 PartitionIndex = 0; PartitionIndex < PartitionCount; ++PartitionIndex)
-	{
-		float DistSquared = FVector::DistSquared(Location, PartitionCenters[PartitionIndex]);
-		if (DistSquared < SmallestDistSquared)
-		{
-			SmallestDistSquared = DistSquared;
-			ClosestPartition = PartitionIndex;
-		}
-	}
-
-	check(ClosestPartition > INDEX_NONE);
-	return ClosestPartition;
-}
 	
 #undef LOCTEXT_NAMESPACE
 
diff --git a/Engine/Plugins/Experimental/ChaosEditor/Source/FractureEditor/Private/FractureToolAutoCluster.h b/Engine/Plugins/Experimental/ChaosEditor/Source/FractureEditor/Private/FractureToolAutoCluster.h
--- a/Engine/Plugins/Experimental/ChaosEditor/Source/FractureEditor/Private/FractureToolAutoCluster.h
+++ b/Engine/Plugins/Experimental/ChaosEditor/Source/FractureEditor/Private/FractureToolAutoCluster.h
@@ -99,53 +99,3 @@ public:
 	TObjectPtr<UFractureAutoClusterSettings> AutoClusterSettings;
 };
 
-
-class FVoronoiPartitioner
-{
-public:
-	FVoronoiPartitioner(const FGeometryCollection* GeometryCollection, int32 ClusterIndex);
-
-	/** Cluster bodies into k partitions using K-Means. Connectivity is ignored: only spatial proximity is considered. */
-	void KMeansPartition(int32 InPartitionCount);
-
-	/** Split any partition islands into their own partition. This will possbily increase number of partitions to exceed desired count. */
-	void SplitDisconnectedPartitions(FGeometryCollection* GeometryCollection);
-
-	/** Merge any partitions w/ only 1 body into a connected, neighboring partition (if any).  This can decrease the number of partitions below the desired count. */
-	void MergeSingleElementPartitions(FGeometryCollection* GeometryCollection);
-
-	int32 GetPartitionCount() const { return PartitionCount; }
-
-	int32 GetNonEmptyPartitionCount() const
-	{
-		return PartitionSize.Num() - Algo::Count(PartitionSize, 0);
-	}
-
-	/** return the GeometryCollection TranformIndices within the partition. */
-	TArray<int32> GetPartition(int32 PartitionIndex) const;
-
-private:
-	void GenerateConnectivity(const FGeometryCollection* GeometryCollection);
-	void CollectConnections(const FGeometryCollection* GeometryCollection, int32 Index, int32 OperatingLevel, TSet<int32>& OutConnections) const;
-	void GenerateCentroids(const FGeometryCollection* GeometryCollection);
-	FVector GenerateCentroid(const FGeometryCollection* GeometryCollection, int32 TransformIndex) const;
-	FBox GenerateBounds(const FGeometryCollection* GeometryCollection, int32 TransformIndex) const;
-	void InitializePartitions();
-	bool Refine();
-	int32 FindClosestPartitionCenter(const FVector& Location) const;
-	void MarkVisited(int32 Index, int32 PartitionIndex);
-
-private:
-	TArray<int32> TransformIndices;
-	TArray<FVector> Centroids;
-	TArray<int32> Partitions;
-	int32 PartitionCount;
-	TArray<int32> PartitionSize;
-	TArray<FVector> PartitionCenters;
-	TArray<TSet<int32>> Connectivity;
-	TArray<bool> Visited;
-
-
-	// Not generally necessary but this is a safety measure to prevent oscillating solves that never converge.
-	const int32 MaxKMeansIterations = 500;
-};
\ No newline at end of file
diff --git a/Engine/Plugins/Experimental/ChaosFlesh/Source/ChaosFleshNodes/Public/Dataflow/ChaosFleshEngineAssetNodes.h b/Engine/Plugins/Experimental/ChaosFlesh/Source/ChaosFleshNodes/Public/Dataflow/ChaosFleshEngineAssetNodes.h
--- a/Engine/Plugins/Experimental/ChaosFlesh/Source/ChaosFleshNodes/Public/Dataflow/ChaosFleshEngineAssetNodes.h
+++ b/Engine/Plugins/Experimental/ChaosFlesh/Source/ChaosFleshNodes/Public/Dataflow/ChaosFleshEngineAssetNodes.h
@@ -43,7 +43,7 @@ struct FFleshAssetTerminalDataflowNode : public FDataflowTerminalNode
 
 public:
 
-	UPROPERTY(meta = (DataflowInput, DataflowOutput, Passthrough = "Collection", DisplayName = "Collection"))
+	UPROPERTY(meta = (DataflowInput, DataflowOutput, DisplayName = "Collection", DataflowPassthrough  = "Collection"))
 	FManagedArrayCollection Collection;
 
 
@@ -126,7 +126,7 @@ public:
 	UPROPERTY(EditAnywhere, Category = "Dataflow")
 	float Density = 1.f;
 
-	UPROPERTY(meta = (DataflowInput, DataflowOutput, DisplayName = "Collection", Passthrough = "Collection"))
+	UPROPERTY(meta = (DataflowInput, DataflowOutput, DisplayName = "Collection", DataflowPassthrough  = "Collection"))
 	FManagedArrayCollection Collection;
 
 	FComputeFleshMassNode(const Dataflow::FNodeParameters& InParam, FGuid InGuid = FGuid::NewGuid())
@@ -154,7 +154,7 @@ struct FComputeFiberFieldNode : public FDataflowNode
 public:
 	//typedef FManagedArrayCollection DataType;
 
-	UPROPERTY(meta = (DataflowInput, DataflowOutput, DisplayName = "Collection", Passthrough = "Collection"))
+	UPROPERTY(meta = (DataflowInput, DataflowOutput, DisplayName = "Collection", DataflowPassthrough  = "Collection"))
 	FManagedArrayCollection Collection;
 
 	//UPROPERTY(meta = (DataflowInput, DisplayName = "OriginVertexIndices"))
diff --git a/Engine/Plugins/Experimental/ChaosFlesh/Source/ChaosFleshNodes/Public/Dataflow/ChaosFleshFiberDirectionInitializationNodes.h b/Engine/Plugins/Experimental/ChaosFlesh/Source/ChaosFleshNodes/Public/Dataflow/ChaosFleshFiberDirectionInitializationNodes.h
--- a/Engine/Plugins/Experimental/ChaosFlesh/Source/ChaosFleshNodes/Public/Dataflow/ChaosFleshFiberDirectionInitializationNodes.h
+++ b/Engine/Plugins/Experimental/ChaosFlesh/Source/ChaosFleshNodes/Public/Dataflow/ChaosFleshFiberDirectionInitializationNodes.h
@@ -18,7 +18,7 @@ struct FGenerateFiberDirectionsDataflowNode : public FDataflowNode
 public:
 	typedef FManagedArrayCollection DataType;
 
-	UPROPERTY(meta = (DataflowInput, DataflowOutput, DisplayName = "Collection", Passthrough = "Collection"))
+	UPROPERTY(meta = (DataflowInput, DataflowOutput, DisplayName = "Collection", DataflowPassthrough = "Collection"))
 	FManagedArrayCollection Collection;
 
 	FGenerateFiberDirectionsDataflowNode(const Dataflow::FNodeParameters& InParam, FGuid InGuid = FGuid::NewGuid())
diff --git a/Engine/Plugins/Experimental/ChaosFlesh/Source/ChaosFleshNodes/Public/Dataflow/ChaosFleshKinematicInitializationNodes.h b/Engine/Plugins/Experimental/ChaosFlesh/Source/ChaosFleshNodes/Public/Dataflow/ChaosFleshKinematicInitializationNodes.h
--- a/Engine/Plugins/Experimental/ChaosFlesh/Source/ChaosFleshNodes/Public/Dataflow/ChaosFleshKinematicInitializationNodes.h
+++ b/Engine/Plugins/Experimental/ChaosFlesh/Source/ChaosFleshNodes/Public/Dataflow/ChaosFleshKinematicInitializationNodes.h
@@ -23,7 +23,7 @@ public:
 	UPROPERTY(meta = (DataflowInput, DisplayName = "SkeletalMesh"))
 	TObjectPtr<USkeletalMesh> SkeletalMeshIn = nullptr;
 
-	UPROPERTY(meta = (DataflowInput, DataflowOutput, DisplayName = "Collection", Passthrough = "Collection"))
+	UPROPERTY(meta = (DataflowInput, DataflowOutput, DisplayName = "Collection", DataflowPassthrough  = "Collection"))
 	FManagedArrayCollection Collection;
 
 	UPROPERTY(EditAnywhere, Category = "Dataflow")
@@ -57,7 +57,7 @@ public:
 	UPROPERTY(EditAnywhere, Category = "Dataflow")
 	FTransform Transform;
 
-	UPROPERTY(meta = (DataflowInput, DataflowOutput, DisplayName = "Collection", Passthrough = "Collection"))
+	UPROPERTY(meta = (DataflowInput, DataflowOutput, DisplayName = "Collection", DataflowPassthrough  = "Collection"))
 	FManagedArrayCollection Collection;
 
 	UPROPERTY(meta = (DataflowInput, DisplayName = "SkeletalMesh"))
@@ -119,7 +119,7 @@ struct FSetVerticesKinematicDataflowNode : public FDataflowNode
 public:
 	typedef FManagedArrayCollection DataType;
 
-	UPROPERTY(meta = (DataflowInput, DataflowOutput, DisplayName = "Collection", Passthrough = "Collection"))
+	UPROPERTY(meta = (DataflowInput, DataflowOutput, DisplayName = "Collection", DataflowPassthrough  = "Collection"))
 		FManagedArrayCollection Collection;
 
 	UPROPERTY(meta = (DataflowInput, DisplayName = "SelectionSet"))
diff --git a/Engine/Plugins/Experimental/ChaosFlesh/Source/ChaosFleshNodes/Public/Dataflow/ChaosFleshRenderInitializationNodes.h b/Engine/Plugins/Experimental/ChaosFlesh/Source/ChaosFleshNodes/Public/Dataflow/ChaosFleshRenderInitializationNodes.h
--- a/Engine/Plugins/Experimental/ChaosFlesh/Source/ChaosFleshNodes/Public/Dataflow/ChaosFleshRenderInitializationNodes.h
+++ b/Engine/Plugins/Experimental/ChaosFlesh/Source/ChaosFleshNodes/Public/Dataflow/ChaosFleshRenderInitializationNodes.h
@@ -23,7 +23,7 @@ public:
 	UPROPERTY(meta = (DataflowInput, DisplayName = "SkeletalMesh"))
 	TObjectPtr<const USkeletalMesh> SkeletalMesh = nullptr;
 
-	UPROPERTY(meta = (DataflowInput, DataflowOutput, DisplayName = "Collection", Passthrough = "Collection"))
+	UPROPERTY(meta = (DataflowInput, DataflowOutput, DisplayName = "Collection", DataflowPassthrough = "Collection"))
 		FManagedArrayCollection Collection;
 
 	FBindForRenderToSkeletalMeshDataflowNode(const Dataflow::FNodeParameters& InParam, FGuid InGuid = FGuid::NewGuid())
diff --git a/Engine/Plugins/Experimental/Fracture/Source/FractureEngine/FractureEngine.Build.cs b/Engine/Plugins/Experimental/Fracture/Source/FractureEngine/FractureEngine.Build.cs
--- a/Engine/Plugins/Experimental/Fracture/Source/FractureEngine/FractureEngine.Build.cs
+++ b/Engine/Plugins/Experimental/Fracture/Source/FractureEngine/FractureEngine.Build.cs
@@ -1,28 +1,30 @@
 // Copyright Epic Games, Inc. All Rights Reserved.
 
-using UnrealBuildTool;
+using System;
 
-public class FractureEngine : ModuleRules
+namespace UnrealBuildTool.Rules
 {
-	public FractureEngine(ReadOnlyTargetRules Target) : base(Target)
+	public class FractureEngine : ModuleRules
 	{
-		PrivateIncludePaths.Add(ModuleDirectory + "/Private");
-		PublicIncludePaths.Add(ModuleDirectory + "/Public");
-		PublicDependencyModuleNames.AddRange(
-			new string[]
-			{
+		public FractureEngine(ReadOnlyTargetRules Target) : base(Target)
+		{
+			PrivateIncludePaths.Add("FractureEngine/Private");
+			PublicIncludePaths.Add(ModuleDirectory + "/Public");
+			PublicDependencyModuleNames.AddRange(
+				new string[]
+				{
+					"Chaos",
+					"Core",
+					"DataflowCore"
+				}
+				);
 
-				// ... add other public dependencies that you statically link with here ...
-			}
-            );
-			
-		
-		PrivateDependencyModuleNames.AddRange(
-			new string[]
-			{
-
-				// ... add private dependencies that you statically link with here ...	
-			}
-            );
+			PrivateDependencyModuleNames.AddRange(
+				new string[]
+				{
+				}
+				);
+		}
 	}
 }
+
diff --git a/Engine/Plugins/Experimental/Fracture/Source/FractureEngine/Private/FractureEngineClustering.cpp b/Engine/Plugins/Experimental/Fracture/Source/FractureEngine/Private/FractureEngineClustering.cpp
--- a/Engine/Plugins/Experimental/Fracture/Source/FractureEngine/Private/FractureEngineClustering.cpp
+++ b/Engine/Plugins/Experimental/Fracture/Source/FractureEngine/Private/FractureEngineClustering.cpp
@@ -1,2 +1,541 @@
 // Copyright Epic Games, Inc. All Rights Reserved.
 
+#include "FractureEngineClustering.h"
+#include "GeometryCollection/GeometryCollectionClusteringUtility.h"
+#include "GeometryCollection/GeometryCollectionAlgo.h"
+#include "GeometryCollection/GeometryCollectionProximityUtility.h"
+
+FVoronoiPartitioner::FVoronoiPartitioner(const FGeometryCollection* GeometryCollection, int32 ClusterIndex)
+{
+	// Collect children of cluster
+	const TManagedArray<TSet<int32>>& Children = GeometryCollection->Children;
+	TransformIndices = Children[ClusterIndex].Array();
+
+	GenerateCentroids(GeometryCollection);
+}
+
+void FVoronoiPartitioner::KMeansPartition(int32 InPartitionCount)
+{
+	PartitionCount = InPartitionCount;
+
+	InitializePartitions();
+	for (int32 Iteration = 0; Iteration < MaxKMeansIterations; Iteration++)
+	{
+		// Refinement is complete when no nodes change partition.
+		if (!Refine())
+		{
+			break;
+		}
+	}
+}
+
+void FVoronoiPartitioner::MergeSingleElementPartitions(FGeometryCollection* GeometryCollection)
+{
+	if (Connectivity.IsEmpty())
+	{
+		GenerateConnectivity(GeometryCollection);
+	}
+	for (int32 ElIdx = 0; ElIdx < TransformIndices.Num(); ElIdx++)
+	{
+		int32 Partition = Partitions[ElIdx];
+		if (PartitionSize[Partition] == 1)
+		{
+			// Find the smallest neighboring partition to merge to
+			// (to help keep partition sizes balanced)
+			int32 SmallestNbrPartition = -1;
+			int32 SmallestNbrSize = TransformIndices.Num()+1;
+			for (int32 NbrEl : Connectivity[ElIdx])
+			{
+				int32 NbrPartition = Partitions[NbrEl];
+				if (NbrPartition == Partition)
+				{
+					continue;
+				}
+				int32 NbrSize = PartitionSize[NbrPartition];
+				if (NbrSize > 0 && (SmallestNbrPartition == -1 || NbrSize < SmallestNbrSize))
+				{
+					SmallestNbrPartition = NbrPartition;
+					SmallestNbrSize = NbrSize;
+				}
+			}
+			if (SmallestNbrPartition != -1)
+			{
+				Partitions[ElIdx] = SmallestNbrPartition;
+				PartitionSize[Partition]--;
+				PartitionSize[SmallestNbrPartition]++;
+			}
+		}
+	}
+}
+
+void FVoronoiPartitioner::SplitDisconnectedPartitions(FGeometryCollection* GeometryCollection)
+{
+	Visited.Init(false, TransformIndices.Num());
+	if (Connectivity.IsEmpty())
+	{
+		GenerateConnectivity(GeometryCollection);
+	}
+	
+	for (int32 PartitionIndex = 0; PartitionIndex < PartitionCount; ++PartitionIndex)
+	{
+		int32 FirstElementInPartitionIndex = Partitions.Find(PartitionIndex);
+		if (FirstElementInPartitionIndex > INDEX_NONE)
+		{
+			MarkVisited(FirstElementInPartitionIndex, PartitionIndex);
+
+			// Place unvisited partition members in a new partition.
+			bool bFoundUnattached = false;
+			for (int32 Index = 0; Index < TransformIndices.Num(); ++Index)
+			{
+				if (Partitions[Index] == PartitionIndex)
+				{
+					if (!Visited[Index])
+					{
+						if (!bFoundUnattached)
+						{
+							bFoundUnattached = true;
+							PartitionCount++;
+							PartitionSize.Add(0);
+						}
+						Partitions[Index] = PartitionCount - 1;
+						PartitionSize[PartitionCount - 1]++;
+						PartitionSize[PartitionIndex]--;
+					}
+				}
+			}
+		}
+	}
+}
+
+void FVoronoiPartitioner::MarkVisited(int32 Index, int32 PartitionIndex)
+{
+	Visited[Index] = true;
+	for (int32 Adjacent : Connectivity[Index])
+	{
+		if (!Visited[Adjacent] && Partitions[Adjacent] == PartitionIndex)
+		{
+			MarkVisited(Adjacent, PartitionIndex);
+		}
+	}
+}
+
+TArray<int32> FVoronoiPartitioner::GetPartition(int32 PartitionIndex) const
+{
+	TArray<int32> PartitionContents;
+	PartitionContents.Reserve(TransformIndices.Num());
+	for (int32 Index = 0; Index < TransformIndices.Num(); ++Index)
+	{
+		if (Partitions[Index] == PartitionIndex)
+		{
+			PartitionContents.Add(TransformIndices[Index]);
+		}
+	}
+
+	return PartitionContents;
+}
+
+void FVoronoiPartitioner::GenerateConnectivity(const FGeometryCollection* GeometryCollection)
+{
+	
+	Connectivity.SetNum(TransformIndices.Num());
+
+	if (!ensure(GeometryCollection->HasAttribute("Level", FGeometryCollection::TransformGroup)))
+	{
+		return;
+	}
+	const TManagedArray<int32>& Levels = GeometryCollection->GetAttribute<int32>("Level", FGeometryCollection::TransformGroup);
+
+	for (int32 Index = 0; Index < TransformIndices.Num(); ++Index)
+	{
+		CollectConnections(GeometryCollection, TransformIndices[Index], Levels[TransformIndices[Index]], Connectivity[Index]);
+	}
+}
+
+void FVoronoiPartitioner::CollectConnections(const FGeometryCollection* GeometryCollection, int32 Index, int32 OperatingLevel, TSet<int32>& OutConnections) const
+{
+	const TManagedArray<TSet<int32>>& Children = GeometryCollection->Children;
+	if (GeometryCollection->IsRigid(Index))
+	{
+		const TManagedArray<TSet<int32>>& Proximity = GeometryCollection->GetAttribute<TSet<int32>>("Proximity", FGeometryCollection::GeometryGroup);
+		const TManagedArray<int32>& GeometryToTransformIndex = GeometryCollection->TransformIndex;
+		const TManagedArray<int32>& TransformToGeometryIndex = GeometryCollection->TransformToGeometryIndex;
+
+		for (int32 Neighbor : Proximity[TransformToGeometryIndex[Index]])
+		{
+			int32 NeighborTransformIndex = FGeometryCollectionClusteringUtility::GetParentOfBoneAtSpecifiedLevel(GeometryCollection, GeometryToTransformIndex[Neighbor], OperatingLevel);
+			int32 ContextIndex = TransformIndices.Find(NeighborTransformIndex);
+			if (ContextIndex > INDEX_NONE)
+			{
+				OutConnections.Add(ContextIndex);
+			}
+		}
+	}
+	else
+	{
+		for (int32 ChildIndex : Children[Index])
+		{
+			CollectConnections(GeometryCollection, ChildIndex, OperatingLevel, OutConnections);
+		}
+	}
+}
+
+void FVoronoiPartitioner::GenerateCentroids(const FGeometryCollection* GeometryCollection)
+{
+	Centroids.SetNum(TransformIndices.Num());
+	ParallelFor(TransformIndices.Num(), [this, GeometryCollection](int32 Index)
+		{
+			Centroids[Index] = GenerateCentroid(GeometryCollection, TransformIndices[Index]);
+		});
+}
+
+FVector FVoronoiPartitioner::GenerateCentroid(const FGeometryCollection* GeometryCollection, int32 TransformIndex) const
+{
+	FBox Bounds = GenerateBounds(GeometryCollection, TransformIndex);
+	return Bounds.GetCenter();
+}
+
+FBox FVoronoiPartitioner::GenerateBounds(const FGeometryCollection* GeometryCollection, int32 TransformIndex) const
+{
+	check(GeometryCollection->IsRigid(TransformIndex) || GeometryCollection->IsClustered(TransformIndex));
+
+	// Return the bounds of all the vertices contained by this branch.
+
+	if (GeometryCollection->IsRigid(TransformIndex))
+	{
+		const TManagedArray<FTransform>& Transforms = GeometryCollection->Transform;
+		const TManagedArray<int32>& Parents = GeometryCollection->Parent;
+		FTransform GlobalTransform = GeometryCollectionAlgo::GlobalMatrix(Transforms, Parents, TransformIndex);
+		
+		int32 GeometryIndex = GeometryCollection->TransformToGeometryIndex[TransformIndex];
+		int32 VertexStart = GeometryCollection->VertexStart[GeometryIndex];
+		int32 VertexCount = GeometryCollection->VertexCount[GeometryIndex];
+		const TManagedArray<FVector3f>& GCVertices = GeometryCollection->Vertex;
+
+		TArray<FVector> Vertices;
+		Vertices.SetNum(VertexCount);
+		for (int32 VertexOffset = 0; VertexOffset < VertexCount; ++VertexOffset)
+		{
+			Vertices[VertexOffset] = GlobalTransform.TransformPosition((FVector)GCVertices[VertexStart + VertexOffset]);
+		}
+
+		return FBox(Vertices);
+	}
+	else // recurse the cluster
+	{
+		const TArray<int32> Children = GeometryCollection->Children[TransformIndex].Array();
+
+		// Empty clusters are problematic
+		if (Children.Num() == 0)
+		{
+			return FBox();
+		}
+
+		FBox Bounds = GenerateBounds(GeometryCollection, Children[0]);
+		for (int32 ChildIndex = 1; ChildIndex < Children.Num(); ++ChildIndex)
+		{
+			Bounds += GenerateBounds(GeometryCollection, Children[ChildIndex]);
+		}
+
+		return Bounds;
+	}
+}
+
+void FVoronoiPartitioner::InitializePartitions()
+{
+	check(PartitionCount > 0);
+
+	PartitionCount = FMath::Min(PartitionCount, TransformIndices.Num());
+
+	// Set initial partition centers as selects from the vertex set
+	PartitionCenters.SetNum(PartitionCount);
+	int32 TransformStride = FMath::Max(1,FMath::FloorToInt(float(TransformIndices.Num()) / float(PartitionCount)));
+	for (int32 PartitionIndex = 0; PartitionIndex < PartitionCount; ++PartitionIndex)
+	{
+		PartitionCenters[PartitionIndex] = Centroids[TransformStride * PartitionIndex];
+	}
+
+	// At beginning, all nodes belong to first partition.
+	Partitions.Init(0, TransformIndices.Num());
+	PartitionSize.Init(0, PartitionCount);
+
+	if (PartitionCount > 0)
+	{
+		PartitionSize[0] = TransformIndices.Num();
+	}
+}
+
+bool FVoronoiPartitioner::Refine()
+{
+	// Assign each transform to its closest partition center.
+	bool bUnchanged = true;
+	for (int32 Index = 0; Index < Centroids.Num(); ++Index)
+	{
+		int32 ClosestPartition = FindClosestPartitionCenter(Centroids[Index]);
+		if (ClosestPartition != Partitions[Index])
+		{
+			bUnchanged = false;
+			PartitionSize[Partitions[Index]]--;
+			PartitionSize[ClosestPartition]++;
+			Partitions[Index] = ClosestPartition;
+		}
+	}
+
+	if (bUnchanged)
+	{
+		return false;
+	}
+
+	// Recalculate partition centers
+	for (int32 PartitionIndex = 0; PartitionIndex < PartitionCount; ++PartitionIndex)
+	{
+		if (PartitionSize[PartitionIndex] > 0)
+		{
+			PartitionCenters[PartitionIndex] = FVector(0.0f, 0.0f, 0.0f);
+		}
+	}
+
+	for (int32 Index = 0; Index < Partitions.Num(); ++Index)
+	{
+		PartitionCenters[Partitions[Index]] += Centroids[Index];
+	}
+
+	for (int32 PartitionIndex = 0; PartitionIndex < PartitionCount; ++PartitionIndex)
+	{
+		if (PartitionSize[PartitionIndex] > 0)
+		{
+			PartitionCenters[PartitionIndex] /= float(PartitionSize[PartitionIndex]);
+		}
+	}
+
+	return true;
+}
+
+int32 FVoronoiPartitioner::FindClosestPartitionCenter(const FVector& Location) const
+{
+	int32 ClosestPartition = INDEX_NONE;
+	float SmallestDistSquared = TNumericLimits<float>::Max();
+
+	for (int32 PartitionIndex = 0; PartitionIndex < PartitionCount; ++PartitionIndex)
+	{
+		float DistSquared = FVector::DistSquared(Location, PartitionCenters[PartitionIndex]);
+		if (DistSquared < SmallestDistSquared)
+		{
+			SmallestDistSquared = DistSquared;
+			ClosestPartition = PartitionIndex;
+		}
+	}
+
+	check(ClosestPartition > INDEX_NONE);
+	return ClosestPartition;
+}
+
+class FFractureEngineBoneSelection
+{
+public:
+	FFractureEngineBoneSelection(const FGeometryCollection& InGeometryCollection, const TArray<int32>& InSelectedBones)
+		: GeometryCollection(InGeometryCollection)
+		, SelectedBones(InSelectedBones)
+	{}
+
+	const TArray<int32>& GetSelectedBones() const { return SelectedBones;  }
+	void ConvertSelectionToClusterNodes();
+
+private:
+	bool IsValidBone(int32 Index) const;
+	bool HasSelectedAncestor(int32 Index) const;
+	void Sanitize(bool bFavorParents = true);
+
+private:
+	const FGeometryCollection& GeometryCollection;
+	TArray<int32> SelectedBones;
+};
+
+bool FFractureEngineBoneSelection::IsValidBone(int32 Index) const
+{
+	return Index >= 0 && Index < GeometryCollection.Parent.Num();
+}
+
+bool FFractureEngineBoneSelection::HasSelectedAncestor(int32 Index) const
+{
+	const TManagedArray<int32>& Parents = GeometryCollection.GetAttribute<int32>("Parent", FGeometryCollection::TransformGroup);
+
+	if (!ensureMsgf(Index >= 0 && Index < Parents.Num(), TEXT("Invalid index in selection: %d"), Index))
+	{
+		return false;
+	}
+
+	int32 CurrIndex = Index;
+	while (Parents[CurrIndex] != INDEX_NONE)
+	{
+		CurrIndex = Parents[CurrIndex];
+		if (SelectedBones.Contains(CurrIndex))
+		{
+			return true;
+		}
+	}
+
+	// We've arrived at the top of the hierarchy with no selected ancestors
+	return false;
+}
+
+void FFractureEngineBoneSelection::Sanitize(bool bFavorParents)
+{
+	// Ensure that selected indices are valid
+	int NumTransforms = GeometryCollection.NumElements(FGeometryCollection::TransformGroup);
+	SelectedBones.RemoveAll([NumTransforms](int32 Index) {
+		return Index == INDEX_NONE || !ensure(Index < NumTransforms);
+	});
+
+	// Ensure that children of a selected node are not also selected.
+	if (bFavorParents)
+	{
+		SelectedBones.RemoveAll([this](int32 Index) {
+			return !IsValidBone(Index) || HasSelectedAncestor(Index);
+		});
+	}
+
+	SelectedBones.Sort();
+}
+
+void FFractureEngineBoneSelection::ConvertSelectionToClusterNodes()
+{
+	// If this is a non-cluster node, select the cluster containing it instead.
+	const TManagedArray<TSet<int32>>& Children = GeometryCollection.Children;
+	const TManagedArray<int32>& Parents = GeometryCollection.Parent;
+	const TManagedArray<int32>& SimulationType = GeometryCollection.SimulationType;
+
+	Sanitize();
+
+	TArray<int32> AddedClusterSelections;
+	for (int32 Index : SelectedBones)
+	{
+		int32 AddParent = FGeometryCollection::Invalid;
+		if (SimulationType[Index] == FGeometryCollection::ESimulationTypes::FST_Rigid)
+		{
+			AddParent = Parents[Index];
+		}
+		else if (SimulationType[Index] == FGeometryCollection::ESimulationTypes::FST_None)
+		{
+			AddParent = Parents[Parents[Index]];
+		}
+		if (AddParent != FGeometryCollection::Invalid)
+		{
+			AddedClusterSelections.AddUnique(AddParent);
+		}
+	}
+	SelectedBones.Append(AddedClusterSelections);
+
+	Sanitize();
+}
+
+void FFractureEngineClustering::AutoCluster(FGeometryCollection& GeometryCollection,
+	const TArray<int32>& BoneIndices,
+	const EFractureEngineClusterSizeMethod ClusterSizeMethod,
+	const uint32 SiteCount,
+	const float SiteCountFraction,
+	const float SiteSize,
+	const bool bEnforceConnectivity,
+	const bool bAvoidIsolated)
+{
+	FFractureEngineBoneSelection Selection(GeometryCollection, BoneIndices);
+
+	Selection.ConvertSelectionToClusterNodes();
+
+	for (const int32 ClusterIndex : Selection.GetSelectedBones())
+	{
+		AutoCluster(GeometryCollection, ClusterIndex, ClusterSizeMethod, SiteCount, SiteCountFraction, SiteSize, bEnforceConnectivity, bAvoidIsolated);
+	}
+}
+
+void FFractureEngineClustering::AutoCluster(FGeometryCollection& GeometryCollection,
+	const int32 ClusterIndex,
+	const EFractureEngineClusterSizeMethod ClusterSizeMethod,
+	const uint32 SiteCount,
+	const float SiteCountFraction,
+	const float SiteSize,
+	const bool bEnforceConnectivity,
+	const bool bAvoidIsolated)
+{
+	FVoronoiPartitioner VoronoiPartition(&GeometryCollection, ClusterIndex);
+	int32 NumChildren = GeometryCollection.Children[ClusterIndex].Num();
+
+	int32 SiteCountToUse = 1;
+	if (ClusterSizeMethod == EFractureEngineClusterSizeMethod::ByNumber)
+	{
+		SiteCountToUse = SiteCount;
+	}
+	else if (ClusterSizeMethod == EFractureEngineClusterSizeMethod::ByFractionOfInput)
+	{
+		SiteCountToUse = FMath::Max(2, NumChildren * SiteCountFraction);
+	}
+	else if (ClusterSizeMethod == EFractureEngineClusterSizeMethod::BySize)
+	{
+		float TotalVolume = GeometryCollection.GetBoundingBox().GetBox().GetVolume();
+		float DesiredVolume = FMath::Pow(SiteSize, 3);
+		SiteCountToUse = FMath::Max(1, TotalVolume / DesiredVolume);
+	}
+
+	VoronoiPartition.KMeansPartition(SiteCountToUse);
+	if (VoronoiPartition.GetPartitionCount() == 0)
+	{
+		return;
+	}
+
+	bool bNeedProximity = bEnforceConnectivity || bAvoidIsolated;
+	if (bNeedProximity)
+	{
+		FGeometryCollectionProximityUtility ProximityUtility(&GeometryCollection);
+		ProximityUtility.UpdateProximity();
+	}
+
+	if (bEnforceConnectivity)
+	{
+		VoronoiPartition.SplitDisconnectedPartitions(&GeometryCollection);
+	}
+
+	if (bAvoidIsolated)
+	{
+		// attempt to remove isolated via merging (may not succeed, as it only merges if there is a cluster in proximity)
+		VoronoiPartition.MergeSingleElementPartitions(&GeometryCollection);
+	}
+
+	int32 NonEmptyPartitionCount = VoronoiPartition.GetNonEmptyPartitionCount();
+	bool bHasEmptyClusters = NonEmptyPartitionCount > 0;
+
+	if (bAvoidIsolated && NonEmptyPartitionCount == 1)
+	{
+		return;
+	}
+
+	int32 PartitionCount = VoronoiPartition.GetPartitionCount();
+	int32 NewClusterIndexStart = GeometryCollection.AddElements(PartitionCount, FGeometryCollection::TransformGroup);
+
+	for (int32 Index = 0; Index < PartitionCount; ++Index)
+	{
+		TArray<int32> NewCluster = VoronoiPartition.GetPartition(Index);
+		if (bAvoidIsolated && NewCluster.Num() == 1)
+		{
+			bHasEmptyClusters = true;
+			NewCluster.Reset();
+		}
+
+		int32 NewClusterIndex = NewClusterIndexStart + Index;
+		GeometryCollection.Parent[NewClusterIndex] = ClusterIndex;
+		GeometryCollection.Children[ClusterIndex].Add(NewClusterIndex);
+		GeometryCollection.BoneName[NewClusterIndex] = "ClusterBone";
+		GeometryCollection.Children[NewClusterIndex] = TSet<int32>(NewCluster);
+		GeometryCollection.SimulationType[NewClusterIndex] = FGeometryCollection::ESimulationTypes::FST_Clustered;
+		GeometryCollection.Transform[NewClusterIndex] = FTransform::Identity;
+		GeometryCollectionAlgo::ParentTransforms(&GeometryCollection, NewClusterIndex, NewCluster);
+	}
+	if (bHasEmptyClusters)
+	{
+		FGeometryCollectionClusteringUtility::RemoveDanglingClusters(&GeometryCollection);
+	}
+	FGeometryCollectionClusteringUtility::UpdateHierarchyLevelOfChildren(&GeometryCollection, ClusterIndex);
+	FGeometryCollectionClusteringUtility::RecursivelyUpdateChildBoneNames(ClusterIndex, GeometryCollection.Children, GeometryCollection.BoneName);
+	FGeometryCollectionClusteringUtility::ValidateResults(&GeometryCollection);
+}
+
+
+
diff --git a/Engine/Plugins/Experimental/Fracture/Source/FractureEngine/Private/FractureEnginePlugin.cpp b/Engine/Plugins/Experimental/Fracture/Source/FractureEngine/Private/FractureEnginePlugin.cpp
new file mode 100644
--- /dev/null
+++ b/Engine/Plugins/Experimental/Fracture/Source/FractureEngine/Private/FractureEnginePlugin.cpp
@@ -0,0 +1,33 @@
+// Copyright Epic Games, Inc. All Rights Reserved.
+
+#include "FractureEnginePlugin.h"
+
+#include "CoreMinimal.h"
+#include "Modules/ModuleManager.h"
+
+
+
+class FFractureEnginePlugin : public IFractureEnginePlugin
+{
+	/** IModuleInterface implementation */
+	virtual void StartupModule() override;
+	virtual void ShutdownModule() override;
+};
+
+IMPLEMENT_MODULE( FFractureEnginePlugin, FractureEngine )
+
+
+
+void FFractureEnginePlugin::StartupModule()
+{
+
+}
+
+
+void FFractureEnginePlugin::ShutdownModule()
+{
+
+}
+
+
+
diff --git a/Engine/Plugins/Experimental/Fracture/Source/FractureEngine/Private/FractureEngineSelection.cpp b/Engine/Plugins/Experimental/Fracture/Source/FractureEngine/Private/FractureEngineSelection.cpp
--- a/Engine/Plugins/Experimental/Fracture/Source/FractureEngine/Private/FractureEngineSelection.cpp
+++ b/Engine/Plugins/Experimental/Fracture/Source/FractureEngine/Private/FractureEngineSelection.cpp
@@ -1,2 +1,374 @@
 // Copyright Epic Games, Inc. All Rights Reserved.
 
+#include "FractureEngineSelection.h"
+#include "GeometryCollection/GeometryCollectionClusteringUtility.h"
+#include "GeometryCollection/GeometryCollectionAlgo.h"
+#include "GeometryCollection/GeometryCollectionProximityUtility.h"
+#include "Algo/RemoveIf.h"
+
+static void ConvertSelectionSetToSelectionArr(const TSet<int32>& SelectionSet, TArray<int32>& SelectionArr)
+{
+	SelectionArr.Empty();
+
+	for (auto& BoneIdx : SelectionSet)
+	{
+		SelectionArr.Add(BoneIdx);
+	}
+}
+
+void FFractureEngineSelection::SelectParent(const FGeometryCollection& GeometryCollection, TArray<int32>& SelectedBones)
+{
+	if (GeometryCollection.HasGroup(FGeometryCollection::TransformGroup) &&
+		GeometryCollection.HasAttribute("Parent", FGeometryCollection::TransformGroup))
+	{
+		const TManagedArray<int32>& Parents = GeometryCollection.GetAttribute<int32>("Parent", FGeometryCollection::TransformGroup);
+
+		TSet<int32> NewSelection;
+		for (int32 Bone : SelectedBones)
+		{
+			int32 ParentBone = Parents[Bone];
+			if (ParentBone != FGeometryCollection::Invalid)
+			{
+				NewSelection.Add(ParentBone);
+			}
+		}
+
+		ConvertSelectionSetToSelectionArr(NewSelection, SelectedBones);
+	}
+}
+
+void FFractureEngineSelection::SelectParent(const FGeometryCollection& GeometryCollection, FDataflowTransformSelection& TransformSelection)
+{
+	TArray<int32> SelectionArr;
+	TransformSelection.AsArray(SelectionArr);
+
+	SelectParent(GeometryCollection, SelectionArr);
+
+	TransformSelection.SetFromArray(SelectionArr);
+}
+
+void FFractureEngineSelection::SelectChildren(const FGeometryCollection& GeometryCollection, TArray<int32>& SelectedBones)
+{
+	if (GeometryCollection.HasGroup(FGeometryCollection::TransformGroup) &&
+		GeometryCollection.HasAttribute("Children", FGeometryCollection::TransformGroup))
+	{
+		const TManagedArray<TSet<int32>>& Children = GeometryCollection.GetAttribute<TSet<int32>>("Children", FGeometryCollection::TransformGroup);
+
+		TSet<int32> NewSelection;
+		for (int32 Bone : SelectedBones)
+		{
+			if (Children[Bone].IsEmpty())
+			{
+				NewSelection.Add(Bone);
+				continue;
+			}
+			for (int32 Child : Children[Bone])
+			{
+				NewSelection.Add(Child);
+			}
+		}
+
+		ConvertSelectionSetToSelectionArr(NewSelection, SelectedBones);
+	}
+}
+
+void FFractureEngineSelection::SelectChildren(const FGeometryCollection& GeometryCollection, FDataflowTransformSelection& TransformSelection)
+{
+	TArray<int32> SelectionArr;
+	TransformSelection.AsArray(SelectionArr);
+
+	SelectChildren(GeometryCollection, SelectionArr);
+
+	TransformSelection.SetFromArray(SelectionArr);
+}
+
+void FFractureEngineSelection::SelectSiblings(const FGeometryCollection& GeometryCollection, TArray<int32>& SelectedBones)
+{
+	if (GeometryCollection.HasGroup(FGeometryCollection::TransformGroup) &&
+		GeometryCollection.HasAttribute("Parent", FGeometryCollection::TransformGroup) &&
+		GeometryCollection.HasAttribute("Children", FGeometryCollection::TransformGroup))
+	{
+		const TManagedArray<int32>& Parents = GeometryCollection.GetAttribute<int32>("Parent", FGeometryCollection::TransformGroup);
+		const TManagedArray<TSet<int32>>& Children = GeometryCollection.GetAttribute<TSet<int32>>("Children", FGeometryCollection::TransformGroup);
+
+		TSet<int32> NewSelection;
+		for (int32 Bone : SelectedBones)
+		{
+			int32 ParentBone = Parents[Bone];
+			if (ParentBone != FGeometryCollection::Invalid)
+			{
+				for (int32 Child : Children[ParentBone])
+				{
+					NewSelection.Add(Child);
+				}
+			}
+		}
+
+		ConvertSelectionSetToSelectionArr(NewSelection, SelectedBones);
+	}
+}
+
+void FFractureEngineSelection::SelectSiblings(const FGeometryCollection& GeometryCollection, FDataflowTransformSelection& TransformSelection)
+{
+	TArray<int32> SelectionArr;
+	TransformSelection.AsArray(SelectionArr);
+
+	SelectSiblings(GeometryCollection, SelectionArr);
+
+	TransformSelection.SetFromArray(SelectionArr);
+}
+
+void FFractureEngineSelection::SelectLevel(const FGeometryCollection& GeometryCollection, TArray<int32>& SelectedBones)
+{
+	if (GeometryCollection.HasGroup(FGeometryCollection::TransformGroup) &&
+		GeometryCollection.HasAttribute("Level", FGeometryCollection::TransformGroup))
+	{
+		const TManagedArray<int32>& Levels = GeometryCollection.GetAttribute<int32>("Levels", FGeometryCollection::TransformGroup);
+
+		TSet<int32> NewSelection;
+		for (int32 Bone : SelectedBones)
+		{
+			int32 Level = Levels[Bone];
+			for (int32 TransformIdx = 0; TransformIdx < GeometryCollection.NumElements(FTransformCollection::TransformGroup); ++TransformIdx)
+			{
+				if (Levels[TransformIdx] == Level)
+				{
+					NewSelection.Add(TransformIdx);
+				}
+			}
+		}
+
+		ConvertSelectionSetToSelectionArr(NewSelection, SelectedBones);
+	}
+}
+
+void FFractureEngineSelection::SelectLevel(const FGeometryCollection& GeometryCollection, FDataflowTransformSelection& TransformSelection)
+{
+	TArray<int32> SelectionArr;
+	TransformSelection.AsArray(SelectionArr);
+
+	SelectLevel(GeometryCollection, SelectionArr);
+
+	TransformSelection.SetFromArray(SelectionArr);
+}
+
+void FFractureEngineSelection::SelectContact(FGeometryCollection& GeometryCollection, TArray<int32>& SelectedBones)
+{
+	if (GeometryCollection.HasGroup(FGeometryCollection::TransformGroup) &&
+		GeometryCollection.HasAttribute("TransformIndex", FGeometryCollection::TransformGroup) &&
+		GeometryCollection.HasAttribute("TransformToGeometryIndex", FGeometryCollection::TransformGroup) &&
+		GeometryCollection.HasGroup(FGeometryCollection::GeometryGroup) &&
+		GeometryCollection.HasAttribute("Proximity", FGeometryCollection::GeometryGroup))
+	{
+		FGeometryCollectionProximityUtility ProximityUtility(&GeometryCollection);
+		ProximityUtility.RequireProximity();
+
+		const TManagedArray<int32>& TransformIndex = GeometryCollection.GetAttribute<int32>("TransformIndex", FGeometryCollection::TransformGroup);
+		const TManagedArray<int32>& TransformToGeometryIndex = GeometryCollection.GetAttribute<int32>("TransformToGeometryIndex", FGeometryCollection::TransformGroup);
+		const TManagedArray<TSet<int32>>& Proximity = GeometryCollection.GetAttribute<TSet<int32>>("Proximity", FGeometryCollection::GeometryGroup);
+
+		TSet<int32> NewSelection;
+		for (int32 Bone : SelectedBones)
+		{
+			NewSelection.Add(Bone);
+			int32 GeometryIdx = TransformToGeometryIndex[Bone];
+			if (GeometryIdx != INDEX_NONE)
+			{
+				const TSet<int32>& Neighbors = Proximity[GeometryIdx];
+				for (int32 NeighborGeometryIndex : Neighbors)
+				{
+					NewSelection.Add(TransformIndex[NeighborGeometryIndex]);
+				}
+			}
+		}
+
+		ConvertSelectionSetToSelectionArr(NewSelection, SelectedBones);
+	}
+}
+
+void FFractureEngineSelection::SelectContact(FGeometryCollection& GeometryCollection, FDataflowTransformSelection& TransformSelection)
+{
+	TArray<int32> SelectionArr;
+	TransformSelection.AsArray(SelectionArr);
+
+	SelectContact(GeometryCollection, SelectionArr);
+
+	TransformSelection.SetFromArray(SelectionArr);
+}
+
+void FFractureEngineSelection::SelectLeaf(const FGeometryCollection& GeometryCollection, TArray<int32>& SelectedBones)
+{
+	if (GeometryCollection.HasGroup(FGeometryCollection::TransformGroup) &&
+		GeometryCollection.HasAttribute("Level", FGeometryCollection::TransformGroup) &&
+		GeometryCollection.HasAttribute("SimulationType", FGeometryCollection::TransformGroup))
+	{
+		SelectedBones.Empty();
+
+		int32 ViewLevel = -1;
+
+		FGeometryCollectionClusteringUtility::GetBonesToLevel(&GeometryCollection, ViewLevel, SelectedBones, true, true);
+
+		const TManagedArray<int32>& SimType = GeometryCollection.GetAttribute<int32>("SimulationType", FGeometryCollection::TransformGroup);
+		const TManagedArray<int32>* Levels = GeometryCollection.FindAttributeTyped<int32>("Level", FGeometryCollection::TransformGroup);
+		
+		SelectedBones.SetNum(Algo::RemoveIf(SelectedBones, [&](int32 BoneIdx)
+			{
+				return SimType[BoneIdx] != FGeometryCollection::ESimulationTypes::FST_Rigid
+					|| (ViewLevel != -1 && Levels && (*Levels)[BoneIdx] != ViewLevel);
+			}));
+	}
+}
+
+void FFractureEngineSelection::SelectLeaf(const FGeometryCollection& GeometryCollection, FDataflowTransformSelection& TransformSelection)
+{
+	TArray<int32> SelectionArr;
+
+	SelectLeaf(GeometryCollection, SelectionArr);
+
+	TransformSelection.SetFromArray(SelectionArr);
+}
+
+void FFractureEngineSelection::SelectCluster(const FGeometryCollection& GeometryCollection, TArray<int32>& SelectedBones)
+{
+	if (GeometryCollection.HasGroup(FGeometryCollection::TransformGroup) &&
+		GeometryCollection.HasAttribute("Level", FGeometryCollection::TransformGroup) &&
+		GeometryCollection.HasAttribute("SimulationType", FGeometryCollection::TransformGroup))
+	{
+		SelectedBones.Empty();
+
+		int32 ViewLevel = -1;
+
+		FGeometryCollectionClusteringUtility::GetBonesToLevel(&GeometryCollection, ViewLevel, SelectedBones, true, true);
+
+		const TManagedArray<int32>& SimType = GeometryCollection.GetAttribute<int32>("SimulationType", FGeometryCollection::TransformGroup);
+		const TManagedArray<int32>* Levels = GeometryCollection.FindAttributeTyped<int32>("Level", FGeometryCollection::TransformGroup);
+
+		SelectedBones.SetNum(Algo::RemoveIf(SelectedBones, [&](int32 BoneIdx)
+			{
+				return SimType[BoneIdx] != FGeometryCollection::ESimulationTypes::FST_Rigid
+					|| (ViewLevel != -1 && Levels && (*Levels)[BoneIdx] != ViewLevel);
+			}));
+	}
+}
+
+void FFractureEngineSelection::SelectCluster(const FGeometryCollection& GeometryCollection, FDataflowTransformSelection& TransformSelection)
+{
+	TArray<int32> SelectionArr;
+
+	SelectCluster(GeometryCollection, SelectionArr);
+
+	TransformSelection.SetFromArray(SelectionArr);
+}
+
+static void RandomShuffleArray(TArray<int32>& InArray, const bool Deterministic, const float RandomSeed)
+{
+	FRandomStream Stream(RandomSeed);
+
+	int32 LastIndex = InArray.Num() - 1;
+
+	static const int32 NumIterationsMin = 10;
+	static const int32 NumIterationsMax = 50;
+
+	int32 NumIterations = Deterministic ? Stream.RandRange(NumIterationsMin, NumIterationsMax) : FMath::RandRange(NumIterationsMin, NumIterationsMax);
+
+	for (int32 IterIdx = 0; IterIdx < NumIterations; ++IterIdx)
+	{
+		for (int32 Idx = 0; Idx <= LastIndex; ++Idx)
+		{
+			int32 Index;
+			if (Deterministic)
+			{
+				Index = Stream.RandRange(Idx, LastIndex);
+			}
+			else
+			{
+				Index = FMath::RandRange(Idx, LastIndex);
+			}
+
+			if (Idx != Index)
+			{
+				InArray.Swap(Idx, Index);
+			}
+		}
+	}
+}
+
+void FFractureEngineSelection::SelectByPercentage(const FGeometryCollection& GeometryCollection, TArray<int32>& SelectedBones, const int32 Percentage, const bool Deterministic, const float RandomSeed)
+{
+	RandomShuffleArray(SelectedBones, Deterministic, RandomSeed);
+
+	int32 NewNumElements = FMath::RoundToInt32(SelectedBones.Num() * Percentage * 0.01f);
+	SelectedBones.SetNum(NewNumElements);
+}
+
+void FFractureEngineSelection::SelectByPercentage(const FGeometryCollection& GeometryCollection, FDataflowTransformSelection& TransformSelection, const int32 Percentage, const bool Deterministic, const float RandomSeed)
+{
+	TArray<int32> SelectionArr;
+	TransformSelection.AsArray(SelectionArr);
+
+	SelectByPercentage(GeometryCollection, SelectionArr, Percentage, Deterministic, RandomSeed);
+
+	TransformSelection.SetFromArray(SelectionArr);
+}
+
+
+void FFractureEngineSelection::SelectBySize(FGeometryCollection& GeometryCollection, TArray<int32>& SelectedBones, const float SizeMin, const float SizeMax)
+{
+	if (GeometryCollection.HasGroup(FGeometryCollection::TransformGroup) &&
+		GeometryCollection.HasAttribute("Size", FGeometryCollection::TransformGroup))
+	{
+		SelectedBones.Empty();
+
+		FGeometryCollectionConvexUtility::SetVolumeAttributes(&GeometryCollection);
+
+		const TManagedArray<float>& Sizes = GeometryCollection.GetAttribute<float>("Size", FTransformCollection::TransformGroup);
+
+		for (int32 BoneIdx = 0; BoneIdx < Sizes.Num(); ++BoneIdx)
+		{
+			if (Sizes[BoneIdx] >= SizeMin && Sizes[BoneIdx] <= SizeMax)
+			{
+				SelectedBones.Add(BoneIdx);
+			}
+		}
+	}
+}
+
+void FFractureEngineSelection::SelectBySize(FGeometryCollection& GeometryCollection, FDataflowTransformSelection& TransformSelection, const float SizeMin, const float SizeMax)
+{
+	TArray<int32> SelectionArr;
+
+	SelectBySize(GeometryCollection, SelectionArr, SizeMin, SizeMax);
+
+	TransformSelection.SetFromArray(SelectionArr);
+}
+
+void FFractureEngineSelection::SelectByVolume(FGeometryCollection& GeometryCollection, TArray<int32>& SelectedBones, const float VolumeMin, const float VolumeMax)
+{
+	if (GeometryCollection.HasGroup(FGeometryCollection::TransformGroup) &&
+		GeometryCollection.HasAttribute("Volume", FGeometryCollection::TransformGroup))
+	{
+		SelectedBones.Empty();
+
+		FGeometryCollectionConvexUtility::SetVolumeAttributes(&GeometryCollection);
+
+		const TManagedArray<float>& Volumes = GeometryCollection.GetAttribute<float>("Volume", FTransformCollection::TransformGroup);
+
+		for (int32 BoneIdx = 0; BoneIdx < Volumes.Num(); ++BoneIdx)
+		{
+			if (Volumes[BoneIdx] >= VolumeMin && Volumes[BoneIdx] <= VolumeMax)
+			{
+				SelectedBones.Add(BoneIdx);
+			}
+		}
+	}
+}
+
+void FFractureEngineSelection::SelectByVolume(FGeometryCollection& GeometryCollection, FDataflowTransformSelection& TransformSelection, const float VolumeMin, const float VolumeMax)
+{
+	TArray<int32> SelectionArr;
+
+	SelectByVolume(GeometryCollection, SelectionArr, VolumeMin, VolumeMax);
+
+	TransformSelection.SetFromArray(SelectionArr);
+}
+
diff --git a/Engine/Plugins/Experimental/Fracture/Source/FractureEngine/Public/FractureEngineClustering.h b/Engine/Plugins/Experimental/Fracture/Source/FractureEngine/Public/FractureEngineClustering.h
--- a/Engine/Plugins/Experimental/Fracture/Source/FractureEngine/Public/FractureEngineClustering.h
+++ b/Engine/Plugins/Experimental/Fracture/Source/FractureEngine/Public/FractureEngineClustering.h
@@ -1,2 +1,90 @@
 // Copyright Epic Games, Inc. All Rights Reserved.
 
+#pragma once
+
+#include "Algo/Count.h"
+#include "GeometryCollection/GeometryCollection.h"
+
+
+class FRACTUREENGINE_API FVoronoiPartitioner
+{
+public:
+	FVoronoiPartitioner(const FGeometryCollection* GeometryCollection, int32 ClusterIndex);
+
+	/** Cluster bodies into k partitions using K-Means. Connectivity is ignored: only spatial proximity is considered. */
+	void KMeansPartition(int32 InPartitionCount);
+
+	/** Split any partition islands into their own partition. This will possbily increase number of partitions to exceed desired count. */
+	void SplitDisconnectedPartitions(FGeometryCollection* GeometryCollection);
+
+	/** Merge any partitions w/ only 1 body into a connected, neighboring partition (if any).  This can decrease the number of partitions below the desired count. */
+	void MergeSingleElementPartitions(FGeometryCollection* GeometryCollection);
+
+	int32 GetPartitionCount() const { return PartitionCount; }
+
+	int32 GetNonEmptyPartitionCount() const
+	{
+		return PartitionSize.Num() - Algo::Count(PartitionSize, 0);
+	}
+
+	/** return the GeometryCollection TranformIndices within the partition. */
+	TArray<int32> GetPartition(int32 PartitionIndex) const;
+
+private:
+	void GenerateConnectivity(const FGeometryCollection* GeometryCollection);
+	void CollectConnections(const FGeometryCollection* GeometryCollection, int32 Index, int32 OperatingLevel, TSet<int32>& OutConnections) const;
+	void GenerateCentroids(const FGeometryCollection* GeometryCollection);
+	FVector GenerateCentroid(const FGeometryCollection* GeometryCollection, int32 TransformIndex) const;
+	FBox GenerateBounds(const FGeometryCollection* GeometryCollection, int32 TransformIndex) const;
+	void InitializePartitions();
+	bool Refine();
+	int32 FindClosestPartitionCenter(const FVector& Location) const;
+	void MarkVisited(int32 Index, int32 PartitionIndex);
+
+private:
+	TArray<int32> TransformIndices;
+	TArray<FVector> Centroids;
+	TArray<int32> Partitions;
+	int32 PartitionCount;
+	TArray<int32> PartitionSize;
+	TArray<FVector> PartitionCenters;
+	TArray<TSet<int32>> Connectivity;
+	TArray<bool> Visited;
+
+
+	// Not generally necessary but this is a safety measure to prevent oscillating solves that never converge.
+	const int32 MaxKMeansIterations = 500;
+};
+
+enum class FRACTUREENGINE_API EFractureEngineClusterSizeMethod : uint8
+{
+	// Cluster by specifying an absolute number of clusters
+	ByNumber,
+	// Cluster by specifying a fraction of the number of input bones
+	ByFractionOfInput,
+	// Cluster by specifying the density of the input bones
+	BySize,
+};
+
+class FRACTUREENGINE_API FFractureEngineClustering
+{
+public:
+
+	static void AutoCluster(FGeometryCollection& GeometryCollection,
+		const TArray<int32>& BoneIndices,
+		const EFractureEngineClusterSizeMethod ClusterSizeMethod,
+		const uint32 SiteCount,
+		const float SiteCountFraction,
+		const float SiteSize,
+		const bool bEnforceConnectivity,
+		const bool bAvoidIsolated);
+
+	static void AutoCluster(FGeometryCollection& GeometryCollection,
+		const int32 ClusterIndex,
+		const EFractureEngineClusterSizeMethod ClusterSizeMethod,
+		const uint32 SiteCount,
+		const float SiteCountFraction,
+		const float SiteSize,
+		const bool bEnforceConnectivity,
+		const bool bAvoidIsolated);
+};
\ No newline at end of file
diff --git a/Engine/Plugins/Experimental/Fracture/Source/FractureEngine/Public/FractureEnginePlugin.h b/Engine/Plugins/Experimental/Fracture/Source/FractureEngine/Public/FractureEnginePlugin.h
new file mode 100644
--- /dev/null
+++ b/Engine/Plugins/Experimental/Fracture/Source/FractureEngine/Public/FractureEnginePlugin.h
@@ -0,0 +1,39 @@
+// Copyright Epic Games, Inc. All Rights Reserved.
+
+#pragma once
+
+#include "CoreMinimal.h"
+#include "Modules/ModuleInterface.h"
+#include "Modules/ModuleManager.h"
+
+
+/**
+ * The public interface to this module
+ */
+class IFractureEnginePlugin : public IModuleInterface
+{
+
+public:
+
+	/**
+	 * Singleton-like access to this module's interface.  This is just for convenience!
+	 * Beware of calling this during the shutdown phase, though.  Your module might have been unloaded already.
+	 *
+	 * @return Returns singleton instance, loading the module on demand if needed
+	 */
+	static inline IFractureEnginePlugin& Get()
+	{
+		return FModuleManager::LoadModuleChecked< IFractureEnginePlugin >( "FractureEnginePlugin" );
+	}
+
+	/**
+	 * Checks to see if this module is loaded and ready.  It is only valid to call Get() if IsAvailable() returns true.
+	 *
+	 * @return True if the module is loaded and ready to use
+	 */
+	static inline bool IsAvailable()
+	{
+		return FModuleManager::Get().IsModuleLoaded( "FractureEnginePlugin" );
+	}
+};
+
diff --git a/Engine/Plugins/Experimental/Fracture/Source/FractureEngine/Public/FractureEngineSelection.h b/Engine/Plugins/Experimental/Fracture/Source/FractureEngine/Public/FractureEngineSelection.h
--- a/Engine/Plugins/Experimental/Fracture/Source/FractureEngine/Public/FractureEngineSelection.h
+++ b/Engine/Plugins/Experimental/Fracture/Source/FractureEngine/Public/FractureEngineSelection.h
@@ -1,2 +1,40 @@
 // Copyright Epic Games, Inc. All Rights Reserved.
 
+#pragma once
+
+#include "GeometryCollection/GeometryCollection.h"
+#include "Dataflow/DataflowSelection.h"
+
+class FRACTUREENGINE_API FFractureEngineSelection
+{
+public:
+	static void SelectParent(const FGeometryCollection& GeometryCollection, TArray<int32>& SelectedBones);
+	static void SelectParent(const FGeometryCollection& GeometryCollection, FDataflowTransformSelection& TransformSelection);
+
+	static void SelectChildren(const FGeometryCollection& GeometryCollection, TArray<int32>& SelectedBones);
+	static void SelectChildren(const FGeometryCollection& GeometryCollection, FDataflowTransformSelection& TransformSelection);
+
+	static void SelectSiblings(const FGeometryCollection& GeometryCollection, TArray<int32>& SelectedBones);
+	static void SelectSiblings(const FGeometryCollection& GeometryCollection, FDataflowTransformSelection& TransformSelection);
+
+	static void SelectLevel(const FGeometryCollection& GeometryCollection, TArray<int32>& SelectedBones);
+	static void SelectLevel(const FGeometryCollection& GeometryCollection, FDataflowTransformSelection& TransformSelection);
+
+	static void SelectContact(FGeometryCollection& GeometryCollection, TArray<int32>& SelectedBones);
+	static void SelectContact(FGeometryCollection& GeometryCollection, FDataflowTransformSelection& TransformSelection);
+
+	static void SelectLeaf(const FGeometryCollection& GeometryCollection, TArray<int32>& SelectedBones);
+	static void SelectLeaf(const FGeometryCollection& GeometryCollection, FDataflowTransformSelection& TransformSelection);
+
+	static void SelectCluster(const FGeometryCollection& GeometryCollection, TArray<int32>& SelectedBones);
+	static void SelectCluster(const FGeometryCollection& GeometryCollection, FDataflowTransformSelection& TransformSelection);
+
+	static void SelectByPercentage(const FGeometryCollection& GeometryCollection, TArray<int32>& SelectedBones, const int32 Percentage, const bool Deterministic, const float RandomSeed);
+	static void SelectByPercentage(const FGeometryCollection& GeometryCollection, FDataflowTransformSelection& TransformSelection, const int32 Percentage, const bool Deterministic, const float RandomSeed);
+
+	static void SelectBySize(FGeometryCollection& GeometryCollection, TArray<int32>& SelectedBones, const float SizeMin, const float SizeMax);
+	static void SelectBySize(FGeometryCollection& GeometryCollection, FDataflowTransformSelection& TransformSelection, const float SizeMin, const float SizeMax);
+
+	static void SelectByVolume(FGeometryCollection& GeometryCollection, TArray<int32>& SelectedBones, const float VolumeMin, const float VolumeMax);
+	static void SelectByVolume(FGeometryCollection& GeometryCollection, FDataflowTransformSelection& TransformSelection, const float VolumeMin, const float VolumeMax);
+};
\ No newline at end of file
diff --git a/Engine/Plugins/Experimental/GeometryCollectionPlugin/GeometryCollectionPlugin.uplugin b/Engine/Plugins/Experimental/GeometryCollectionPlugin/GeometryCollectionPlugin.uplugin
--- a/Engine/Plugins/Experimental/GeometryCollectionPlugin/GeometryCollectionPlugin.uplugin
+++ b/Engine/Plugins/Experimental/GeometryCollectionPlugin/GeometryCollectionPlugin.uplugin
@@ -52,6 +52,10 @@
 		{
 			"Name": "PlanarCut",
 			"Enabled": true
+		},
+		{
+			"Name": "Fracture",
+			"Enabled": true
 		}
 	]
 }
diff --git a/Engine/Plugins/Experimental/GeometryCollectionPlugin/Source/GeometryCollectionNodes/GeometryCollectionNodes.Build.cs b/Engine/Plugins/Experimental/GeometryCollectionPlugin/Source/GeometryCollectionNodes/GeometryCollectionNodes.Build.cs
--- a/Engine/Plugins/Experimental/GeometryCollectionPlugin/Source/GeometryCollectionNodes/GeometryCollectionNodes.Build.cs
+++ b/Engine/Plugins/Experimental/GeometryCollectionPlugin/Source/GeometryCollectionNodes/GeometryCollectionNodes.Build.cs
@@ -18,10 +18,22 @@ namespace UnrealBuildTool.Rules
 					"GeometryCollectionEngine",
 					"GeometryCore",
 					"Voronoi",
+					"GeometryFramework",
+					"MeshDescription",
+					"StaticMeshDescription",
 					"PlanarCut",
 					"Engine"
 				}
 			);
+
+			PrivateDependencyModuleNames.AddRange(
+				new string[] 
+				{
+					"DynamicMesh",
+					"MeshConversion",
+					"FractureEngine"
+				}
+			);
 		}
 	}
 }
diff --git a/Engine/Plugins/Experimental/GeometryCollectionPlugin/Source/GeometryCollectionNodes/Private/Dataflow/GeometryCollectionNodes.cpp b/Engine/Plugins/Experimental/GeometryCollectionPlugin/Source/GeometryCollectionNodes/Private/Dataflow/GeometryCollectionNodes.cpp
--- a/Engine/Plugins/Experimental/GeometryCollectionPlugin/Source/GeometryCollectionNodes/Private/Dataflow/GeometryCollectionNodes.cpp
+++ b/Engine/Plugins/Experimental/GeometryCollectionPlugin/Source/GeometryCollectionNodes/Private/Dataflow/GeometryCollectionNodes.cpp
@@ -14,32 +14,38 @@
 #include "Logging/LogMacros.h"
 #include "Templates/SharedPointer.h"
 #include "UObject/UnrealTypePrivate.h"
+#include "DynamicMeshToMeshDescription.h"
+#include "MeshDescriptionToDynamicMesh.h"
+#include "StaticMeshAttributes.h"
+#include "DynamicMeshEditor.h"
+#include "Operations/MeshBoolean.h"
 
 #include "EngineGlobals.h"
 #include "GeometryCollection/GeometryCollectionAlgo.h"
 #include "GeometryCollection/GeometryCollectionClusteringUtility.h"
 #include "GeometryCollection/GeometryCollectionConvexUtility.h"
-#include "Math/UnrealMathUtility.h"
-#include "PlanarCut.h"
 #include "Voronoi/Voronoi.h"
+#include "PlanarCut.h"
+#include "GeometryCollection/GeometryCollectionProximityUtility.h"
+#include "FractureEngineClustering.h"
+#include "FractureEngineSelection.h"
+#include "GeometryCollection/Facades/CollectionBoundsFacade.h"
+#include "GeometryCollection/Facades/CollectionAnchoringFacade.h"
 
 #include UE_INLINE_GENERATED_CPP_BY_NAME(GeometryCollectionNodes)
 
 namespace Dataflow
 {
-	static const FLinearColor CDefaultNodeBodyTintColor = FLinearColor(0.f, 0.f, 0.f, 0.5f);
-
 	void GeometryCollectionEngineNodes()
 	{
+		static const FLinearColor CDefaultNodeBodyTintColor = FLinearColor(0.f, 0.f, 0.f, 0.5f);
+
 		DATAFLOW_NODE_REGISTER_CREATION_FACTORY(FGetCollectionAssetDataflowNode);
 		DATAFLOW_NODE_REGISTER_CREATION_FACTORY(FExampleCollectionEditDataflowNode);
 		DATAFLOW_NODE_REGISTER_CREATION_FACTORY(FSetCollectionAssetDataflowNode);
 		DATAFLOW_NODE_REGISTER_CREATION_FACTORY(FAppendCollectionAssetsDataflowNode);
 		DATAFLOW_NODE_REGISTER_CREATION_FACTORY(FResetGeometryCollectionDataflowNode);
 
-		DATAFLOW_NODE_REGISTER_CREATION_FACTORY(FClusterFlattenDataflowNode);
-		DATAFLOW_NODE_REGISTER_CREATION_FACTORY(FRemoveOnBreakDataflowNode);
-
 		DATAFLOW_NODE_REGISTER_CREATION_FACTORY(FPrintStringDataflowNode);
 		DATAFLOW_NODE_REGISTER_CREATION_FACTORY(FLogStringDataflowNode);
 		DATAFLOW_NODE_REGISTER_CREATION_FACTORY(FMakeLiteralStringDataflowNode);
@@ -74,6 +80,28 @@
 		DATAFLOW_NODE_REGISTER_CREATION_FACTORY(FHashVectorDataflowNode);
 		DATAFLOW_NODE_REGISTER_CREATION_FACTORY(FFloatToIntDataflowNode);
 		DATAFLOW_NODE_REGISTER_CREATION_FACTORY(FMathConstantsDataflowNode);
+		DATAFLOW_NODE_REGISTER_CREATION_FACTORY(FGetArrayElementDataflowNode);
+		DATAFLOW_NODE_REGISTER_CREATION_FACTORY(FGetNumArrayElementsDataflowNode);
+		DATAFLOW_NODE_REGISTER_CREATION_FACTORY(FGetBoundingBoxesDataflowNode);
+		DATAFLOW_NODE_REGISTER_CREATION_FACTORY(FGetCentroidsDataflowNode);
+		DATAFLOW_NODE_REGISTER_CREATION_FACTORY(FPointsToMeshDataflowNode);
+		DATAFLOW_NODE_REGISTER_CREATION_FACTORY(FBoxToMeshDataflowNode);
+		DATAFLOW_NODE_REGISTER_CREATION_FACTORY(FMeshInfoDataflowNode);
+		DATAFLOW_NODE_REGISTER_CREATION_FACTORY(FMeshToCollectionDataflowNode);
+		DATAFLOW_NODE_REGISTER_CREATION_FACTORY(FStaticMeshToMeshDataflowNode);
+		DATAFLOW_NODE_REGISTER_CREATION_FACTORY(FTransformDataflowNode);
+		DATAFLOW_NODE_REGISTER_CREATION_FACTORY(FMeshAppendDataflowNode);
+		DATAFLOW_NODE_REGISTER_CREATION_FACTORY(FMeshBooleanDataflowNode);
+		DATAFLOW_NODE_REGISTER_CREATION_FACTORY(FMeshCopyToPointsDataflowNode);
+		DATAFLOW_NODE_REGISTER_CREATION_FACTORY(FCompareIntDataflowNode);
+		DATAFLOW_NODE_REGISTER_CREATION_FACTORY(FBranchDataflowNode);
+		DATAFLOW_NODE_REGISTER_CREATION_FACTORY(FGetMeshDataDataflowNode);
+		DATAFLOW_NODE_REGISTER_CREATION_FACTORY(FGetSchemaDataflowNode);
+		DATAFLOW_NODE_REGISTER_CREATION_FACTORY(FAutoClusterDataflowNode);
+		DATAFLOW_NODE_REGISTER_CREATION_FACTORY(FClusterFlattenDataflowNode);
+		DATAFLOW_NODE_REGISTER_CREATION_FACTORY(FRemoveOnBreakDataflowNode);
+		DATAFLOW_NODE_REGISTER_CREATION_FACTORY(FSetAnchorStateDataflowNode);
+		DATAFLOW_NODE_REGISTER_CREATION_FACTORY(FProximityDataflowNode);
 
 		// GeometryCollection
 		DATAFLOW_NODE_REGISTER_CREATION_FACTORY_NODE_COLORS_BY_CATEGORY("GeometryCollection", FLinearColor(0.55f, 0.45f, 1.0f), CDefaultNodeBodyTintColor);
@@ -89,7 +117,8 @@
 		DATAFLOW_NODE_REGISTER_CREATION_FACTORY_NODE_COLORS_BY_CATEGORY("Math", FLinearColor(0.f, 0.4f, 0.8f), CDefaultNodeBodyTintColor);
 		// Generators
 		DATAFLOW_NODE_REGISTER_CREATION_FACTORY_NODE_COLORS_BY_CATEGORY("Generators", FLinearColor(.7f, 0.7f, 0.7f), CDefaultNodeBodyTintColor);
-
+		// Mesh
+		DATAFLOW_NODE_REGISTER_CREATION_FACTORY_NODE_COLORS_BY_CATEGORY("Mesh", FLinearColor(1.f, 0.16f, 0.05f), CDefaultNodeBodyTintColor);
 	}
 }
 
@@ -239,7 +268,7 @@
 	if (PrintToLog)
 	{
 		FString Value = GetValue<FString>(Context, &String);
-		UE_LOG(LogTemp, Warning, TEXT("Text, %s"), *Value);
+		UE_LOG(LogTemp, Warning, TEXT("[Dataflow Log] %s"), *Value);
 	}
 }
 
@@ -253,15 +282,15 @@
 
 void ComputeBoundingBox(FBoundingBoxDataflowNode::DataType& Collection, FBox& BoundingBox)
 {
-	if (Collection.HasAttribute("Transform", "Transform") &&
-		Collection.HasAttribute("Parent", "Transform") &&
-		Collection.HasAttribute("TransformIndex", "Geometry") &&
-		Collection.HasAttribute("BoundingBox", "Geometry"))
+	if (Collection.HasAttribute("Transform", FGeometryCollection::TransformGroup) &&
+		Collection.HasAttribute("Parent", FGeometryCollection::TransformGroup) &&
+		Collection.HasAttribute("TransformIndex", FGeometryCollection::GeometryGroup) &&
+		Collection.HasAttribute("BoundingBox", FGeometryCollection::GeometryGroup))
 	{
-		const TManagedArray<FTransform>& Transforms = Collection.GetAttribute<FTransform>("Transform", "Transform");
-		const TManagedArray<int32>& ParentIndices = Collection.GetAttribute<int32>("Parent", "Transform");
-		const TManagedArray<int32>& TransformIndices = Collection.GetAttribute<int32>("TransformIndex", "Geometry");
-		const TManagedArray<FBox>& BoundingBoxes = Collection.GetAttribute<FBox>("BoundingBox", "Geometry");
+		const TManagedArray<FTransform>& Transforms = Collection.GetAttribute<FTransform>("Transform", FGeometryCollection::TransformGroup);
+		const TManagedArray<int32>& ParentIndices = Collection.GetAttribute<int32>("Parent", FGeometryCollection::TransformGroup);
+		const TManagedArray<int32>& TransformIndices = Collection.GetAttribute<int32>("TransformIndex", FGeometryCollection::GeometryGroup);
+		const TManagedArray<FBox>& BoundingBoxes = Collection.GetAttribute<FBox>("BoundingBox", FGeometryCollection::GeometryGroup);
 
 		TArray<FMatrix> TmpGlobalMatrices;
 		GeometryCollectionAlgo::GlobalMatrices(Transforms, ParentIndices, TmpGlobalMatrices);
@@ -378,7 +407,7 @@
 
 			TArray<FVector> PointsArr;
 			PointsArr.Reserve(NumPoints);
-			for (int32 idx = 0; idx < NumPoints; ++idx)
+			for (int32 Idx = 0; Idx < NumPoints; ++Idx)
 			{
 				PointsArr.Emplace(BBox.Min + FVector(RandStream.FRand(), RandStream.FRand(), RandStream.FRand()) * Extent);
 			}
@@ -506,74 +535,6 @@
 	return MaxDisp;
 }
 
-void FClusterFlattenDataflowNode::Evaluate(Dataflow::FContext& Context, const FDataflowOutput* Out) const
-{
-	if (Out->IsA<FManagedArrayCollection>(&Collection))
-	{
-		FManagedArrayCollection CollectionPtr = GetValue<FManagedArrayCollection>(Context, &Collection);
-		if (FGeometryCollection* GeomCollection = CollectionPtr.NewCopy<FGeometryCollection>())
-		{
-			FGeometryCollectionClusteringUtility::UpdateHierarchyLevelOfChildren(GeomCollection, -1);
-			
-			const TManagedArray<int32>& Levels = GeomCollection->GetAttribute<int32>("Level", FGeometryCollection::TransformGroup);
-
-			// Populate Selected Bones in an Array
-			// @todo(harsha) Implement with Selection
-			// For every bone in selected array: [ClusterIndex]
-			int32 ClusterIndex = 0;
-			TArray<int32> LeafBones;
-			FGeometryCollectionClusteringUtility::GetLeafBones(GeomCollection, ClusterIndex, true, LeafBones);
-			FGeometryCollectionClusteringUtility::ClusterBonesUnderExistingNode(GeomCollection, ClusterIndex, LeafBones);
-			FGeometryCollectionClusteringUtility::RemoveDanglingClusters(GeomCollection);
-			// End for
-
-			FGeometryCollectionClusteringUtility::UpdateHierarchyLevelOfChildren(GeomCollection, -1);
-			SetValue<FManagedArrayCollection>(Context, (const FManagedArrayCollection&)(*GeomCollection), &Collection);
-		}
-	}
-}
-
-void FRemoveOnBreakDataflowNode::Evaluate(Dataflow::FContext& Context, const FDataflowOutput* Out) const
-{
-	if (Out->IsA<FManagedArrayCollection>(&Collection))
-	{
-		FManagedArrayCollection CollectionPtr = GetValue<FManagedArrayCollection>(Context, &Collection);
-		if (FGeometryCollection* GeometryCollection = CollectionPtr.NewCopy<FGeometryCollection>())
-		{
-			if (!GeometryCollection->HasAttribute("RemoveOnBreak", FGeometryCollection::TransformGroup))
-	 		{
-	 			TManagedArray<FVector4f>& NewRemoveOnBreak = GeometryCollection->AddAttribute<FVector4f>("RemoveOnBreak", FGeometryCollection::TransformGroup);
-	 			NewRemoveOnBreak.Fill(FRemoveOnBreakData::DisabledPackedData);
-	 		}
-
-			TManagedArray<FVector4f>& RemoveOnBreak = GeometryCollection->ModifyAttribute<FVector4f>("RemoveOnBreak", FGeometryCollection::TransformGroup);
-
-			const FVector2f PostBreakTimerData = GetValue<FVector2f>(Context, &PostBreakTimer);
-			const FVector2f RemovalTimerData = GetValue<FVector2f>(Context, &RemovalTimer);
-			const bool ClusterCrumblingData = GetValue<bool>(Context, &ClusterCrumbling);
-
-			RemoveOnBreak.Fill(FVector4f{PostBreakTimerData.X, PostBreakTimerData.Y, RemovalTimerData.X, RemovalTimerData.Y});
-
-			// @todo(harsha) Implement with Selection
-	 		// const FRemoveOnBreakData RemoveOnBreakData(true, PostBreakTimerData, ClusterCrumblingData, RemovalTimerData);
-	 		// for (int32 Index : SelectedBones)
-	 		// {
-				// // if root bone, then do not set 
-				// if (GeometryCollection->Parent[Index] == INDEX_NONE)
-				// {
-				// 	RemoveOnBreak[Index] = FRemoveOnBreakData::DisabledPackedData;
-				// }
-				// else
-				// {
-				// 	RemoveOnBreak[Index] = RemoveOnBreakData.GetPackedData();
-				// }
-	 		// }
-			
-	 		SetValue<FManagedArrayCollection>(Context, (const FManagedArrayCollection&)(*GeometryCollection), &Collection);
-		}
-	}
-}
-
 void FVoronoiFractureDataflowNode::Evaluate(Dataflow::FContext& Context, const FDataflowOutput* Out) const
 {
 	if (Out->IsA<FManagedArrayCollection>(&Collection))
@@ -589,15 +550,15 @@
 				//
 				FBox BoundingBox(ForceInit);
 
-				if (GeomCollection->HasAttribute("Transform", "Transform") &&
-					GeomCollection->HasAttribute("Parent", "Transform") &&
-					GeomCollection->HasAttribute("TransformIndex", "Geometry") &&
-					GeomCollection->HasAttribute("BoundingBox", "Geometry"))
+				if (GeomCollection->HasAttribute("Transform", FGeometryCollection::TransformGroup) &&
+					GeomCollection->HasAttribute("Parent", FGeometryCollection::TransformGroup) &&
+					GeomCollection->HasAttribute("TransformIndex", FGeometryCollection::GeometryGroup) &&
+					GeomCollection->HasAttribute("BoundingBox", FGeometryCollection::GeometryGroup))
 				{
-					const TManagedArray<FTransform>& Transforms = GeomCollection->GetAttribute<FTransform>("Transform", "Transform");
-					const TManagedArray<int32>& ParentIndices = GeomCollection->GetAttribute<int32>("Parent", "Transform");
-					const TManagedArray<int32>& TransformIndices = GeomCollection->GetAttribute<int32>("TransformIndex", "Geometry");
-					const TManagedArray<FBox>& BoundingBoxes = GeomCollection->GetAttribute<FBox>("BoundingBox", "Geometry");
+					const TManagedArray<FTransform>& Transforms = GeomCollection->GetAttribute<FTransform>("Transform", FGeometryCollection::TransformGroup);
+					const TManagedArray<int32>& ParentIndices = GeomCollection->GetAttribute<int32>("Parent", FGeometryCollection::TransformGroup);
+					const TManagedArray<int32>& TransformIndices = GeomCollection->GetAttribute<int32>("TransformIndex", FGeometryCollection::GeometryGroup);
+					const TManagedArray<FBox>& BoundingBoxes = GeomCollection->GetAttribute<FBox>("BoundingBox", FGeometryCollection::GeometryGroup);
 
 					TArray<FMatrix> TmpGlobalMatrices;
 					GeometryCollectionAlgo::GlobalMatrices(Transforms, ParentIndices, TmpGlobalMatrices);
@@ -647,7 +608,6 @@
 
 					int ResultGeometryIndex = CutMultipleWithPlanarCells(VoronoiPlanarCells, *GeomCollection, TransformIndicesArray, GroutVal, CollisionSampleSpacingVal, RandomSeedVal, FTransform().Identity);
 
-					//					Out->SetValue<FManagedArrayCollection>(FManagedArrayCollection(*GeomCollection), Context);
 					SetValue<FManagedArrayCollection>(Context, (const FManagedArrayCollection&)(*GeomCollection), &Collection);
 				}
 			}
@@ -758,7 +718,7 @@
 	}
 }
 
-void AddAdditionalAttributesIfRequired(FGeometryCollection* GeometryCollection)
+static void AddAdditionalAttributesIfRequired(FGeometryCollection* GeometryCollection)
 {
 	FGeometryCollectionClusteringUtility::UpdateHierarchyLevelOfChildren(GeometryCollection, -1);
 }
@@ -1081,4 +1041,613 @@
 	}
 }
 
+void FGetArrayElementDataflowNode::Evaluate(Dataflow::FContext& Context, const FDataflowOutput* Out) const
+{
+	if (Out->IsA<FVector>(&Point))
+	{
+		const TArray<FVector>& Array = GetValue<TArray<FVector>>(Context, &Points);
+		if (Index >= 0 && Index < Array.Num())
+		{
+			SetValue<FVector>(Context, Array[Index], &Point);
+		}
+	}
+}
 
+void FGetNumArrayElementsDataflowNode::Evaluate(Dataflow::FContext& Context, const FDataflowOutput* Out) const
+{
+	if (Out->IsA<int32>(&NumElements))
+	{
+		SetValue<int32>(Context, GetValue<TArray<FVector>>(Context, &Points).Num(), &NumElements);
+	}
+}
+
+void FGetBoundingBoxesDataflowNode::Evaluate(Dataflow::FContext& Context, const FDataflowOutput* Out) const
+{
+	if (Out->IsA<TArray<FBox>>(&BoundingBoxes))
+	{
+		FManagedArrayCollection InCollection = GetValue<FManagedArrayCollection>(Context, &Collection);
+		FDataflowTransformSelection InTransformSelection = GetValue<FDataflowTransformSelection>(Context, &TransformSelection);
+
+		const TManagedArray<FBox>* BoundingBoxesArr = GeometryCollection::Facades::FBoundsFacade::GetBoundingBoxes(InCollection);
+
+		SetValue<TArray<FBox>>(Context, BoundingBoxesArr->GetConstArray(), &BoundingBoxes);
+	}
+}
+
+void FGetCentroidsDataflowNode::Evaluate(Dataflow::FContext& Context, const FDataflowOutput* Out) const
+{
+	if (Out->IsA<TArray<FVector>>(&Centroids))
+	{
+		FManagedArrayCollection InCollection = GetValue<FManagedArrayCollection>(Context, &Collection);
+		FDataflowTransformSelection InTransformSelection = GetValue<FDataflowTransformSelection>(Context, &TransformSelection);
+
+		const TManagedArray<FBox>* BoundingBoxes = GeometryCollection::Facades::FBoundsFacade::GetBoundingBoxes(InCollection);
+
+		TArray<FVector> CentroidsArr;
+		for (int32 Idx = 0; Idx < BoundingBoxes->Num(); ++Idx)
+		{
+			const FBox& BoundingBox = (*BoundingBoxes)[Idx];
+			if (BoundingBox.IsValid)
+			{
+				if (InTransformSelection.Num() > 0)
+				{
+					if (InTransformSelection.IsSelected(Idx))
+					{
+						CentroidsArr.Add(BoundingBox.GetCenter());
+					}
+				}
+				else
+				{
+					CentroidsArr.Add(BoundingBox.GetCenter());
+				}
+			}
+		}
+		SetValue<TArray<FVector>>(Context, CentroidsArr, &Centroids);
+	}
+}
+
+void FPointsToMeshDataflowNode::Evaluate(Dataflow::FContext& Context, const FDataflowOutput* Out) const
+{
+	if (Out->IsA<TObjectPtr<UDynamicMesh>>(&Mesh) || Out->IsA<int32>(&TriangleCount))
+	{
+		TArray<FVector> PointsArr = GetValue<TArray<FVector>>(Context, &Points);
+		if (PointsArr.Num() > 0)
+		{
+			TObjectPtr<UDynamicMesh> DynamicMesh = NewObject<UDynamicMesh>();
+			DynamicMesh->Reset();
+
+			UE::Geometry::FDynamicMesh3& DynMesh = DynamicMesh->GetMeshRef();
+
+			for (auto& Point : PointsArr)
+			{
+				DynMesh.AppendVertex(Point);
+			}
+
+			SetValue<TObjectPtr<UDynamicMesh>>(Context, DynamicMesh, &Mesh);
+			SetValue<int32>(Context, DynamicMesh->GetTriangleCount(), &TriangleCount);
+		}
+		else
+		{
+			SetValue<TObjectPtr<UDynamicMesh>>(Context, NewObject<UDynamicMesh>(), &Mesh);
+			SetValue<int32>(Context, 0, &TriangleCount);
+		}
+	}
+}
+
+void FBoxToMeshDataflowNode::Evaluate(Dataflow::FContext& Context, const FDataflowOutput* Out) const
+{
+	if (Out->IsA<TObjectPtr<UDynamicMesh>>(&Mesh) || Out->IsA<int32>(&TriangleCount))
+	{
+		TObjectPtr<UDynamicMesh> DynamicMesh = NewObject<UDynamicMesh>();
+		DynamicMesh->Reset();
+
+		UE::Geometry::FDynamicMesh3& DynMesh = DynamicMesh->GetMeshRef();
+
+		FBox InBox = GetValue<FBox>(Context, &Box);
+		FVector Min = InBox.Min;
+		FVector Max = InBox.Max;
+
+		// Add vertices
+		int Vertex0 = DynMesh.AppendVertex(Min);
+		int Vertex1 = DynMesh.AppendVertex(FVector(Min.X, Max.Y, Min.Z));
+		int Vertex2 = DynMesh.AppendVertex(FVector(Max.X, Max.Y, Min.Z));
+		int Vertex3 = DynMesh.AppendVertex(FVector(Max.X, Min.Y, Min.Z));
+		int Vertex4 = DynMesh.AppendVertex(FVector(Min.X, Min.Y, Max.Z));
+		int Vertex5 = DynMesh.AppendVertex(FVector(Min.X, Max.Y, Max.Z));
+		int Vertex6 = DynMesh.AppendVertex(Max);
+		int Vertex7 = DynMesh.AppendVertex(FVector(Max.X, Min.Y, Max.Z));
+
+		// Add triangles
+		int GroupID = 0;
+		DynMesh.AppendTriangle(Vertex0, Vertex1, Vertex3, GroupID);
+		DynMesh.AppendTriangle(Vertex1, Vertex2, Vertex3, GroupID);
+		DynMesh.AppendTriangle(Vertex3, Vertex6, Vertex7, GroupID);
+		DynMesh.AppendTriangle(Vertex3, Vertex2, Vertex6, GroupID);
+		DynMesh.AppendTriangle(Vertex7, Vertex4, Vertex0, GroupID);
+		DynMesh.AppendTriangle(Vertex0, Vertex3, Vertex7, GroupID);
+		DynMesh.AppendTriangle(Vertex0, Vertex4, Vertex5, GroupID);
+		DynMesh.AppendTriangle(Vertex0, Vertex5, Vertex1, GroupID);
+		DynMesh.AppendTriangle(Vertex1, Vertex5, Vertex6, GroupID);
+		DynMesh.AppendTriangle(Vertex6, Vertex2, Vertex1, GroupID);
+		DynMesh.AppendTriangle(Vertex4, Vertex6, Vertex5, GroupID);
+		DynMesh.AppendTriangle(Vertex4, Vertex7, Vertex6, GroupID);
+
+		SetValue<TObjectPtr<UDynamicMesh>>(Context, DynamicMesh, &Mesh);
+		SetValue<int32>(Context, DynamicMesh->GetTriangleCount(), &TriangleCount);
+	}
+}
+
+void FMeshInfoDataflowNode::Evaluate(Dataflow::FContext& Context, const FDataflowOutput* Out) const
+{
+	if (Out->IsA<FString>(&InfoString))
+	{
+		UE::Geometry::FDynamicMesh3& DynMesh = (GetValue<TObjectPtr<UDynamicMesh>>(Context, &Mesh))->GetMeshRef();
+
+		SetValue<FString>(Context, DynMesh.MeshInfoString(), &InfoString);
+	}
+}
+
+void FMeshToCollectionDataflowNode::Evaluate(Dataflow::FContext& Context, const FDataflowOutput* Out) const
+{
+	if (Out->IsA<FManagedArrayCollection>(&Collection))
+	{
+		UE::Geometry::FDynamicMesh3& DynMesh = (GetValue<TObjectPtr<UDynamicMesh>>(Context, &Mesh))->GetMeshRef();
+		if (DynMesh.VertexCount() > 0)
+		{
+			FMeshDescription MeshDescription;
+			FStaticMeshAttributes Attributes(MeshDescription);
+			Attributes.Register();
+
+			FDynamicMeshToMeshDescription Converter;
+			Converter.Convert(&DynMesh, MeshDescription, true);
+
+			FGeometryCollection NewGeometryCollection = FGeometryCollection();
+			FGeometryCollectionEngineConversion::AppendMeshDescription(&MeshDescription, FString(TEXT("TEST")), 0, FTransform().Identity, &NewGeometryCollection);
+
+			FManagedArrayCollection Col = FManagedArrayCollection();
+			NewGeometryCollection.CopyTo(&Col);
+
+			SetValue<FManagedArrayCollection>(Context, Col, &Collection);
+		}
+		else
+		{
+			SetValue<FManagedArrayCollection>(Context, FManagedArrayCollection(), &Collection);
+		}
+	}
+}
+
+void FStaticMeshToMeshDataflowNode::Evaluate(Dataflow::FContext& Context, const FDataflowOutput* Out) const
+{
+	if (Out->IsA<TObjectPtr<UDynamicMesh>>(&Mesh))
+	{
+		if (FMeshDescription* MeshDescription = UseHiRes ? StaticMesh->GetHiResMeshDescription() : StaticMesh->GetMeshDescription(LODLevel))
+		{
+			TObjectPtr<UDynamicMesh> DynamicMesh = NewObject<UDynamicMesh>();
+			DynamicMesh->Reset();
+
+			UE::Geometry::FDynamicMesh3& DynMesh = DynamicMesh->GetMeshRef();
+			{
+				FMeshDescriptionToDynamicMesh ConverterToDynamicMesh;
+				ConverterToDynamicMesh.Convert(MeshDescription, DynMesh);
+			}
+
+			SetValue<TObjectPtr<UDynamicMesh>>(Context, DynamicMesh, &Mesh);
+		}
+		else
+		{
+			SetValue<TObjectPtr<UDynamicMesh>>(Context, NewObject<UDynamicMesh>(), &Mesh);
+		}
+	}
+}
+
+void FTransformDataflowNode::Evaluate(Dataflow::FContext& Context, const FDataflowOutput* Out) const
+{
+	if (Out->IsA<TObjectPtr<UDynamicMesh>>(&Mesh))
+	{
+		SetValue<TObjectPtr<UDynamicMesh>>(Context, NewObject<UDynamicMesh>(), &Mesh);
+	}
+}
+
+void FMeshAppendDataflowNode::Evaluate(Dataflow::FContext& Context, const FDataflowOutput* Out) const
+{
+	if (Out->IsA<TObjectPtr<UDynamicMesh>>(&Mesh))
+	{
+		UE::Geometry::FDynamicMesh3& DynMesh1 = (GetValue<TObjectPtr<UDynamicMesh>>(Context, &Mesh1))->GetMeshRef();
+		UE::Geometry::FDynamicMesh3& DynMesh2 = (GetValue<TObjectPtr<UDynamicMesh>>(Context, &Mesh2))->GetMeshRef();
+
+		if (DynMesh1.VertexCount() > 0 || DynMesh2.VertexCount() > 0)
+		{
+			TObjectPtr<UDynamicMesh> DynamicMesh = NewObject<UDynamicMesh>();
+			DynamicMesh->Reset();
+
+			UE::Geometry::FDynamicMesh3& ResultDynMesh = DynamicMesh->GetMeshRef();
+
+			UE::Geometry::FDynamicMeshEditor MeshEditor(&ResultDynMesh);
+
+			UE::Geometry::FMeshIndexMappings IndexMaps1;
+			MeshEditor.AppendMesh(&DynMesh1, IndexMaps1);
+
+			UE::Geometry::FMeshIndexMappings IndexMaps2;
+			MeshEditor.AppendMesh(&DynMesh2, IndexMaps2);
+
+			SetValue<TObjectPtr<UDynamicMesh>>(Context, DynamicMesh, &Mesh);
+		}
+		else
+		{
+			SetValue<TObjectPtr<UDynamicMesh>>(Context, NewObject<UDynamicMesh>(), &Mesh);
+		}
+	}
+}
+
+void FMeshBooleanDataflowNode::Evaluate(Dataflow::FContext& Context, const FDataflowOutput* Out) const
+{
+	if (Out->IsA<TObjectPtr<UDynamicMesh>>(&Mesh))
+	{
+		UE::Geometry::FDynamicMesh3& DynMesh1 = (GetValue<TObjectPtr<UDynamicMesh>>(Context, &Mesh1))->GetMeshRef();
+		UE::Geometry::FDynamicMesh3& DynMesh2 = (GetValue<TObjectPtr<UDynamicMesh>>(Context, &Mesh2))->GetMeshRef();
+
+		if (DynMesh1.VertexCount() > 0 && DynMesh2.VertexCount() > 0)
+		{
+			// Get output
+			TObjectPtr<UDynamicMesh> DynamicMesh = NewObject<UDynamicMesh>();
+			DynamicMesh->Reset();
+			UE::Geometry::FDynamicMesh3& ResultDynMesh = DynamicMesh->GetMeshRef();
+
+			UE::Geometry::FMeshBoolean::EBooleanOp BoolOp = UE::Geometry::FMeshBoolean::EBooleanOp::Intersect;
+			if (Operation == EMeshBooleanOperationEnum::Dataflow_MeshBoolean_Intersect)
+			{
+				BoolOp = UE::Geometry::FMeshBoolean::EBooleanOp::Intersect;
+			}
+			else if (Operation == EMeshBooleanOperationEnum::Dataflow_MeshBoolean_Union)
+			{
+				BoolOp = UE::Geometry::FMeshBoolean::EBooleanOp::Union;
+			}
+			else if (Operation == EMeshBooleanOperationEnum::Dataflow_MeshBoolean_Difference)
+			{
+				BoolOp = UE::Geometry::FMeshBoolean::EBooleanOp::Difference;
+			}
+
+			UE::Geometry::FMeshBoolean Boolean(&DynMesh1, &DynMesh2, &ResultDynMesh, BoolOp);
+			Boolean.bSimplifyAlongNewEdges = true;
+			Boolean.PreserveUVsOnlyForMesh = 0; // slight warping of the autogenerated cell UVs generally doesn't matter
+			Boolean.bWeldSharedEdges = false;
+			Boolean.bTrackAllNewEdges = true;
+			if (!Boolean.Compute())
+			{
+				SetValue<TObjectPtr<UDynamicMesh>>(Context, NewObject<UDynamicMesh>(), &Mesh);
+				return;
+			}
+
+			SetValue<TObjectPtr<UDynamicMesh>>(Context, DynamicMesh, &Mesh);
+		}
+		else
+		{
+			SetValue<TObjectPtr<UDynamicMesh>>(Context, NewObject<UDynamicMesh>(), &Mesh);
+		}
+	}
+}
+
+void FMeshCopyToPointsDataflowNode::Evaluate(Dataflow::FContext& Context, const FDataflowOutput* Out) const
+{
+	if (Out->IsA<TObjectPtr<UDynamicMesh>>(&Mesh))
+	{
+		TArray<FVector> InPoints = GetValue<TArray<FVector>>(Context, &Points);
+		UE::Geometry::FDynamicMesh3& InDynMeshToCopy = (GetValue<TObjectPtr<UDynamicMesh>>(Context, &MeshToCopy))->GetMeshRef();
+
+		if (InPoints.Num() > 0 && InDynMeshToCopy.VertexCount() > 0)
+		{
+			TObjectPtr<UDynamicMesh> DynamicMesh = NewObject<UDynamicMesh>();
+			DynamicMesh->Reset();
+
+			UE::Geometry::FDynamicMesh3& ResultDynMesh = DynamicMesh->GetMeshRef();
+			UE::Geometry::FDynamicMeshEditor MeshEditor(&ResultDynMesh);
+
+			for (auto& Point : InPoints)
+			{
+				UE::Geometry::FDynamicMesh3 DynMeshTemp(InDynMeshToCopy);
+				UE::Geometry::FRefCountVector VertexRefCounts = DynMeshTemp.GetVerticesRefCounts();
+
+				UE::Geometry::FRefCountVector::IndexIterator ItVertexID = VertexRefCounts.BeginIndices();
+				const UE::Geometry::FRefCountVector::IndexIterator ItEndVertexID = VertexRefCounts.EndIndices();
+
+				while (ItVertexID != ItEndVertexID)
+				{
+					DynMeshTemp.SetVertex(*ItVertexID, Scale * DynMeshTemp.GetVertex(*ItVertexID) + Point);
+					++ItVertexID;
+				}
+
+				UE::Geometry::FMeshIndexMappings IndexMaps;
+				MeshEditor.AppendMesh(&DynMeshTemp, IndexMaps);
+			}
+
+			SetValue<TObjectPtr<UDynamicMesh>>(Context, DynamicMesh, &Mesh);
+		}
+		else
+		{
+			SetValue<TObjectPtr<UDynamicMesh>>(Context, NewObject<UDynamicMesh>(), &Mesh);
+		}
+	}
+}
+
+void FCompareIntDataflowNode::Evaluate(Dataflow::FContext& Context, const FDataflowOutput* Out) const
+{
+	if (Out->IsA<bool>(&Result))
+	{
+		int32 IntAValue = GetValue<int32>(Context, &IntA);
+		int32 IntBValue = GetValue<int32>(Context, &IntB);
+		bool ResultValue;
+
+		if (Operation == ECompareOperationEnum::Dataflow_Compare_Equal)
+		{
+			ResultValue = IntAValue == IntBValue ? true : false;
+		}
+		else if (Operation == ECompareOperationEnum::Dataflow_Compare_Smaller)
+		{
+			ResultValue = IntAValue < IntBValue ? true : false;
+		}
+		else if (Operation == ECompareOperationEnum::Dataflow_Compare_SmallerOrEqual)
+		{
+			ResultValue = IntAValue <= IntBValue ? true : false;
+		}
+		else if (Operation == ECompareOperationEnum::Dataflow_Compare_Greater)
+		{
+			ResultValue = IntAValue > IntBValue ? true : false;
+		}
+		else if (Operation == ECompareOperationEnum::Dataflow_Compare_GreaterOrEqual)
+		{
+			ResultValue = IntAValue >= IntBValue ? true : false;
+		}
+
+		SetValue<bool>(Context, ResultValue, &Result);
+	}
+}
+
+void FBranchDataflowNode::Evaluate(Dataflow::FContext& Context, const FDataflowOutput* Out) const
+{
+	if (Out->IsA<TObjectPtr<UDynamicMesh>>(&Mesh))
+	{
+		if (GetValue<bool>(Context, &Condition))
+		{
+			SetValue<TObjectPtr<UDynamicMesh>>(Context, GetValue<TObjectPtr<UDynamicMesh>>(Context, &MeshA), &Mesh);
+		}
+		else
+		{
+			SetValue<TObjectPtr<UDynamicMesh>>(Context, GetValue<TObjectPtr<UDynamicMesh>>(Context, &MeshB), &Mesh);
+		}
+	}
+}
+
+void FGetMeshDataDataflowNode::Evaluate(Dataflow::FContext& Context, const FDataflowOutput* Out) const
+{
+	if (Out->IsA<int32>(&VertexCount) ||
+		Out->IsA<int32>(&EdgeCount) ||
+		Out->IsA<int32>(&TriangleCount))
+	{
+		TObjectPtr<UDynamicMesh> DynamicMesh = GetValue<TObjectPtr<UDynamicMesh>>(Context, &Mesh);
+		UE::Geometry::FDynamicMesh3& DynMesh = DynamicMesh->GetMeshRef();
+
+		SetValue<int32>(Context, DynMesh.VertexCount(), &VertexCount);
+		SetValue<int32>(Context, DynMesh.EdgeCount(), &EdgeCount);
+		SetValue<int32>(Context, DynMesh.TriangleCount(), &TriangleCount);
+	}
+}
+
+namespace {
+	inline FName GetArrayTypeString(FManagedArrayCollection::EArrayType ArrayType)
+	{
+		switch (ArrayType)
+		{
+#define MANAGED_ARRAY_TYPE(a,A)	case EManagedArrayType::F##A##Type:\
+			return FName(#A);
+#include "GeometryCollection/ManagedArrayTypeValues.inl"
+#undef MANAGED_ARRAY_TYPE
+		}
+		return FName();
+	}
+}
+
+void FGetSchemaDataflowNode::Evaluate(Dataflow::FContext& Context, const FDataflowOutput* Out) const
+{
+	if (Out->IsA<FString>(&String))
+	{
+		FManagedArrayCollection InCollection = GetValue<FManagedArrayCollection>(Context, &Collection);
+
+		FString OutputStr;
+		OutputStr.Appendf(TEXT("\n----------------------------------------\n"));
+		for (auto& Group : InCollection.GroupNames())
+		{
+			if (InCollection.HasGroup(Group))
+			{
+				int32 NumElems = InCollection.NumElements(Group);
+
+				OutputStr.Appendf(TEXT("Group: %s  Number of Elements: %d\n"), *Group.ToString(), NumElems);
+				OutputStr.Appendf(TEXT("Attributes:\n"));
+
+				for (auto& Attr : InCollection.AttributeNames(Group))
+				{
+					if (InCollection.HasAttribute(Attr, Group))
+					{
+						FString TypeStr = GetArrayTypeString(InCollection.GetAttributeType(Attr, Group)).ToString();
+						OutputStr.Appendf(TEXT("\t%s\t[%s]\n"), *Attr.ToString(), *TypeStr);
+					}
+				}
+
+				OutputStr.Appendf(TEXT("\n--------------------\n"));
+			}
+		}
+		OutputStr.Appendf(TEXT("----------------------------------------\n"));
+
+		SetValue<FString>(Context, OutputStr, &String);
+	}
+}
+
+void FAutoClusterDataflowNode::Evaluate(Dataflow::FContext& Context, const FDataflowOutput* Out) const
+{
+	if (Out->IsA<FManagedArrayCollection>(&Collection))
+	{
+		FManagedArrayCollection InCollection = GetValue<FManagedArrayCollection>(Context, &Collection);
+		FDataflowTransformSelection InTransformSelection = GetValue<FDataflowTransformSelection>(Context, &TransformSelection);
+
+		if (FGeometryCollection* GeomCollection = InCollection.NewCopy<FGeometryCollection>())
+		{
+			EClusterSizeMethodEnum InClusterSizeMethod = ClusterSizeMethod;
+			int32 InClusterSites = GetValue<int32>(Context, &ClusterSites);
+			float InClusterFraction = GetValue<float>(Context, &ClusterFraction);
+			float InSiteSize = GetValue<float>(Context, &SiteSize);
+			bool InAutoCluster = AutoCluster;
+			bool InAvoidIsolated = AvoidIsolated;
+
+			TArray<int32> SelectedBones;
+			InTransformSelection.AsArray(SelectedBones);
+
+			FFractureEngineClustering::AutoCluster(*GeomCollection, 
+				SelectedBones, 
+				(EFractureEngineClusterSizeMethod)InClusterSizeMethod, 
+				InClusterSites,
+				InClusterFraction,
+				InSiteSize,
+				InAutoCluster,
+				InAvoidIsolated);
+
+			SetValue<FManagedArrayCollection>(Context, (const FManagedArrayCollection&)(*GeomCollection), &Collection);
+		}
+	}
+}
+
+
+void FClusterFlattenDataflowNode::Evaluate(Dataflow::FContext& Context, const FDataflowOutput* Out) const
+{
+	if (Out->IsA<FManagedArrayCollection>(&Collection))
+	{
+		FManagedArrayCollection CollectionPtr = GetValue<FManagedArrayCollection>(Context, &Collection);
+		if (FGeometryCollection* GeomCollection = CollectionPtr.NewCopy<FGeometryCollection>())
+		{
+			FGeometryCollectionClusteringUtility::UpdateHierarchyLevelOfChildren(GeomCollection, -1);
+
+			const TManagedArray<int32>& Levels = GeomCollection->GetAttribute<int32>("Level", FGeometryCollection::TransformGroup);
+
+			// Populate Selected Bones in an Array
+			// @todo(harsha) Implement with Selection
+			// For every bone in selected array: [ClusterIndex]
+			int32 ClusterIndex = 0;
+			TArray<int32> LeafBones;
+			FGeometryCollectionClusteringUtility::GetLeafBones(GeomCollection, ClusterIndex, true, LeafBones);
+			FGeometryCollectionClusteringUtility::ClusterBonesUnderExistingNode(GeomCollection, ClusterIndex, LeafBones);
+			FGeometryCollectionClusteringUtility::RemoveDanglingClusters(GeomCollection);
+			// End for
+
+			FGeometryCollectionClusteringUtility::UpdateHierarchyLevelOfChildren(GeomCollection, -1);
+			SetValue<FManagedArrayCollection>(Context, (const FManagedArrayCollection&)(*GeomCollection), &Collection);
+		}
+	}
+}
+
+
+void FRemoveOnBreakDataflowNode::Evaluate(Dataflow::FContext& Context, const FDataflowOutput* Out) const
+{
+	if (Out->IsA<FManagedArrayCollection>(&Collection))
+	{
+		FManagedArrayCollection CollectionPtr = GetValue<FManagedArrayCollection>(Context, &Collection);
+		if (FGeometryCollection* GeometryCollection = CollectionPtr.NewCopy<FGeometryCollection>())
+		{
+			if (!GeometryCollection->HasAttribute("RemoveOnBreak", FGeometryCollection::TransformGroup))
+			{
+				TManagedArray<FVector4f>& NewRemoveOnBreak = GeometryCollection->AddAttribute<FVector4f>("RemoveOnBreak", FGeometryCollection::TransformGroup);
+				NewRemoveOnBreak.Fill(FRemoveOnBreakData::DisabledPackedData);
+			}
+
+			TManagedArray<FVector4f>& RemoveOnBreak = GeometryCollection->ModifyAttribute<FVector4f>("RemoveOnBreak", FGeometryCollection::TransformGroup);
+
+			const FVector2f PostBreakTimerData = GetValue<FVector2f>(Context, &PostBreakTimer);
+			const FVector2f RemovalTimerData = GetValue<FVector2f>(Context, &RemovalTimer);
+			const bool ClusterCrumblingData = GetValue<bool>(Context, &ClusterCrumbling);
+
+			RemoveOnBreak.Fill(FVector4f{ PostBreakTimerData.X, PostBreakTimerData.Y, RemovalTimerData.X, RemovalTimerData.Y });
+
+			// @todo(harsha) Implement with Selection
+			// const FRemoveOnBreakData RemoveOnBreakData(true, PostBreakTimerData, ClusterCrumblingData, RemovalTimerData);
+			// for (int32 Index : SelectedBones)
+			// {
+				// // if root bone, then do not set 
+				// if (GeometryCollection->Parent[Index] == INDEX_NONE)
+				// {
+				// 	RemoveOnBreak[Index] = FRemoveOnBreakData::DisabledPackedData;
+				// }
+				// else
+				// {
+				// 	RemoveOnBreak[Index] = RemoveOnBreakData.GetPackedData();
+				// }
+			// }
+
+			SetValue<FManagedArrayCollection>(Context, (const FManagedArrayCollection&)(*GeometryCollection), &Collection);
+		}
+	}
+}
+
+
+void FSetAnchorStateDataflowNode::Evaluate(Dataflow::FContext& Context, const FDataflowOutput* Out) const
+{
+	if (Out->IsA<FManagedArrayCollection>(&Collection))
+	{
+		FManagedArrayCollection InCollection = GetValue<FManagedArrayCollection>(Context, &Collection);
+		FDataflowTransformSelection InTransformSelection = GetValue<FDataflowTransformSelection>(Context, &TransformSelection);
+
+		if (FGeometryCollection* GeomCollection = InCollection.NewCopy<FGeometryCollection>())
+		{
+			Chaos::Facades::FCollectionAnchoringFacade AnchoringFacade(*GeomCollection);
+			if (!AnchoringFacade.HasAnchoredAttribute())
+			{
+				AnchoringFacade.AddAnchoredAttribute();
+			}
+
+			bool bAnchored = (AnchorState == EAnchorStateEnum::Dataflow_AnchorState_Anchored) ? true : false;
+			TArray<int32> BoneIndices;
+			InTransformSelection.AsArray(BoneIndices);
+			AnchoringFacade.SetAnchored(BoneIndices, bAnchored);
+
+			if (SetNotSelectedBonesToOppositeState)
+			{
+				InTransformSelection.Invert();
+				InTransformSelection.AsArray(BoneIndices);
+				AnchoringFacade.SetAnchored(BoneIndices, !bAnchored);
+			}
+
+			SetValue<FManagedArrayCollection>(Context, (const FManagedArrayCollection&)(*GeomCollection), &Collection);
+		}
+	}
+}
+
+
+void FProximityDataflowNode::Evaluate(Dataflow::FContext& Context, const FDataflowOutput* Out) const
+{
+	if (Out->IsA<FManagedArrayCollection>(&Collection))
+	{
+		FManagedArrayCollection InCollection = GetValue<FManagedArrayCollection>(Context, &Collection);
+
+		if (FGeometryCollection* GeomCollection = InCollection.NewCopy<FGeometryCollection>())
+		{
+			FGeometryCollectionProximityPropertiesInterface::FProximityProperties Properties = GeomCollection->GetProximityProperties();
+
+			Properties.Method = (EProximityMethod)ProximityMethod;
+			Properties.DistanceThreshold = DistanceThreshold;
+			Properties.bUseAsConnectionGraph = bUseAsConnectionGraph;
+
+			GeomCollection->SetProximityProperties(Properties);
+
+			// Invalidate proximity
+			FGeometryCollectionProximityUtility ProximityUtility(GeomCollection);
+			ProximityUtility.InvalidateProximity();
+
+			SetValue<FManagedArrayCollection>(Context, (const FManagedArrayCollection&)(*GeomCollection), &Collection);
+		}
+	}
+}
+
+
+
+
+
+
+
+
+
diff --git a/Engine/Plugins/Experimental/GeometryCollectionPlugin/Source/GeometryCollectionNodes/Private/Dataflow/GeometryCollectionNodesPlugin.cpp b/Engine/Plugins/Experimental/GeometryCollectionPlugin/Source/GeometryCollectionNodes/Private/Dataflow/GeometryCollectionNodesPlugin.cpp
--- a/Engine/Plugins/Experimental/GeometryCollectionPlugin/Source/GeometryCollectionNodes/Private/Dataflow/GeometryCollectionNodesPlugin.cpp
+++ b/Engine/Plugins/Experimental/GeometryCollectionPlugin/Source/GeometryCollectionNodes/Private/Dataflow/GeometryCollectionNodesPlugin.cpp
@@ -9,6 +9,7 @@
 #include "Dataflow/Nodes/GeometryCollectionAssetNodes.h"
 #include "Dataflow/GeometryCollectionProcessingNodes.h"
 #include "Dataflow/GeometryCollectionSkeletalMeshNodes.h"
+#include "Dataflow/GeometryCollectionSelectionNodes.h"
 
 #define LOCTEXT_NAMESPACE "DataflowNodes"
 
@@ -19,6 +20,7 @@ void IGeometryCollectionNodesPlugin::StartupModule()
 	Dataflow::GeometryCollectionEngineAssetNodes();
 	Dataflow::GeometryCollectionProcessingNodes();
 	Dataflow::GeometryCollectionSkeletalMeshNodes();
+	Dataflow::GeometryCollectionSelectionNodes();
 }
 
 void IGeometryCollectionNodesPlugin::ShutdownModule()
diff --git a/Engine/Plugins/Experimental/GeometryCollectionPlugin/Source/GeometryCollectionNodes/Private/Dataflow/GeometryCollectionSelectionNodes.cpp b/Engine/Plugins/Experimental/GeometryCollectionPlugin/Source/GeometryCollectionNodes/Private/Dataflow/GeometryCollectionSelectionNodes.cpp
new file mode 100644
--- /dev/null
+++ b/Engine/Plugins/Experimental/GeometryCollectionPlugin/Source/GeometryCollectionNodes/Private/Dataflow/GeometryCollectionSelectionNodes.cpp
@@ -0,0 +1,620 @@
+// Copyright Epic Games, Inc. All Rights Reserved.
+
+#include "Dataflow/GeometryCollectionSelectionNodes.h"
+#include "Dataflow/DataflowCore.h"
+
+#include "Engine/StaticMesh.h"
+#include "GeometryCollection/GeometryCollectionObject.h"
+#include "GeometryCollection/ManagedArrayCollection.h"
+#include "GeometryCollection/GeometryCollection.h"
+#include "GeometryCollection/GeometryCollectionEngineUtility.h"
+#include "GeometryCollection/GeometryCollectionEngineConversion.h"
+#include "Logging/LogMacros.h"
+#include "Templates/SharedPointer.h"
+#include "UObject/UnrealTypePrivate.h"
+#include "DynamicMeshToMeshDescription.h"
+#include "MeshDescriptionToDynamicMesh.h"
+#include "StaticMeshAttributes.h"
+#include "DynamicMeshEditor.h"
+#include "Operations/MeshBoolean.h"
+
+#include "EngineGlobals.h"
+#include "GeometryCollection/GeometryCollectionAlgo.h"
+#include "GeometryCollection/GeometryCollectionClusteringUtility.h"
+#include "GeometryCollection/GeometryCollectionConvexUtility.h"
+#include "Voronoi/Voronoi.h"
+#include "PlanarCut.h"
+#include "GeometryCollection/GeometryCollectionProximityUtility.h"
+#include "FractureEngineClustering.h"
+#include "FractureEngineSelection.h"
+
+#include UE_INLINE_GENERATED_CPP_BY_NAME(GeometryCollectionSelectionNodes)
+
+namespace Dataflow
+{
+
+	void GeometryCollectionSelectionNodes()
+	{
+		static const FLinearColor CDefaultNodeBodyTintColor = FLinearColor(0.f, 0.f, 0.f, 0.5f);
+
+		DATAFLOW_NODE_REGISTER_CREATION_FACTORY(FCollectionTransformSelectionAllDataflowNode);
+		DATAFLOW_NODE_REGISTER_CREATION_FACTORY(FCollectionTransformSelectionSetOperationDataflowNode);
+		DATAFLOW_NODE_REGISTER_CREATION_FACTORY(FCollectionTransformSelectionInfoDataflowNode);
+		DATAFLOW_NODE_REGISTER_CREATION_FACTORY(FCollectionTransformSelectionNoneDataflowNode);
+		DATAFLOW_NODE_REGISTER_CREATION_FACTORY(FCollectionTransformSelectionInvertDataflowNode);
+		DATAFLOW_NODE_REGISTER_CREATION_FACTORY(FCollectionTransformSelectionRandomDataflowNode);
+		DATAFLOW_NODE_REGISTER_CREATION_FACTORY(FCollectionTransformSelectionRootDataflowNode);
+		DATAFLOW_NODE_REGISTER_CREATION_FACTORY(FCollectionTransformSelectionCustomDataflowNode);
+		DATAFLOW_NODE_REGISTER_CREATION_FACTORY(FCollectionTransformSelectionParentDataflowNode);
+		DATAFLOW_NODE_REGISTER_CREATION_FACTORY(FCollectionTransformSelectionByPercentageDataflowNode);
+		DATAFLOW_NODE_REGISTER_CREATION_FACTORY(FCollectionTransformSelectionChildrenDataflowNode);
+		DATAFLOW_NODE_REGISTER_CREATION_FACTORY(FCollectionTransformSelectionSiblingsDataflowNode);
+		DATAFLOW_NODE_REGISTER_CREATION_FACTORY(FCollectionTransformSelectionLevelDataflowNode);
+		DATAFLOW_NODE_REGISTER_CREATION_FACTORY(FCollectionTransformSelectionContactDataflowNode);
+		DATAFLOW_NODE_REGISTER_CREATION_FACTORY(FCollectionTransformSelectionLeafDataflowNode);
+		DATAFLOW_NODE_REGISTER_CREATION_FACTORY(FCollectionTransformSelectionClusterDataflowNode);
+		DATAFLOW_NODE_REGISTER_CREATION_FACTORY(FCollectionTransformSelectionBySizeDataflowNode);
+		DATAFLOW_NODE_REGISTER_CREATION_FACTORY(FCollectionTransformSelectionByVolumeDataflowNode);
+
+		// GeometryCollection|Selection
+		DATAFLOW_NODE_REGISTER_CREATION_FACTORY_NODE_COLORS_BY_CATEGORY("GeometryCollection|Selection", FLinearColor(1.f, 1.f, 0.05f), CDefaultNodeBodyTintColor);
+	}
+}
+
+
+void FCollectionTransformSelectionAllDataflowNode::Evaluate(Dataflow::FContext& Context, const FDataflowOutput* Out) const
+{
+	if (Out->IsA<FManagedArrayCollection>(&Collection) ||
+		Out->IsA<FDataflowTransformSelection>(&TransformSelection))
+	{
+		FManagedArrayCollection InCollection = GetValue<FManagedArrayCollection>(Context, &Collection);
+
+		int32 NumElems = InCollection.NumElements(FGeometryCollection::TransformGroup);
+
+		FDataflowTransformSelection NewTransformSelection;
+		NewTransformSelection.Initialize(NumElems, true);
+
+		SetValue<FManagedArrayCollection>(Context, InCollection, &Collection);
+		SetValue<FDataflowTransformSelection>(Context, NewTransformSelection, &TransformSelection);
+	}
+}
+
+
+void FCollectionTransformSelectionSetOperationDataflowNode::Evaluate(Dataflow::FContext& Context, const FDataflowOutput* Out) const
+{
+	if (Out->IsA<FDataflowTransformSelection>(&TransformSelection))
+	{
+		FDataflowTransformSelection InTransformSelectionA = GetValue<FDataflowTransformSelection>(Context, &TransformSelectionA);
+		FDataflowTransformSelection InTransformSelectionB = GetValue<FDataflowTransformSelection>(Context, &TransformSelectionB);
+
+		FDataflowTransformSelection NewTransformSelection;
+
+		if (InTransformSelectionA.Num() == InTransformSelectionB.Num())
+		{
+			if (Operation == ESetOperationEnum::Dataflow_SetOperation_AND)
+			{
+				InTransformSelectionA.AND(InTransformSelectionB, NewTransformSelection);
+			}
+			else if (Operation == ESetOperationEnum::Dataflow_SetOperation_OR)
+			{
+				InTransformSelectionA.OR(InTransformSelectionB, NewTransformSelection);
+			}
+			else if (Operation == ESetOperationEnum::Dataflow_SetOperation_XOR)
+			{
+				InTransformSelectionA.XOR(InTransformSelectionB, NewTransformSelection);
+			}
+		}
+		else
+		{
+			// ERROR: INPUT TRANSFORMSELECTIONS HAVE DIFFERENT NUMBER OF ELEMENTS
+			FString ErrorStr = "Input TransformSelections have different number of elements.";
+			UE_LOG(LogTemp, Error, TEXT("[Dataflow ERROR] %s"), *ErrorStr);
+		}
+
+		SetValue<FDataflowTransformSelection>(Context, NewTransformSelection, &TransformSelection);
+	}
+}
+
+
+namespace {
+	struct BoneInfo {
+		int32 BoneIndex;
+		int32 Level;
+	};
+}
+
+static void ExpandRecursive(const int32 BoneIndex, int32 Level, const TManagedArray<TSet<int32>>& Children, TArray<BoneInfo>& BoneHierarchy)
+{
+	BoneHierarchy.Add({ BoneIndex, Level });
+
+	TSet<int32> ChildrenSet = Children[BoneIndex];
+	if (ChildrenSet.Num() > 0)
+	{
+		for (auto& Child : ChildrenSet)
+		{
+			ExpandRecursive(Child, Level + 1, Children, BoneHierarchy);
+		}
+	}
+}
+
+
+static void BuildHierarchicalOutput(const TManagedArray<int32>& Parents, 
+	const TManagedArray<TSet<int32>>& Children, 
+	const TManagedArray<FString>& BoneNames,
+	FDataflowTransformSelection& TransformSelection, 
+	FString& OutputStr)
+{
+	TArray<BoneInfo> BoneHierarchy;
+
+	int32 NumElements = Parents.Num();
+	for (int32 Index = 0; Index < NumElements; ++Index)
+	{
+		if (Parents[Index] == FGeometryCollection::Invalid)
+		{
+			ExpandRecursive(Index, 0, Children, BoneHierarchy);
+		}
+	}
+
+	// Get level max
+	int32 LevelMax = -1;
+	int32 BoneNameLengthMax = -1;
+	for (int32 Idx = 0; Idx < BoneHierarchy.Num(); ++Idx)
+	{
+		if (BoneHierarchy[Idx].Level > LevelMax)
+		{
+			LevelMax = BoneHierarchy[Idx].Level;
+		}
+
+		int32 BoneNameLength = BoneNames[Idx].Len();
+		if (BoneNameLength > BoneNameLengthMax)
+		{
+			BoneNameLengthMax = BoneNameLength;
+		}
+	}
+
+	const int32 BoneIndexWidth = 2 + LevelMax * 2 + 6;
+	const int32 BoneNameWidth = BoneNameLengthMax + 2;
+	const int32 SelectedWidth = 10;
+
+	for (int32 Idx = 0; Idx < BoneHierarchy.Num(); ++Idx)
+	{
+		FString BoneIndexStr, BoneNameStr;
+		BoneIndexStr.Reserve(BoneIndexWidth);
+		BoneNameStr.Reserve(BoneNameWidth);
+
+		if (BoneHierarchy[Idx].Level == 0)
+		{
+			BoneIndexStr.Appendf(TEXT("[%d]"), BoneHierarchy[Idx].BoneIndex);
+		}
+		else
+		{
+			BoneIndexStr.Appendf(TEXT(" |"));
+			for (int32 Idx1 = 0; Idx1 < BoneHierarchy[Idx].Level; ++Idx1)
+			{
+				BoneIndexStr.Appendf(TEXT("--"));
+			}
+			BoneIndexStr.Appendf(TEXT("[%d]"), BoneHierarchy[Idx].BoneIndex);
+		}
+		BoneIndexStr = BoneIndexStr.RightPad(BoneIndexWidth);
+
+		BoneNameStr.Appendf(TEXT("%s"), *BoneNames[Idx]);
+		BoneNameStr = BoneNameStr.RightPad(BoneNameWidth);
+
+		OutputStr.Appendf(TEXT("%s%s%s\n\n"), *BoneIndexStr, *BoneNameStr, (TransformSelection.IsSelected(BoneHierarchy[Idx].BoneIndex) ? TEXT("Selected") : TEXT("---")));
+	}
+
+}
+
+
+void FCollectionTransformSelectionInfoDataflowNode::Evaluate(Dataflow::FContext& Context, const FDataflowOutput* Out) const
+{
+	if (Out->IsA<FString>(&String))
+	{
+		FDataflowTransformSelection InTransformSelection = GetValue<FDataflowTransformSelection>(Context, &TransformSelection);
+		FManagedArrayCollection InCollection = GetValue<FManagedArrayCollection>(Context, &Collection);
+
+		FString OutputStr;
+
+		OutputStr.Appendf(TEXT("\n----------------------------------------\n"));
+		OutputStr.Appendf(TEXT("Number of Elements: %d\n"), InTransformSelection.Num());
+
+		// Hierarchical display
+		if (InCollection.HasGroup(FGeometryCollection::TransformGroup) &&
+			InCollection.HasAttribute("Parent", FGeometryCollection::TransformGroup) &&
+			InCollection.HasAttribute("Children", FGeometryCollection::TransformGroup) &&
+			InCollection.HasAttribute("BoneName", FGeometryCollection::TransformGroup))
+		{
+			if (InTransformSelection.Num() == InCollection.NumElements(FGeometryCollection::TransformGroup))
+			{
+				const TManagedArray<int32>& Parents = InCollection.GetAttribute<int32>("Parent", FGeometryCollection::TransformGroup);
+				const TManagedArray<TSet<int32>>& Children = InCollection.GetAttribute<TSet<int32>>("Children", FGeometryCollection::TransformGroup);
+				const TManagedArray<FString>& BoneNames = InCollection.GetAttribute<FString>("BoneName", FGeometryCollection::TransformGroup);
+
+				BuildHierarchicalOutput(Parents, Children, BoneNames, InTransformSelection, OutputStr);
+			}
+			else
+			{
+				// ERROR: TransformSelection doesn't match the Collection
+				FString ErrorStr = "TransformSelection doesn't match the Collection.";
+				UE_LOG(LogTemp, Error, TEXT("[Dataflow ERROR] %s"), *ErrorStr);
+			}
+		}
+		else
+		// Simple display
+		{
+			for (int32 Idx = 0; Idx < InTransformSelection.Num(); ++Idx)
+			{
+				OutputStr.Appendf(TEXT("%4d: %s\n"), Idx, (InTransformSelection.IsSelected(Idx) ? TEXT("Selected") : TEXT("---")));
+			}
+		}
+
+		OutputStr.Appendf(TEXT("----------------------------------------\n"));
+
+		SetValue<FString>(Context, OutputStr, &String);
+	}
+}
+
+
+void FCollectionTransformSelectionNoneDataflowNode::Evaluate(Dataflow::FContext& Context, const FDataflowOutput* Out) const
+{
+	if (Out->IsA<FManagedArrayCollection>(&Collection) ||
+		Out->IsA<FDataflowTransformSelection>(&TransformSelection))
+	{
+		FManagedArrayCollection InCollection = GetValue<FManagedArrayCollection>(Context, &Collection);
+
+		const int32 NumElems = InCollection.NumElements(FGeometryCollection::TransformGroup);
+
+		FDataflowTransformSelection NewTransformSelection;
+		NewTransformSelection.Initialize(NumElems, false);
+
+		SetValue<FManagedArrayCollection>(Context, InCollection, &Collection);
+		SetValue<FDataflowTransformSelection>(Context, NewTransformSelection, &TransformSelection);
+	}
+}
+
+
+void FCollectionTransformSelectionInvertDataflowNode::Evaluate(Dataflow::FContext& Context, const FDataflowOutput* Out) const
+{
+	if (Out->IsA<FDataflowTransformSelection>(&TransformSelection))
+	{
+		FDataflowTransformSelection InTransformSelection = GetValue<FDataflowTransformSelection>(Context, &TransformSelection);
+
+		InTransformSelection.Invert();
+
+		SetValue<FDataflowTransformSelection>(Context, InTransformSelection, &TransformSelection);
+	}
+}
+
+
+void FCollectionTransformSelectionRandomDataflowNode::Evaluate(Dataflow::FContext& Context, const FDataflowOutput* Out) const
+{
+	if (Out->IsA<FManagedArrayCollection>(&Collection) ||
+		Out->IsA<FDataflowTransformSelection>(&TransformSelection))
+	{
+		FManagedArrayCollection InCollection = GetValue<FManagedArrayCollection>(Context, &Collection);
+
+		int32 NumElems = InCollection.NumElements(FGeometryCollection::TransformGroup);
+
+		FDataflowTransformSelection NewTransformSelection;
+		NewTransformSelection.Initialize(NumElems, false);
+
+		float RandomSeedVal = GetValue<float>(Context, &RandomSeed);
+		float RandomThresholdVal = GetValue<float>(Context, &RandomThreshold);
+
+		FRandomStream Stream(RandomSeedVal);
+
+		for (int32 Idx = 0; Idx < NumElems; ++Idx)
+		{
+			if (Deterministic)
+			{
+				float RandomVal = Stream.FRandRange(0.f, 1.f);
+
+				if (RandomVal > RandomThresholdVal)
+				{
+					NewTransformSelection.SetSelected(Idx);
+				}
+			}
+			else
+			{
+				float RandomVal = FMath::FRandRange(0.f, 1.f);
+
+				if (RandomVal > RandomThresholdVal)
+				{
+					NewTransformSelection.SetSelected(Idx);
+				}
+			}
+		}
+
+		SetValue<FManagedArrayCollection>(Context, InCollection, &Collection);
+		SetValue<FDataflowTransformSelection>(Context, NewTransformSelection, &TransformSelection);
+	}
+}
+
+
+void FCollectionTransformSelectionRootDataflowNode::Evaluate(Dataflow::FContext& Context, const FDataflowOutput* Out) const
+{
+	if (Out->IsA<FManagedArrayCollection>(&Collection) ||
+		Out->IsA<FDataflowTransformSelection>(&TransformSelection))
+	{
+		FManagedArrayCollection InCollection = GetValue<FManagedArrayCollection>(Context, &Collection);
+
+		if (InCollection.HasGroup(FGeometryCollection::TransformGroup) &&
+			InCollection.HasAttribute("Parent", FGeometryCollection::TransformGroup))
+		{
+			const TManagedArray<int32>& Parent = InCollection.GetAttribute<int32>("Parent", FGeometryCollection::TransformGroup);
+
+			int32 NumElems = InCollection.NumElements(FGeometryCollection::TransformGroup);
+
+			FDataflowTransformSelection NewTransformSelection;
+			NewTransformSelection.Initialize(NumElems, false);
+
+			for (int32 Idx = 0; Idx < NumElems; ++Idx)
+			{
+				if (Parent[Idx] == FGeometryCollection::Invalid)
+				{
+					NewTransformSelection.SetSelected(Idx);
+				}
+			}
+
+			SetValue<FDataflowTransformSelection>(Context, NewTransformSelection, &TransformSelection);
+		}
+		else
+		{
+			SetValue<FDataflowTransformSelection>(Context, FDataflowTransformSelection(), &TransformSelection);
+		}
+
+		SetValue<FManagedArrayCollection>(Context, InCollection, &Collection);
+	}
+}
+
+
+void FCollectionTransformSelectionCustomDataflowNode::Evaluate(Dataflow::FContext& Context, const FDataflowOutput* Out) const
+{
+	if (Out->IsA<FManagedArrayCollection>(&Collection) ||
+		Out->IsA<FDataflowTransformSelection>(&TransformSelection))
+	{
+		FManagedArrayCollection InCollection = GetValue<FManagedArrayCollection>(Context, &Collection);
+
+		if (InCollection.HasGroup(FGeometryCollection::TransformGroup))
+		{
+			const int32 NumElems = InCollection.NumElements(FGeometryCollection::TransformGroup);
+
+			FDataflowTransformSelection NewTransformSelection;
+			NewTransformSelection.Initialize(NumElems, false);
+
+			const FString InBoneIndices = GetValue<FString>(Context, &BoneIndicies);
+
+			TArray<FString> Indicies;
+			InBoneIndices.ParseIntoArray(Indicies, TEXT(" "), true);
+
+			for (FString IndexStr : Indicies)
+			{
+				if (IndexStr.IsNumeric())
+				{
+					int32 Index = FCString::Atoi(*IndexStr);
+					if (Index > 0 && Index < NumElems)
+					{
+						NewTransformSelection.SetSelected(Index);
+					}
+					else
+					{
+						// ERROR: INVALID INDEX
+						FString ErrorStr = "Invalid specified index found.";
+						UE_LOG(LogTemp, Error, TEXT("[Dataflow ERROR] %s"), *ErrorStr);
+					}
+				}
+			}
+
+			SetValue<FDataflowTransformSelection>(Context, NewTransformSelection, &TransformSelection);
+		}
+		else
+		{
+			SetValue<FDataflowTransformSelection>(Context, FDataflowTransformSelection(), &TransformSelection);
+		}
+
+		SetValue<FManagedArrayCollection>(Context, InCollection, &Collection);
+	}
+}
+
+
+void FCollectionTransformSelectionParentDataflowNode::Evaluate(Dataflow::FContext& Context, const FDataflowOutput* Out) const
+{
+	if (Out->IsA<FDataflowTransformSelection>(&TransformSelection) ||
+		Out->IsA<FManagedArrayCollection>(&Collection))
+	{
+		FManagedArrayCollection InCollection = GetValue<FManagedArrayCollection>(Context, &Collection);
+		FDataflowTransformSelection InTransformSelection = GetValue<FDataflowTransformSelection>(Context, &TransformSelection);
+
+		if (FGeometryCollection* GeomCollection = InCollection.NewCopy<FGeometryCollection>())
+		{
+			FFractureEngineSelection::SelectParent(*GeomCollection, InTransformSelection);
+		}
+		
+		SetValue<FDataflowTransformSelection>(Context, InTransformSelection, &TransformSelection);
+		SetValue<FManagedArrayCollection>(Context, InCollection, &Collection);
+	}
+}
+
+
+void FCollectionTransformSelectionByPercentageDataflowNode::Evaluate(Dataflow::FContext& Context, const FDataflowOutput* Out) const
+{
+	if (Out->IsA<FDataflowTransformSelection>(&TransformSelection) ||
+		Out->IsA<FManagedArrayCollection>(&Collection))
+	{
+		FManagedArrayCollection InCollection = GetValue<FManagedArrayCollection>(Context, &Collection);
+		FDataflowTransformSelection InTransformSelection = GetValue<FDataflowTransformSelection>(Context, &TransformSelection);
+
+		int32 InPercentage = GetValue<int32>(Context, &Percentage);
+		float InRandomSeed = GetValue<float>(Context, &RandomSeed);
+
+		if (FGeometryCollection* GeomCollection = InCollection.NewCopy<FGeometryCollection>())
+		{
+			FFractureEngineSelection::SelectByPercentage(*GeomCollection, InTransformSelection, InPercentage, Deterministic, InRandomSeed);
+		}
+
+		SetValue<FDataflowTransformSelection>(Context, InTransformSelection, &TransformSelection);
+		SetValue<FManagedArrayCollection>(Context, InCollection, &Collection);
+	}
+}
+
+
+void FCollectionTransformSelectionChildrenDataflowNode::Evaluate(Dataflow::FContext& Context, const FDataflowOutput* Out) const
+{
+	if (Out->IsA<FDataflowTransformSelection>(&TransformSelection) ||
+		Out->IsA<FManagedArrayCollection>(&Collection))
+	{
+		FManagedArrayCollection InCollection = GetValue<FManagedArrayCollection>(Context, &Collection);
+		FDataflowTransformSelection InTransformSelection = GetValue<FDataflowTransformSelection>(Context, &TransformSelection);
+
+		if (FGeometryCollection* GeomCollection = InCollection.NewCopy<FGeometryCollection>())
+		{
+			FFractureEngineSelection::SelectChildren(*GeomCollection, InTransformSelection);
+		}
+
+		SetValue<FDataflowTransformSelection>(Context, InTransformSelection, &TransformSelection);
+		SetValue<FManagedArrayCollection>(Context, InCollection, &Collection);
+	}
+}
+
+
+void FCollectionTransformSelectionSiblingsDataflowNode::Evaluate(Dataflow::FContext& Context, const FDataflowOutput* Out) const
+{
+	if (Out->IsA<FDataflowTransformSelection>(&TransformSelection) ||
+		Out->IsA<FManagedArrayCollection>(&Collection))
+	{
+		FManagedArrayCollection InCollection = GetValue<FManagedArrayCollection>(Context, &Collection);
+		FDataflowTransformSelection InTransformSelection = GetValue<FDataflowTransformSelection>(Context, &TransformSelection);
+
+		if (FGeometryCollection* GeomCollection = InCollection.NewCopy<FGeometryCollection>())
+		{
+			FFractureEngineSelection::SelectSiblings(*GeomCollection, InTransformSelection);
+		}
+
+		SetValue<FDataflowTransformSelection>(Context, InTransformSelection, &TransformSelection);
+		SetValue<FManagedArrayCollection>(Context, InCollection, &Collection);
+	}
+}
+
+
+void FCollectionTransformSelectionLevelDataflowNode::Evaluate(Dataflow::FContext& Context, const FDataflowOutput* Out) const
+{
+	if (Out->IsA<FDataflowTransformSelection>(&TransformSelection) ||
+		Out->IsA<FManagedArrayCollection>(&Collection))
+	{
+		FManagedArrayCollection InCollection = GetValue<FManagedArrayCollection>(Context, &Collection);
+		FDataflowTransformSelection InTransformSelection = GetValue<FDataflowTransformSelection>(Context, &TransformSelection);
+
+		if (FGeometryCollection* GeomCollection = InCollection.NewCopy<FGeometryCollection>())
+		{
+			FFractureEngineSelection::SelectLevel(*GeomCollection, InTransformSelection);
+		}
+
+		SetValue<FDataflowTransformSelection>(Context, InTransformSelection, &TransformSelection);
+		SetValue<FManagedArrayCollection>(Context, InCollection, &Collection);
+	}
+}
+
+
+void FCollectionTransformSelectionContactDataflowNode::Evaluate(Dataflow::FContext& Context, const FDataflowOutput* Out) const
+{
+	if (Out->IsA<FDataflowTransformSelection>(&TransformSelection) ||
+		Out->IsA<FManagedArrayCollection>(&Collection))
+	{
+		FManagedArrayCollection InCollection = GetValue<FManagedArrayCollection>(Context, &Collection);
+		FDataflowTransformSelection InTransformSelection = GetValue<FDataflowTransformSelection>(Context, &TransformSelection);
+
+		if (FGeometryCollection* GeomCollection = InCollection.NewCopy<FGeometryCollection>())
+		{
+			FFractureEngineSelection::SelectContact(*GeomCollection, InTransformSelection);
+		}
+
+		SetValue<FDataflowTransformSelection>(Context, InTransformSelection, &TransformSelection);
+		SetValue<FManagedArrayCollection>(Context, InCollection, &Collection);
+	}
+}
+
+
+void FCollectionTransformSelectionLeafDataflowNode::Evaluate(Dataflow::FContext& Context, const FDataflowOutput* Out) const
+{
+	if (Out->IsA<FDataflowTransformSelection>(&TransformSelection) ||
+		Out->IsA<FManagedArrayCollection>(&Collection))
+	{
+		FManagedArrayCollection InCollection = GetValue<FManagedArrayCollection>(Context, &Collection);
+
+		FDataflowTransformSelection NewTransformSelection;
+
+		if (FGeometryCollection* GeomCollection = InCollection.NewCopy<FGeometryCollection>())
+		{
+			FFractureEngineSelection::SelectLeaf(*GeomCollection, NewTransformSelection);
+		}
+
+		SetValue<FDataflowTransformSelection>(Context, NewTransformSelection, &TransformSelection);
+		SetValue<FManagedArrayCollection>(Context, InCollection, &Collection);
+	}
+}
+
+
+void FCollectionTransformSelectionClusterDataflowNode::Evaluate(Dataflow::FContext& Context, const FDataflowOutput* Out) const
+{
+	if (Out->IsA<FDataflowTransformSelection>(&TransformSelection) ||
+		Out->IsA<FManagedArrayCollection>(&Collection))
+	{
+		FManagedArrayCollection InCollection = GetValue<FManagedArrayCollection>(Context, &Collection);
+
+		FDataflowTransformSelection NewTransformSelection;
+
+		if (FGeometryCollection* GeomCollection = InCollection.NewCopy<FGeometryCollection>())
+		{
+			FFractureEngineSelection::SelectCluster(*GeomCollection, NewTransformSelection);
+		}
+
+		SetValue<FDataflowTransformSelection>(Context, NewTransformSelection, &TransformSelection);
+		SetValue<FManagedArrayCollection>(Context, InCollection, &Collection);
+	}
+}
+
+
+void FCollectionTransformSelectionBySizeDataflowNode::Evaluate(Dataflow::FContext& Context, const FDataflowOutput* Out) const
+{
+	if (Out->IsA<FDataflowTransformSelection>(&TransformSelection) ||
+		Out->IsA<FManagedArrayCollection>(&Collection))
+	{
+		FManagedArrayCollection InCollection = GetValue<FManagedArrayCollection>(Context, &Collection);
+		float InSizeMin = GetValue<float>(Context, &SizeMin);
+		float InSizeMax = GetValue<float>(Context, &SizeMax);
+
+		FDataflowTransformSelection NewTransformSelection;
+
+		if (FGeometryCollection* GeomCollection = InCollection.NewCopy<FGeometryCollection>())
+		{
+			FFractureEngineSelection::SelectBySize(*GeomCollection, NewTransformSelection, InSizeMin, InSizeMax);
+		}
+
+		SetValue<FDataflowTransformSelection>(Context, NewTransformSelection, &TransformSelection);
+		SetValue<FManagedArrayCollection>(Context, InCollection, &Collection);
+	}
+}
+
+
+void FCollectionTransformSelectionByVolumeDataflowNode::Evaluate(Dataflow::FContext& Context, const FDataflowOutput* Out) const
+{
+	if (Out->IsA<FDataflowTransformSelection>(&TransformSelection) ||
+		Out->IsA<FManagedArrayCollection>(&Collection))
+	{
+		FManagedArrayCollection InCollection = GetValue<FManagedArrayCollection>(Context, &Collection);
+		float InVolumeMin = GetValue<float>(Context, &VolumeMin);
+		float InVolumeMax = GetValue<float>(Context, &VolumeMax);
+
+		FDataflowTransformSelection NewTransformSelection;
+
+		if (FGeometryCollection* GeomCollection = InCollection.NewCopy<FGeometryCollection>())
+		{
+			FFractureEngineSelection::SelectByVolume(*GeomCollection, NewTransformSelection, InVolumeMin, InVolumeMax);
+		}
+
+		SetValue<FDataflowTransformSelection>(Context, NewTransformSelection, &TransformSelection);
+		SetValue<FManagedArrayCollection>(Context, InCollection, &Collection);
+	}
+}
+
+
diff --git a/Engine/Plugins/Experimental/GeometryCollectionPlugin/Source/GeometryCollectionNodes/Public/Dataflow/GeometryCollectionNodes.h b/Engine/Plugins/Experimental/GeometryCollectionPlugin/Source/GeometryCollectionNodes/Public/Dataflow/GeometryCollectionNodes.h
--- a/Engine/Plugins/Experimental/GeometryCollectionPlugin/Source/GeometryCollectionNodes/Public/Dataflow/GeometryCollectionNodes.h
+++ b/Engine/Plugins/Experimental/GeometryCollectionPlugin/Source/GeometryCollectionNodes/Public/Dataflow/GeometryCollectionNodes.h
@@ -5,11 +5,15 @@
 #include "CoreMinimal.h"
 #include "Dataflow/DataflowEngine.h"
 #include "GeometryCollection/ManagedArrayCollection.h"
+#include "DynamicMesh/DynamicMesh3.h"
+#include "UDynamicMesh.h"
+#include "Dataflow/DataflowSelection.h"
 
 #include "GeometryCollectionNodes.generated.h"
 
 
 class FGeometryCollection;
+class UStaticMesh;
 
 USTRUCT()
 struct FGetCollectionAssetDataflowNode : public FDataflowNode
@@ -49,7 +53,7 @@ public:
 	UPROPERTY(EditAnywhere, Category = "Dataflow", meta = (EditCondition = "false", EditConditionHides));
 	float Scale = 1.0;
 
-	UPROPERTY(meta = (DataflowInput, DataflowOutput, DisplayName = "Collection", Passthrough = "Collection"))
+	UPROPERTY(meta = (DataflowInput, DataflowOutput, DataflowPassthrough = "Collection"))
 	FManagedArrayCollection Collection;
 
 	FExampleCollectionEditDataflowNode(const Dataflow::FNodeParameters& InParam, FGuid InGuid = FGuid::NewGuid())
@@ -72,7 +76,7 @@ struct FSetCollectionAssetDataflowNode : public FDataflowTerminalNode
 public:
 	typedef FManagedArrayCollection DataType;
 
-	UPROPERTY(meta = (DataflowInput, DisplayName = "Collection"))
+	UPROPERTY(meta = (DataflowInput))
 	FManagedArrayCollection Collection;
 
 	FSetCollectionAssetDataflowNode(const Dataflow::FNodeParameters& InParam, FGuid InGuid = FGuid::NewGuid())
@@ -90,12 +94,12 @@ USTRUCT()
 struct FAppendCollectionAssetsDataflowNode : public FDataflowNode
 {
 	GENERATED_USTRUCT_BODY()
-		DATAFLOW_NODE_DEFINE_INTERNAL(FAppendCollectionAssetsDataflowNode, "AppendCollections", "GeometryCollection", "")
+	DATAFLOW_NODE_DEFINE_INTERNAL(FAppendCollectionAssetsDataflowNode, "AppendCollections", "GeometryCollection", "")
 
 public:
 	typedef FManagedArrayCollection DataType;
 
-	UPROPERTY(meta = (DataflowInput, DisplayOutput, DisplayName = "Collection", Passthrough = "Collection"))
+	UPROPERTY(meta = (DataflowInput, DataflowOutput, DisplayName = "Collection", DataflowPassthrough = "Collection"))
 	FManagedArrayCollection Collection1;
 
 	UPROPERTY(meta = (DataflowInput, DisplayName = "Collection"))
@@ -123,7 +127,7 @@ struct FResetGeometryCollectionDataflowNode : public FDataflowNode
 public:
 	typedef FManagedArrayCollection DataType;
 
-	UPROPERTY(meta = (DataflowOutput, DisplayName = "Collection"))
+	UPROPERTY(meta = (DataflowOutput))
 	FManagedArrayCollection Collection;
 
 	FResetGeometryCollectionDataflowNode(const Dataflow::FNodeParameters& InParam, FGuid InGuid = FGuid::NewGuid())
@@ -156,7 +160,7 @@ public:
 	UPROPERTY(EditAnywhere, Category = "Print");
 	float Duration = 2.f;
 
-	UPROPERTY(EditAnywhere, Category = "Print", meta = (DataflowInput, DisplayName = "String"));
+	UPROPERTY(EditAnywhere, Category = "Print", meta = (DataflowInput));
 	FString String = FString("");
 
 	FPrintStringDataflowNode(const Dataflow::FNodeParameters& InParam, FGuid InGuid = FGuid::NewGuid())
@@ -180,7 +184,7 @@ public:
 	UPROPERTY(EditAnywhere, Category = "Print");
 	bool PrintToLog = true;
 
-	UPROPERTY(EditAnywhere, Category = "Print", meta = (DataflowInput, DisplayName = "String"));
+	UPROPERTY(EditAnywhere, Category = "Print", meta = (DataflowInput));
 	FString String = FString("");
 
 	FLogStringDataflowNode(const Dataflow::FNodeParameters& InParam, FGuid InGuid = FGuid::NewGuid())
@@ -204,7 +208,7 @@ public:
 	UPROPERTY(EditAnywhere, Category = "String");
 	FString Value = FString("");
 
-	UPROPERTY(meta = (DataflowOutput, DisplayName = "String"))
+	UPROPERTY(meta = (DataflowOutput))
 	FString String;
 
 	FMakeLiteralStringDataflowNode(const Dataflow::FNodeParameters& InParam, FGuid InGuid = FGuid::NewGuid())
@@ -226,10 +230,10 @@ struct FBoundingBoxDataflowNode : public FDataflowNode
 public:
 	typedef FManagedArrayCollection DataType;
 
-	UPROPERTY(meta = (DataflowInput, DisplayName = "Collection"))
+	UPROPERTY(meta = (DataflowInput))
 	FManagedArrayCollection Collection;
 
-	UPROPERTY(meta = (DataflowOutput, DisplayName = "BoundingBox"))
+	UPROPERTY(meta = (DataflowOutput))
 	FBox BoundingBox = FBox(ForceInit);
 
 	FBoundingBoxDataflowNode(const Dataflow::FNodeParameters& InParam, FGuid InGuid = FGuid::NewGuid())
@@ -250,22 +254,22 @@ struct FExpandBoundingBoxDataflowNode : public FDataflowNode
 	DATAFLOW_NODE_DEFINE_INTERNAL(FExpandBoundingBoxDataflowNode, "ExpandBoundingBox", "Utilities|Box", "")
 
 public:
-	UPROPERTY(meta = (DataflowInput, DisplayName = "BoundingBox"))
+	UPROPERTY(meta = (DataflowInput))
 	FBox BoundingBox = FBox(ForceInit);;
 
-	UPROPERTY(meta = (DataflowOutput, DisplayName = "Min"))
+	UPROPERTY(meta = (DataflowOutput))
 	FVector Min = FVector(0.0);
 
-	UPROPERTY(meta = (DataflowOutput, DisplayName = "Max"))
+	UPROPERTY(meta = (DataflowOutput))
 	FVector Max = FVector(0.0);
 
-	UPROPERTY(meta = (DataflowOutput, DisplayName = "Center"))
+	UPROPERTY(meta = (DataflowOutput))
 	FVector Center = FVector(0.0);
 
-	UPROPERTY(meta = (DataflowOutput, DisplayName = "HalfExtents"))
+	UPROPERTY(meta = (DataflowOutput))
 	FVector HalfExtents = FVector(0.0);
 
-	UPROPERTY(meta = (DataflowOutput, DisplayName = "Volume"))
+	UPROPERTY(meta = (DataflowOutput))
 	float Volume = 0.0;
 
 	FExpandBoundingBoxDataflowNode(const Dataflow::FNodeParameters& InParam, FGuid InGuid = FGuid::NewGuid())
@@ -290,10 +294,10 @@ struct FVectorToStringDataflowNode : public FDataflowNode
 	DATAFLOW_NODE_DEFINE_INTERNAL(FVectorToStringDataflowNode, "VectorToString", "Utilities|String", "")
 
 public:
-	UPROPERTY(EditAnywhere, Category = "Vector", meta = (DataflowInput, DisplayName = "Vector"))
+	UPROPERTY(EditAnywhere, Category = "Vector", meta = (DataflowInput))
 	FVector Vector = FVector(0.0);
 
-	UPROPERTY(meta = (DataflowOutput, DisplayName = "String"))
+	UPROPERTY(meta = (DataflowOutput))
 	FString String = FString("");
 
 	FVectorToStringDataflowNode(const Dataflow::FNodeParameters& InParam, FGuid InGuid = FGuid::NewGuid())
@@ -314,10 +318,10 @@ struct FFloatToStringDataflowNode : public FDataflowNode
 	DATAFLOW_NODE_DEFINE_INTERNAL(FFloatToStringDataflowNode, "FloatToString", "Utilities|String", "")
 
 public:
-	UPROPERTY(EditAnywhere, Category = "Float", meta = (DataflowInput, DisplayName = "Float"))
+	UPROPERTY(EditAnywhere, Category = "Float", meta = (DataflowInput))
 	float Float = 0.f;
 
-	UPROPERTY(meta = (DataflowOutput, DisplayName = "String"))
+	UPROPERTY(meta = (DataflowOutput))
 	FString String = FString("");
 
 	FFloatToStringDataflowNode(const Dataflow::FNodeParameters& InParam, FGuid InGuid = FGuid::NewGuid())
@@ -341,7 +345,7 @@ public:
 	UPROPERTY(EditAnywhere, Category = "Point")
 	TArray<FVector> Point;
 
-	UPROPERTY(meta = (DataflowOutput, DisplayName = "Points"))
+	UPROPERTY(meta = (DataflowOutput))
 	TArray<FVector> Points;
 
 	FMakePointsDataflowNode(const Dataflow::FNodeParameters& InParam, FGuid InGuid = FGuid::NewGuid())
@@ -374,19 +378,19 @@ public:
 	UPROPERTY(EditAnywhere, Category = "Box", meta = (DisplayName = "Input Data Type"));
 	EMakeBoxDataTypeEnum DataType = EMakeBoxDataTypeEnum::Dataflow_MakeBox_DataType_MinMax;
 
-	UPROPERTY(EditAnywhere, Category = "Box", meta = (DataflowInput, DisplayName = "Min", EditCondition = "DataType == EMakeBoxDataTypeEnum::Dataflow_MakeBox_DataType_MinMax", EditConditionHides));
+	UPROPERTY(EditAnywhere, Category = "Box", meta = (DataflowInput, EditCondition = "DataType == EMakeBoxDataTypeEnum::Dataflow_MakeBox_DataType_MinMax", EditConditionHides));
 	FVector Min = FVector(0.0);
 
-	UPROPERTY(EditAnywhere, Category = "Box", meta = (DataflowInput, DisplayName = "Max", EditCondition = "DataType == EMakeBoxDataTypeEnum::Dataflow_MakeBox_DataType_MinMax", EditConditionHides));
+	UPROPERTY(EditAnywhere, Category = "Box", meta = (DataflowInput, EditCondition = "DataType == EMakeBoxDataTypeEnum::Dataflow_MakeBox_DataType_MinMax", EditConditionHides));
 	FVector Max = FVector(0.0);
 
-	UPROPERTY(EditAnywhere, Category = "Box", meta = (DataflowInput, DisplayName = "Center", EditCondition = "DataType == EMakeBoxDataTypeEnum::Dataflow_MakeBox_DataType_CenterSize", EditConditionHides));
+	UPROPERTY(EditAnywhere, Category = "Box", meta = (DataflowInput, EditCondition = "DataType == EMakeBoxDataTypeEnum::Dataflow_MakeBox_DataType_CenterSize", EditConditionHides));
 	FVector Center = FVector(0.0);
 
-	UPROPERTY(EditAnywhere, Category = "Box", meta = (DataflowInput, DisplayName = "Size", EditCondition = "DataType == EMakeBoxDataTypeEnum::Dataflow_MakeBox_DataType_CenterSize", EditConditionHides));
+	UPROPERTY(EditAnywhere, Category = "Box", meta = (DataflowInput, EditCondition = "DataType == EMakeBoxDataTypeEnum::Dataflow_MakeBox_DataType_CenterSize", EditConditionHides));
 	FVector Size = FVector(0.0);
 
-	UPROPERTY(meta = (DataflowOutput, DisplayName = "Box"));
+	UPROPERTY(meta = (DataflowOutput));
 	FBox Box = FBox(ForceInit);
 
 	FMakeBoxDataflowNode(const Dataflow::FNodeParameters& InParam, FGuid InGuid = FGuid::NewGuid())
@@ -410,28 +414,28 @@ struct FUniformScatterPointsDataflowNode : public FDataflowNode
 		DATAFLOW_NODE_DEFINE_INTERNAL(FUniformScatterPointsDataflowNode, "UniformScatterPoints", "Generators|Point", "")
 
 public:
-	UPROPERTY(EditAnywhere, Category = "Scatter", meta = (DataflowInput, DisplayName = "Min Number Of Points"));
+	UPROPERTY(EditAnywhere, Category = "Scatter", meta = (DataflowInput));
 	int32 MinNumberOfPoints = 20;
 
-	UPROPERTY(EditAnywhere, Category = "Scatter", meta = (DataflowInput, DisplayName = "Max Number Of Points"));
+	UPROPERTY(EditAnywhere, Category = "Scatter", meta = (DataflowInput));
 	int32 MaxNumberOfPoints = 20;
 
-	UPROPERTY(EditAnywhere, Category = "Scatter", meta = (DataflowInput, DisplayName = "Random Seed"));
+	UPROPERTY(EditAnywhere, Category = "Scatter", meta = (DataflowInput));
 	float RandomSeed = -1.f;
 
-	UPROPERTY(meta = (DataflowInput, DisplayName = "BoundingBox"))
-		FBox BoundingBox = FBox(ForceInit);
+	UPROPERTY(meta = (DataflowInput))
+	FBox BoundingBox = FBox(ForceInit);
 
-	UPROPERTY(meta = (DataflowOutput, DisplayName = "Points"))
-		TArray<FVector> Points;
+	UPROPERTY(meta = (DataflowOutput))
+	TArray<FVector> Points;
 
 	FUniformScatterPointsDataflowNode(const Dataflow::FNodeParameters& InParam, FGuid InGuid = FGuid::NewGuid())
 		: FDataflowNode(InParam, InGuid)
 	{
+		RegisterInputConnection(&BoundingBox);
 		RegisterInputConnection(&MinNumberOfPoints);
 		RegisterInputConnection(&MaxNumberOfPoints);
 		RegisterInputConnection(&RandomSeed);
-		RegisterInputConnection(&BoundingBox);
 		RegisterOutputConnection(&Points);
 	}
 
@@ -443,35 +447,35 @@ USTRUCT()
 struct FRadialScatterPointsDataflowNode : public FDataflowNode
 {
 	GENERATED_USTRUCT_BODY()
-		DATAFLOW_NODE_DEFINE_INTERNAL(FRadialScatterPointsDataflowNode, "RadialScatterPoints", "Generators|Point", "")
+	DATAFLOW_NODE_DEFINE_INTERNAL(FRadialScatterPointsDataflowNode, "RadialScatterPoints", "Generators|Point", "")
 
 public:
-	UPROPERTY(EditAnywhere, Category = "Scatter", meta = (DataflowInput, DisplayName = "Center"));
+	UPROPERTY(EditAnywhere, Category = "Scatter", meta = (DataflowInput));
 	FVector Center = FVector(0.0);
 
-	UPROPERTY(EditAnywhere, Category = "Scatter", meta = (DataflowInput, DisplayName = "Normal"));
+	UPROPERTY(EditAnywhere, Category = "Scatter", meta = (DataflowInput));
 	FVector Normal = FVector(0.0, 0.0, 1.0);
 
-	UPROPERTY(EditAnywhere, Category = "Scatter", meta = (DataflowInput, DisplayName = "Radius", UIMin = 0.01f));
+	UPROPERTY(EditAnywhere, Category = "Scatter", meta = (DataflowInput, UIMin = 0.01f));
 	float Radius = 50.f;
 
-	UPROPERTY(EditAnywhere, Category = "Scatter", meta = (DataflowInput, DisplayName = "Angular Steps", UIMin = 1, UIMax = 50));
+	UPROPERTY(EditAnywhere, Category = "Scatter", meta = (DataflowInput, UIMin = 1, UIMax = 50));
 	int32 AngularSteps = 5;
 
-	UPROPERTY(EditAnywhere, Category = "Scatter", meta = (DataflowInput, DisplayName = "Radial Steps", UIMin = 1, UIMax = 50));
+	UPROPERTY(EditAnywhere, Category = "Scatter", meta = (DataflowInput, UIMin = 1, UIMax = 50));
 	int32 RadialSteps = 5;
 
-	UPROPERTY(EditAnywhere, Category = "Scatter", meta = (DataflowInput, DisplayName = "Angle Offset"));
+	UPROPERTY(EditAnywhere, Category = "Scatter", meta = (DataflowInput));
 	float AngleOffset = 0.f;
 
-	UPROPERTY(EditAnywhere, Category = "Scatter", meta = (DataflowInput, DisplayName = "Variability", UIMin = 0.f));
+	UPROPERTY(EditAnywhere, Category = "Scatter", meta = (DataflowInput, UIMin = 0.f));
 	float Variability = 0.f;
 
-	UPROPERTY(EditAnywhere, Category = "Scatter", meta = (DataflowInput, DisplayName = "Random Seed"));
+	UPROPERTY(EditAnywhere, Category = "Scatter", meta = (DataflowInput));
 	float RandomSeed = -1.f;
 
-	UPROPERTY(meta = (DataflowOutput, DisplayName = "Points"))
-		TArray<FVector> Points;
+	UPROPERTY(meta = (DataflowOutput))
+	TArray<FVector> Points;
 
 	FRadialScatterPointsDataflowNode(const Dataflow::FNodeParameters& InParam, FGuid InGuid = FGuid::NewGuid())
 		: FDataflowNode(InParam, InGuid)
@@ -508,7 +512,7 @@ public:
 	UPROPERTY(EditAnywhere, Category = "Float");
 	float Value = 0.f;
 
-	UPROPERTY(meta = (DataflowOutput, DisplayName = "Float"))
+	UPROPERTY(meta = (DataflowOutput))
 	float Float = 0.f;
 
 	FMakeLiteralFloatDataflowNode(const Dataflow::FNodeParameters& InParam, FGuid InGuid = FGuid::NewGuid())
@@ -554,7 +558,7 @@ public:
 	UPROPERTY(EditAnywhere, Category = "Bool");
 	bool Value = false;
 
-	UPROPERTY(meta = (DataflowOutput, DisplayName = "Bool"))
+	UPROPERTY(meta = (DataflowOutput))
 	bool Bool = false;
 
 	FMakeLiteralBoolDataflowNode(const Dataflow::FNodeParameters& InParam, FGuid InGuid = FGuid::NewGuid())
@@ -567,6 +571,12 @@ public:
 
 };
 
+
+/**
+ *
+ * Converts an integer to a float
+ *
+ */
 USTRUCT()
 struct FMakeLiteralVectorDataflowNode : public FDataflowNode
 {
@@ -590,6 +600,12 @@ public:
 
 };
 
+
+/**
+ *
+ * Converts an Int to a String
+ *
+ */
 USTRUCT()
 struct FIntToStringDataflowNode : public FDataflowNode
 {
@@ -597,10 +613,10 @@ struct FIntToStringDataflowNode : public FDataflowNode
 		DATAFLOW_NODE_DEFINE_INTERNAL(FIntToStringDataflowNode, "IntToString", "Utilities|String", "")
 
 public:
-	UPROPERTY(EditAnywhere, Category = "Int", meta = (DataflowInput, DisplayName = "Int"))
+	UPROPERTY(EditAnywhere, Category = "Int", meta = (DataflowInput))
 	int32 Int = 0;
 
-	UPROPERTY(meta = (DataflowOutput, DisplayName = "String"))
+	UPROPERTY(meta = (DataflowOutput))
 	FString String = FString("");
 
 	FIntToStringDataflowNode(const Dataflow::FNodeParameters& InParam, FGuid InGuid = FGuid::NewGuid())
@@ -614,6 +630,12 @@ public:
 
 };
 
+
+/**
+ *
+ * Converts a Bool to a String in a form of ("true", "false")
+ *
+ */
 USTRUCT()
 struct FBoolToStringDataflowNode : public FDataflowNode
 {
@@ -621,10 +643,10 @@ struct FBoolToStringDataflowNode : public FDataflowNode
 	DATAFLOW_NODE_DEFINE_INTERNAL(FBoolToStringDataflowNode, "BoolToString", "Utilities|String", "")
 
 public:
-	UPROPERTY(EditAnywhere, Category = "Bool", meta = (DataflowInput, DisplayName = "Bool"))
+	UPROPERTY(EditAnywhere, Category = "Bool", meta = (DataflowInput))
 	bool Bool = false;
 
-	UPROPERTY(meta = (DataflowOutput, DisplayName = "String"))
+	UPROPERTY(meta = (DataflowOutput))
 	FString String = FString("");
 
 	FBoolToStringDataflowNode(const Dataflow::FNodeParameters& InParam, FGuid InGuid = FGuid::NewGuid())
@@ -638,6 +660,12 @@ public:
 
 };
 
+
+/**
+ *
+ * Expands a Vector into X, Y, Z components
+ *
+ */
 USTRUCT()
 struct FExpandVectorDataflowNode : public FDataflowNode
 {
@@ -645,16 +673,16 @@ struct FExpandVectorDataflowNode : public FDataflowNode
 	DATAFLOW_NODE_DEFINE_INTERNAL(FExpandVectorDataflowNode, "ExpandVector", "Utilities|Vector", "")
 
 public:
-	UPROPERTY(meta = (DataflowInput, DisplayName = "Vector"))
+	UPROPERTY(meta = (DataflowInput))
 	FVector Vector = FVector(0.0);
 
-	UPROPERTY(meta = (DataflowOutput, DisplayName = "X"))
+	UPROPERTY(meta = (DataflowOutput))
 	float X = 0.f;
 
-	UPROPERTY(meta = (DataflowOutput, DisplayName = "Y"))
+	UPROPERTY(meta = (DataflowOutput))
 	float Y = 0.f;
 
-	UPROPERTY(meta = (DataflowOutput, DisplayName = "Z"))
+	UPROPERTY(meta = (DataflowOutput))
 	float Z = 0.f;
 
 
@@ -671,6 +699,12 @@ public:
 
 };
 
+
+/**
+ *
+ * Converts an Int to a Float
+ *
+ */
 USTRUCT()
 struct FIntToFloatDataflowNode : public FDataflowNode
 {
@@ -678,10 +712,10 @@ struct FIntToFloatDataflowNode : public FDataflowNode
 	DATAFLOW_NODE_DEFINE_INTERNAL(FIntToFloatDataflowNode, "IntToFloat", "Math|Conversions", "")
 
 public:
-	UPROPERTY(EditAnywhere, Category = "Int", meta = (DataflowInput, DisplayName = "Int"))
+	UPROPERTY(EditAnywhere, Category = "Int", meta = (DataflowInput))
 	int32 Int = 0;
 
-	UPROPERTY(meta = (DataflowOutput, DisplayName = "Float"))
+	UPROPERTY(meta = (DataflowOutput))
 	float Float = 0.f;
 
 	FIntToFloatDataflowNode(const Dataflow::FNodeParameters& InParam, FGuid InGuid = FGuid::NewGuid())
@@ -695,64 +729,12 @@ public:
 
 };
 
-USTRUCT()
-struct FClusterFlattenDataflowNode : public FDataflowNode
-{
-	GENERATED_USTRUCT_BODY()
-	DATAFLOW_NODE_DEFINE_INTERNAL(FClusterFlattenDataflowNode, "Flatten", "Fracture|Cluster", "")
-
-public:
-	// @todo(harsha) Support Selections
-
-	UPROPERTY(meta = (DataflowInput, DataflowOutput, DisplayName = "Collection"))
-	FManagedArrayCollection Collection;
-
-	FClusterFlattenDataflowNode(const Dataflow::FNodeParameters& InParam, FGuid InGuid = FGuid::NewGuid())
-		: FDataflowNode(InParam, InGuid)
-	{
-		RegisterInputConnection(&Collection);
-		RegisterOutputConnection(&Collection);
-	}
-
-	virtual void Evaluate(Dataflow::FContext& Context, const FDataflowOutput* Out) const override;
-
-};
-
-USTRUCT()
-struct FRemoveOnBreakDataflowNode : public FDataflowNode
-{
-	GENERATED_USTRUCT_BODY()
-	DATAFLOW_NODE_DEFINE_INTERNAL(FRemoveOnBreakDataflowNode, "RemoveOnBreak", "Fracture|Utilities", "")
-
-public:
-	// @todo(harsha) Support Selections
-
-	UPROPERTY(meta = (DataflowInput, DataflowOutput, DisplayName = "Collection"))
-	FManagedArrayCollection Collection;
-
-	UPROPERTY(EditAnywhere, Category = "Removal", meta = (DataflowInput, DisplayName = "PostBreakTimer"))
-	FVector2f PostBreakTimer{0.0, 0.0};
-
-	UPROPERTY(EditAnywhere, Category = "Removal", meta = (DataflowInput, DisplayName = "RemovalTimer"))
-	FVector2f RemovalTimer{0.0, 1.0};
-
-	UPROPERTY(EditAnywhere, Category = "Removal", meta = (DataflowInput, DisplayName = "ClusterCrumbling"))
-	bool ClusterCrumbling = false;
-
-	FRemoveOnBreakDataflowNode(const Dataflow::FNodeParameters& InParam, FGuid InGuid = FGuid::NewGuid())
-		: FDataflowNode(InParam, InGuid)
-	{
-		RegisterInputConnection(&Collection);
-		RegisterOutputConnection(&Collection);
-		RegisterInputConnection(&PostBreakTimer);
-		RegisterInputConnection(&RemovalTimer);
-		RegisterInputConnection(&ClusterCrumbling);
-	}
-
-	virtual void Evaluate(Dataflow::FContext& Context, const FDataflowOutput* Out) const override;
-
-};
 
+/**
+ *
+ * Generates a Voronoi fracture
+ *
+ */
 USTRUCT()
 struct FVoronoiFractureDataflowNode : public FDataflowNode
 {
@@ -763,43 +745,43 @@ public:
 	UPROPERTY(meta = (DataflowInput, DataflowOutput, DisplayName = "Collection"))
 	FManagedArrayCollection Collection;
 
-	UPROPERTY(meta = (DataflowInput, DisplayName = "Points"))
+	UPROPERTY(meta = (DataflowInput))
 	TArray<FVector> Points;
 
-	UPROPERTY(EditAnywhere, Category = "Fracture", meta = (DataflowInput, DisplayName = "RandomSeed"));
+	UPROPERTY(EditAnywhere, Category = "Fracture", meta = (DataflowInput));
 	float RandomSeed = -1.f;
 
-	UPROPERTY(EditAnywhere, Category = "Fracture", meta = (DataflowInput, DisplayName = "ChanceToFracture", UIMin = 0.f, UIMax = 1.f));
+	UPROPERTY(EditAnywhere, Category = "Fracture", meta = (DataflowInput, UIMin = 0.f, UIMax = 1.f));
 	float ChanceToFracture = 1.f;
 
 	UPROPERTY(EditAnywhere, Category = "Fracture");
 	bool GroupFracture = true;
 
-	UPROPERTY(EditAnywhere, Category = "Fracture", meta = (DataflowInput, DisplayName = "Grout", UIMin = 0.f));
+	UPROPERTY(EditAnywhere, Category = "Fracture", meta = (DataflowInput, UIMin = 0.f));
 	float Grout = 0.f;
 
-	UPROPERTY(EditAnywhere, Category = "Noise", meta = (DataflowInput, DisplayName = "Amplitude", UIMin = 0.f));
+	UPROPERTY(EditAnywhere, Category = "Noise", meta = (DataflowInput, UIMin = 0.f));
 	float Amplitude = 0.f;
 
-	UPROPERTY(EditAnywhere, Category = "Noise", meta = (DataflowInput, DisplayName = "Frequency", UIMin = 0.00001f));
+	UPROPERTY(EditAnywhere, Category = "Noise", meta = (DataflowInput, UIMin = 0.00001f));
 	float Frequency = 0.1f;
 
-	UPROPERTY(EditAnywhere, Category = "Noise", meta = (DataflowInput, DisplayName = "Persistence", UIMin = 0.f));
+	UPROPERTY(EditAnywhere, Category = "Noise", meta = (DataflowInput, UIMin = 0.f));
 	float Persistence = 0.5f;
 
-	UPROPERTY(EditAnywhere, Category = "Noise", meta = (DataflowInput, DisplayName = "Lacunarity", UIMin = 0.f));
+	UPROPERTY(EditAnywhere, Category = "Noise", meta = (DataflowInput, UIMin = 0.f));
 	float Lacunarity = 2.f;
 
-	UPROPERTY(EditAnywhere, Category = "Noise", meta = (DataflowInput, DisplayName = "OctaveNumber", UIMin = 0.f));
+	UPROPERTY(EditAnywhere, Category = "Noise", meta = (DataflowInput, UIMin = 0.f));
 	int32 OctaveNumber = 4;
 
-	UPROPERTY(EditAnywhere, Category = "Noise", meta = (DataflowInput, DisplayName = "PointSpacing", UIMin = 0.f));
+	UPROPERTY(EditAnywhere, Category = "Noise", meta = (DataflowInput, UIMin = 0.f));
 	float PointSpacing = 10.f;
 
 	UPROPERTY(EditAnywhere, Category = "Collision");
 	bool AddSamplesForCollision = false;
 
-	UPROPERTY(EditAnywhere, Category = "Collision", meta = (DataflowInput, DisplayName = "CollisionSampleSpacing", UIMin = 0.f));
+	UPROPERTY(EditAnywhere, Category = "Collision", meta = (DataflowInput, UIMin = 0.f));
 	float CollisionSampleSpacing = 50.f;
 
 	FVoronoiFractureDataflowNode(const Dataflow::FNodeParameters& InParam, FGuid InGuid = FGuid::NewGuid())
@@ -824,6 +806,12 @@ public:
 
 };
 
+
+/**
+ *
+ * Concatenates two strings together to make a new string
+ *
+ */
 USTRUCT()
 struct FStringAppendDataflowNode : public FDataflowNode
 {
@@ -831,13 +819,13 @@ struct FStringAppendDataflowNode : public FDataflowNode
 	DATAFLOW_NODE_DEFINE_INTERNAL(FStringAppendDataflowNode, "StringAppend", "Utilities|String", "")
 
 public:
-	UPROPERTY(EditAnywhere, Category = "String", meta = (DataflowInput, DisplayName = "String1"))
+	UPROPERTY(EditAnywhere, Category = "String", meta = (DataflowInput))
 	FString String1 = FString("");
 
-	UPROPERTY(EditAnywhere, Category = "String", meta = (DataflowInput, DisplayName = "String2"))
+	UPROPERTY(EditAnywhere, Category = "String", meta = (DataflowInput))
 	FString String2 = FString("");
 
-	UPROPERTY(meta = (DataflowOutput, DisplayName = "String"))
+	UPROPERTY(meta = (DataflowOutput))
 	FString String = FString("");
 
 	FStringAppendDataflowNode(const Dataflow::FNodeParameters& InParam, FGuid InGuid = FGuid::NewGuid())
@@ -852,6 +840,12 @@ public:
 
 };
 
+
+/**
+ *
+ * Generates a random float
+ *
+ */
 USTRUCT()
 struct FRandomFloatDataflowNode : public FDataflowNode
 {
@@ -862,10 +856,10 @@ public:
 	UPROPERTY(EditAnywhere, Category = "Seed")
 	bool Deterministic = false;
 
-	UPROPERTY(EditAnywhere, Category = "Seed", meta = (DataflowInput, DisplayName = "RandomSeed", EditCondition = "Deterministic"))
+	UPROPERTY(EditAnywhere, Category = "Seed", meta = (DataflowInput, EditCondition = "Deterministic"))
 	float RandomSeed = 0.f;
 
-	UPROPERTY(meta = (DataflowOutput, DisplayName = "Float"))
+	UPROPERTY(meta = (DataflowOutput))
 	float Float = 0.f;
 
 	FRandomFloatDataflowNode(const Dataflow::FNodeParameters& InParam, FGuid InGuid = FGuid::NewGuid())
@@ -880,6 +874,11 @@ public:
 
 };
 
+/**
+ *
+ * Generates a random float between Min and Max
+ *
+ */
 USTRUCT()
 struct FRandomFloatInRangeDataflowNode : public FDataflowNode
 {
@@ -890,16 +889,16 @@ public:
 	UPROPERTY(EditAnywhere, Category = "Seed")
 	bool Deterministic = false;
 
-	UPROPERTY(EditAnywhere, Category = "Seed", meta = (DataflowInput, DisplayName = "RandomSeed", EditCondition = "Deterministic"))
+	UPROPERTY(EditAnywhere, Category = "Seed", meta = (DataflowInput, EditCondition = "Deterministic"))
 	float RandomSeed = 0.f;
 	
-	UPROPERTY(EditAnywhere, Category = "Random", meta = (DataflowInput, DisplayName = "Min"))
+	UPROPERTY(EditAnywhere, Category = "Random", meta = (DataflowInput))
 	float Min = 0.f;
 
-	UPROPERTY(EditAnywhere, Category = "Random", meta = (DataflowInput, DisplayName = "Max"))
+	UPROPERTY(EditAnywhere, Category = "Random", meta = (DataflowInput))
 	float Max = 1.f;
 
-	UPROPERTY(meta = (DataflowOutput, DisplayName = "Float"))
+	UPROPERTY(meta = (DataflowOutput))
 	float Float = 0.f;
 
 	FRandomFloatInRangeDataflowNode(const Dataflow::FNodeParameters& InParam, FGuid InGuid = FGuid::NewGuid())
@@ -916,6 +915,12 @@ public:
 
 };
 
+
+/**
+ *
+ * Returns a random vector with length of 1
+ *
+ */
 USTRUCT()
 struct FRandomUnitVectorDataflowNode : public FDataflowNode
 {
@@ -926,10 +931,10 @@ public:
 	UPROPERTY(EditAnywhere, Category = "Seed")
 	bool Deterministic = false;
 
-	UPROPERTY(EditAnywhere, Category = "Seed", meta = (DataflowInput, DisplayName = "RandomSeed", EditCondition = "Deterministic"))
+	UPROPERTY(EditAnywhere, Category = "Seed", meta = (DataflowInput, EditCondition = "Deterministic"))
 	float RandomSeed = 0.f;
 	
-	UPROPERTY(meta = (DataflowOutput, DisplayName = "Vector"))
+	UPROPERTY(meta = (DataflowOutput))
 	FVector Vector = FVector(0.0);
 
 	FRandomUnitVectorDataflowNode(const Dataflow::FNodeParameters& InParam, FGuid InGuid = FGuid::NewGuid())
@@ -944,6 +949,12 @@ public:
 
 };
 
+
+/**
+ *
+ * Returns a random vector with length of 1, within the specified cone, with uniform random distribution 
+ *
+ */
 USTRUCT()
 struct FRandomUnitVectorInConeDataflowNode : public FDataflowNode
 {
@@ -954,16 +965,18 @@ public:
 	UPROPERTY(EditAnywhere, Category = "Seed")
 	bool Deterministic = false;
 
-	UPROPERTY(EditAnywhere, Category = "Seed", meta = (DataflowInput, DisplayName = "RandomSeed", EditCondition = "Deterministic"))
+	UPROPERTY(EditAnywhere, Category = "Seed", meta = (DataflowInput, EditCondition = "Deterministic"))
 	float RandomSeed = 0.f;
 
-	UPROPERTY(EditAnywhere, Category = "Random", meta = (DataflowInput, DisplayName = "ConeDirection"))
+	/** The base "center" direction of the cone */
+	UPROPERTY(EditAnywhere, Category = "Random", meta = (DataflowInput))
 	FVector ConeDirection = FVector(0.0, 0.0, 1.0);
 
-	UPROPERTY(EditAnywhere, Category = "Random", meta = (DataflowInput, DisplayName = "ConeHalfAngle"))
+	/** The half-angle of the cone (from ConeDir to edge), in degrees */
+	UPROPERTY(EditAnywhere, Category = "Random", meta = (DataflowInput))
 	float ConeHalfAngle = PI / 4.f;
 
-	UPROPERTY(meta = (DataflowOutput, DisplayName = "Vector"))
+	UPROPERTY(meta = (DataflowOutput))
 	FVector Vector = FVector(0.0);
 
 	FRandomUnitVectorInConeDataflowNode(const Dataflow::FNodeParameters& InParam, FGuid InGuid = FGuid::NewGuid())
@@ -980,6 +993,12 @@ public:
 
 };
 
+
+/**
+ *
+ * Converts radians to degrees
+ *
+ */
 USTRUCT()
 struct FRadiansToDegreesDataflowNode : public FDataflowNode
 {
@@ -987,10 +1006,10 @@ struct FRadiansToDegreesDataflowNode : public FDataflowNode
 	DATAFLOW_NODE_DEFINE_INTERNAL(FRadiansToDegreesDataflowNode, "RadiansToDegrees", "Math|Trigonometry", "")
 
 public:
-	UPROPERTY(EditAnywhere, Category = "Radians", meta = (DataflowInput, DisplayName = "Radians"))
+	UPROPERTY(EditAnywhere, Category = "Radians", meta = (DataflowInput))
 	float Radians = 0.f;
 
-	UPROPERTY(meta = (DataflowOutput, DisplayName = "Degrees"))
+	UPROPERTY(meta = (DataflowOutput))
 	float Degrees = 0.f;
 
 	FRadiansToDegreesDataflowNode(const Dataflow::FNodeParameters& InParam, FGuid InGuid = FGuid::NewGuid())
@@ -1004,6 +1023,12 @@ public:
 
 };
 
+
+/**
+ *
+ * Converts degrees to radians
+ *
+ */
 USTRUCT()
 struct FDegreesToRadiansDataflowNode : public FDataflowNode
 {
@@ -1011,10 +1036,10 @@ struct FDegreesToRadiansDataflowNode : public FDataflowNode
 	DATAFLOW_NODE_DEFINE_INTERNAL(FDegreesToRadiansDataflowNode, "DegreesToRadians", "Math|Trigonometry", "")
 
 public:
-	UPROPERTY(EditAnywhere, Category = "Degrees", meta = (DataflowInput, DisplayName = "Degrees"))
+	UPROPERTY(EditAnywhere, Category = "Degrees", meta = (DataflowInput))
 	float Degrees = 0.f;
 
-	UPROPERTY(meta = (DataflowOutput, DisplayName = "Radians"))
+	UPROPERTY(meta = (DataflowOutput))
 	float Radians = 0.f;
 
 	FDegreesToRadiansDataflowNode(const Dataflow::FNodeParameters& InParam, FGuid InGuid = FGuid::NewGuid())
@@ -1028,6 +1053,12 @@ public:
 
 };
 
+
+/**
+ *
+ * "Explodes" the pieces from the Collection for better visualization
+ *
+ */
 USTRUCT()
 struct FExplodedViewDataflowNode : public FDataflowNode
 {
@@ -1035,13 +1066,13 @@ struct FExplodedViewDataflowNode : public FDataflowNode
 	DATAFLOW_NODE_DEFINE_INTERNAL(FExplodedViewDataflowNode, "ExplodedView", "Fracture|Utilities", "")
 
 public:
-	UPROPERTY(meta = (DataflowInput, DataflowOutput, DisplayName = "Collection"))
+	UPROPERTY(meta = (DataflowInput, DataflowOutput))
 	FManagedArrayCollection Collection;
 
-	UPROPERTY(EditAnywhere, Category = "Scale", meta = (DataflowInput, DisplayName = "UniformScale"))
+	UPROPERTY(EditAnywhere, Category = "Scale", meta = (DataflowInput))
 	float UniformScale = 1.f;
 
-	UPROPERTY(EditAnywhere, Category = "Scale", meta = (DataflowInput, DisplayName = "Scale"))
+	UPROPERTY(EditAnywhere, Category = "Scale", meta = (DataflowInput))
 	FVector Scale = FVector(1.0);
 
 	FExplodedViewDataflowNode(const Dataflow::FNodeParameters& InParam, FGuid InGuid = FGuid::NewGuid())
@@ -1061,6 +1092,12 @@ private:
 
 };
 
+
+/**
+ *
+ * Generates convex hull representation for the bones for simulation
+ *
+ */
 USTRUCT()
 struct FCreateNonOverlappingConvexHullsDataflowNode : public FDataflowNode
 {
@@ -1068,16 +1105,16 @@ struct FCreateNonOverlappingConvexHullsDataflowNode : public FDataflowNode
 	DATAFLOW_NODE_DEFINE_INTERNAL(FCreateNonOverlappingConvexHullsDataflowNode, "CreateNonOverlappingConvexHulls", "Fracture|Utilities", "")
 
 public:
-	UPROPERTY(meta = (DataflowInput, DataflowOutput, DisplayName = "Collection"))
+	UPROPERTY(meta = (DataflowInput, DataflowOutput))
 	FManagedArrayCollection Collection;
 
-	UPROPERTY(EditAnywhere, Category = "Convex", meta = (DataflowInput, DisplayName = "CanRemoveFraction", UIMin = 0.01f, UIMax = 1.f))
+	UPROPERTY(EditAnywhere, Category = "Convex", meta = (DataflowInput, UIMin = 0.01f, UIMax = 1.f))
 	float CanRemoveFraction = 0.5f;
 
-	UPROPERTY(EditAnywhere, Category = "Convex", meta = (DataflowInput, DisplayName = "CanExceedFraction", UIMin = 0.f))
+	UPROPERTY(EditAnywhere, Category = "Convex", meta = (DataflowInput, UIMin = 0.f))
 	float CanExceedFraction = 0.5f;
 
-	UPROPERTY(EditAnywhere, Category = "Convex", meta = (DataflowInput, DisplayName = "SimplificationDistanceThreshold", UIMin = 0.f))
+	UPROPERTY(EditAnywhere, Category = "Convex", meta = (DataflowInput, UIMin = 0.f))
 	float SimplificationDistanceThreshold = 10.f;
 
 	FCreateNonOverlappingConvexHullsDataflowNode(const Dataflow::FNodeParameters& InParam, FGuid InGuid = FGuid::NewGuid())
@@ -1094,6 +1131,12 @@ public:
 
 };
 
+
+/**
+ *
+ * Cuts geometry using a set of noised up planes
+ *
+ */
 USTRUCT()
 struct FPlaneCutterDataflowNode : public FDataflowNode
 {
@@ -1101,43 +1144,43 @@ struct FPlaneCutterDataflowNode : public FDataflowNode
 	DATAFLOW_NODE_DEFINE_INTERNAL(FPlaneCutterDataflowNode, "PlaneCutter", "Fracture", "")
 
 public:
-	UPROPERTY(meta = (DataflowInput, DataflowOutput, DisplayName = "Collection"))
+	UPROPERTY(meta = (DataflowInput, DataflowOutput))
 	FManagedArrayCollection Collection;
 
-	UPROPERTY(meta = (DataflowInput, DisplayName = "BoundingBox"))
+	UPROPERTY(meta = (DataflowInput))
 	FBox BoundingBox = FBox(ForceInit);
 
-	UPROPERTY(EditAnywhere, Category = "Fracture", meta = (DataflowInput, DisplayName = "NumPlanes", UIMin = 1))
+	UPROPERTY(EditAnywhere, Category = "Fracture", meta = (DataflowInput, UIMin = 1))
 	int32 NumPlanes = 1;
 
-	UPROPERTY(EditAnywhere, Category = "Fracture", meta = (DataflowInput, DisplayName = "RandomSeed"))
+	UPROPERTY(EditAnywhere, Category = "Fracture", meta = (DataflowInput))
 	float RandomSeed = -1.f;
 
-	UPROPERTY(EditAnywhere, Category = "Fracture", meta = (DataflowInput, DisplayName = "Grout", UIMin = 0.f))
+	UPROPERTY(EditAnywhere, Category = "Fracture", meta = (DataflowInput, UIMin = 0.f))
 	float Grout = 0.f;
 
-	UPROPERTY(EditAnywhere, Category = "Noise", meta = (DataflowInput, DisplayName = "Amplitude", UIMin = 0.f));
+	UPROPERTY(EditAnywhere, Category = "Noise", meta = (DataflowInput, UIMin = 0.f));
 	float Amplitude = 0.f;
 
-	UPROPERTY(EditAnywhere, Category = "Noise", meta = (DataflowInput, DisplayName = "Frequency", UIMin = 0.00001f));
+	UPROPERTY(EditAnywhere, Category = "Noise", meta = (DataflowInput, UIMin = 0.00001f));
 	float Frequency = 0.1f;
 
-	UPROPERTY(EditAnywhere, Category = "Noise", meta = (DataflowInput, DisplayName = "Persistence", UIMin = 0.f));
+	UPROPERTY(EditAnywhere, Category = "Noise", meta = (DataflowInput, UIMin = 0.f));
 	float Persistence = 0.5f;
 
-	UPROPERTY(EditAnywhere, Category = "Noise", meta = (DataflowInput, DisplayName = "Lacunarity", UIMin = 0.f));
+	UPROPERTY(EditAnywhere, Category = "Noise", meta = (DataflowInput, UIMin = 0.f));
 	float Lacunarity = 2.f;
 
-	UPROPERTY(EditAnywhere, Category = "Noise", meta = (DataflowInput, DisplayName = "OctaveNumber", UIMin = 0.f));
+	UPROPERTY(EditAnywhere, Category = "Noise", meta = (DataflowInput, UIMin = 0.f));
 	int32 OctaveNumber = 4;
 
-	UPROPERTY(EditAnywhere, Category = "Noise", meta = (DataflowInput, DisplayName = "PointSpacing", UIMin = 0.f));
+	UPROPERTY(EditAnywhere, Category = "Noise", meta = (DataflowInput, UIMin = 0.f));
 	float PointSpacing = 10.f;
 
 	UPROPERTY(EditAnywhere, Category = "Collision");
 	bool AddSamplesForCollision = false;
 
-	UPROPERTY(EditAnywhere, Category = "Collision", meta = (DataflowInput, DisplayName = "CollisionSampleSpacing", UIMin = 0.f));
+	UPROPERTY(EditAnywhere, Category = "Collision", meta = (DataflowInput, UIMin = 0.f));
 	float CollisionSampleSpacing = 50.f;
 
 	FPlaneCutterDataflowNode(const Dataflow::FNodeParameters& InParam, FGuid InGuid = FGuid::NewGuid())
@@ -1162,6 +1205,12 @@ public:
 
 };
 
+
+/**
+ *
+ * Generates a hash value from a string
+ *
+ */
 USTRUCT()
 struct FHashStringDataflowNode : public FDataflowNode
 {
@@ -1169,13 +1218,14 @@ struct FHashStringDataflowNode : public FDataflowNode
 		DATAFLOW_NODE_DEFINE_INTERNAL(FHashStringDataflowNode, "HashString", "Utilities|String", "")
 
 public:
-	UPROPERTY(meta = (DataflowInput, DisplayName = "String"))
+	/** String to hash */
+	UPROPERTY(meta = (DataflowInput))
 	FString String = FString("");
 
-	UPROPERTY(meta = (DataflowOutput, DisplayName = "Hash"))
+	/** Generated hash value */
+	UPROPERTY(meta = (DataflowOutput))
 	int32 Hash = 0;
 
-
 	FHashStringDataflowNode(const Dataflow::FNodeParameters& InParam, FGuid InGuid = FGuid::NewGuid())
 		: FDataflowNode(InParam, InGuid)
 	{
@@ -1187,6 +1237,12 @@ public:
 
 };
 
+
+/**
+ *
+ * Generates a hash value from a vector
+ *
+ */
 USTRUCT()
 struct FHashVectorDataflowNode : public FDataflowNode
 {
@@ -1194,13 +1250,14 @@ struct FHashVectorDataflowNode : public FDataflowNode
 	DATAFLOW_NODE_DEFINE_INTERNAL(FHashVectorDataflowNode, "HashVector", "Utilities|Vector", "")
 
 public:
-	UPROPERTY(meta = (DataflowInput, DisplayName = "Vector"))
+	/** Vector to hash */
+	UPROPERTY(meta = (DataflowInput))
 	FVector Vector = FVector(0.0);
 
-	UPROPERTY(meta = (DataflowOutput, DisplayName = "Hash"))
+	/** Generated hash value */
+	UPROPERTY(meta = (DataflowOutput))
 	int32 Hash = 0;
 
-
 	FHashVectorDataflowNode(const Dataflow::FNodeParameters& InParam, FGuid InGuid = FGuid::NewGuid())
 		: FDataflowNode(InParam, InGuid)
 	{
@@ -1212,6 +1269,7 @@ public:
 
 };
 
+
 UENUM(BlueprintType)
 enum class EFloatToIntFunctionEnum : uint8
 {
@@ -1224,20 +1282,28 @@ enum class EFloatToIntFunctionEnum : uint8
 	Dataflow_Max                UMETA(Hidden)
 };
 
+/**
+ *
+ * Converts a Float to Int using the specified method
+ *
+ */
 USTRUCT()
 struct FFloatToIntDataflowNode : public FDataflowNode
 {
 	GENERATED_USTRUCT_BODY()
-		DATAFLOW_NODE_DEFINE_INTERNAL(FFloatToIntDataflowNode, "FloatToInt", "Math|Conversions", "")
+	DATAFLOW_NODE_DEFINE_INTERNAL(FFloatToIntDataflowNode, "FloatToInt", "Math|Conversions", "")
 
 public:
-	UPROPERTY(EditAnywhere, Category = "Float", meta = (DisplayName = "Function"));
+	/** Method to convert */
+	UPROPERTY(EditAnywhere, Category = "Float");
 	EFloatToIntFunctionEnum Function = EFloatToIntFunctionEnum::Dataflow_FloatToInt_Function_Round;
 
-	UPROPERTY(EditAnywhere, Category = "Float", meta = (DataflowInput, DisplayName = "Float"))
+	/** Float value to convert */
+	UPROPERTY(EditAnywhere, Category = "Float", meta = (DataflowInput))
 	float Float = 0.f;
 
-	UPROPERTY(meta = (DataflowOutput, DisplayName = "Int"))
+	/** Int output */
+	UPROPERTY(meta = (DataflowOutput))
 	int32 Int = 0;
 
 	FFloatToIntDataflowNode(const Dataflow::FNodeParameters& InParam, FGuid InGuid = FGuid::NewGuid())
@@ -1251,6 +1317,7 @@ public:
 
 };
 
+
 UENUM(BlueprintType)
 enum class EMathConstantsEnum : uint8
 {
@@ -1273,6 +1340,12 @@ enum class EMathConstantsEnum : uint8
 	Dataflow_Max                UMETA(Hidden)
 };
 
+
+/**
+ *
+ * Offers a selection of Math constants
+ *
+ */
 USTRUCT()
 struct FMathConstantsDataflowNode : public FDataflowNode
 {
@@ -1280,10 +1353,12 @@ struct FMathConstantsDataflowNode : public FDataflowNode
 	DATAFLOW_NODE_DEFINE_INTERNAL(FMathConstantsDataflowNode, "MathConstants", "Math|Utilities", "")
 
 public:
-	UPROPERTY(EditAnywhere, Category = "Constants", meta = (DisplayName = "Constant"));
+	/** Math constant to output */
+	UPROPERTY(EditAnywhere, Category = "Constants");
 	EMathConstantsEnum Constant = EMathConstantsEnum::Dataflow_MathConstants_Pi;
 
-	UPROPERTY(meta = (DataflowOutput, DisplayName = "Float"))
+	/** Selected Math constant */
+	UPROPERTY(meta = (DataflowOutput))
 	float Float = 0;
 
 	FMathConstantsDataflowNode(const Dataflow::FNodeParameters& InParam, FGuid InGuid = FGuid::NewGuid())
@@ -1297,10 +1372,922 @@ public:
 };
 
 
-namespace Dataflow
+/**
+ *
+ * Returns the specified element from an array
+ *
+ */
+USTRUCT()
+struct FGetArrayElementDataflowNode : public FDataflowNode
 {
-	void GeometryCollectionEngineNodes();
+	GENERATED_USTRUCT_BODY()
+	DATAFLOW_NODE_DEFINE_INTERNAL(FGetArrayElementDataflowNode, "GetArrayElement", "Utilities|Array", "")
+
+public:
+	/** Element index */
+	UPROPERTY(EditAnywhere, Category = "Index");
+	int32 Index = 0;
+
+	/** Array to get the element from */
+	UPROPERTY(meta = (DataflowInput))
+	TArray<FVector> Points;
+
+	/** Specified element */
+	UPROPERTY(meta = (DataflowOutput))
+	FVector Point = FVector(0.0);
+
+	FGetArrayElementDataflowNode(const Dataflow::FNodeParameters& InParam, FGuid InGuid = FGuid::NewGuid())
+		: FDataflowNode(InParam, InGuid)
+	{
+		RegisterInputConnection(&Points);
+		RegisterInputConnection(&Index);
+		RegisterOutputConnection(&Point);
+	}
+
+	virtual void Evaluate(Dataflow::FContext& Context, const FDataflowOutput* Out) const override;
+
+};
 
 
+/**
+ *
+ * Returns the number of elements in an array
+ *
+ */
+USTRUCT()
+struct FGetNumArrayElementsDataflowNode : public FDataflowNode
+{
+	GENERATED_USTRUCT_BODY()
+	DATAFLOW_NODE_DEFINE_INTERNAL(FGetNumArrayElementsDataflowNode, "GetNumArrayElements", "Utilities|Array", "")
+
+public:
+	/** Array input */
+	UPROPERTY(meta = (DataflowInput))
+	TArray<FVector> Points;
+
+	/** Number of elements in the array */
+	UPROPERTY(meta = (DataflowOutput))
+	int32 NumElements = 0;
+
+	FGetNumArrayElementsDataflowNode(const Dataflow::FNodeParameters& InParam, FGuid InGuid = FGuid::NewGuid())
+		: FDataflowNode(InParam, InGuid)
+	{
+		RegisterInputConnection(&Points);
+		RegisterOutputConnection(&NumElements);
+	}
+
+	virtual void Evaluate(Dataflow::FContext& Context, const FDataflowOutput* Out) const override;
+
+};
+
+
+/**
+ *
+ * Gets BoundingBoxes of pieces from a Collection
+ *
+ */
+USTRUCT()
+struct FGetBoundingBoxesDataflowNode : public FDataflowNode
+{
+	GENERATED_USTRUCT_BODY()
+	DATAFLOW_NODE_DEFINE_INTERNAL(FGetBoundingBoxesDataflowNode, "GetBoundingBoxes", "GeometryCollection", "")
+
+public:
+	/** Input Collection */
+	UPROPERTY(meta = (DataflowInput))
+	FManagedArrayCollection Collection;
+
+	/** The BoundingBoxes will be output for the bones selected in the TransformSelection  */
+	UPROPERTY(meta = (DataflowInput, DisplayName = "TransformSelection"))
+	FDataflowTransformSelection TransformSelection;
+
+	/** Output BoundingBoxes */
+	UPROPERTY(meta = (DataflowOutput))
+	TArray<FBox> BoundingBoxes;
+
+	FGetBoundingBoxesDataflowNode(const Dataflow::FNodeParameters& InParam, FGuid InGuid = FGuid::NewGuid())
+		: FDataflowNode(InParam, InGuid)
+	{
+		RegisterInputConnection(&Collection);
+		RegisterInputConnection(&TransformSelection);
+		RegisterOutputConnection(&BoundingBoxes);
+	}
+
+	virtual void Evaluate(Dataflow::FContext& Context, const FDataflowOutput* Out) const override;
+
+};
+
+
+/**
+ *
+ * Gets centroids of pieces from a Collection
+ *
+ */
+USTRUCT()
+struct FGetCentroidsDataflowNode : public FDataflowNode
+{
+	GENERATED_USTRUCT_BODY()
+	DATAFLOW_NODE_DEFINE_INTERNAL(FGetCentroidsDataflowNode, "GetCentroids", "GeometryCollection", "")
+
+public:
+	/** Input Collection */
+	UPROPERTY(meta = (DataflowInput))
+	FManagedArrayCollection Collection;
+
+	/** The centroids will be output for the bones selected in the TransformSelection  */
+	UPROPERTY(meta = (DataflowInput, DisplayName = "TransformSelection"))
+	FDataflowTransformSelection TransformSelection;
+
+	/** Output centroids */
+	UPROPERTY(meta = (DataflowOutput))
+	TArray<FVector> Centroids = TArray<FVector>();
+
+	FGetCentroidsDataflowNode(const Dataflow::FNodeParameters& InParam, FGuid InGuid = FGuid::NewGuid())
+		: FDataflowNode(InParam, InGuid)
+	{
+		RegisterInputConnection(&Collection);
+		RegisterInputConnection(&TransformSelection);
+		RegisterOutputConnection(&Centroids);
+	}
+
+	virtual void Evaluate(Dataflow::FContext& Context, const FDataflowOutput* Out) const override;
+
+};
+
+
+/**
+ *
+ * Converts points into a DynamicMesh
+ *
+ */
+USTRUCT()
+struct FPointsToMeshDataflowNode : public FDataflowNode
+{
+	GENERATED_USTRUCT_BODY()
+	DATAFLOW_NODE_DEFINE_INTERNAL(FPointsToMeshDataflowNode, "PointsToMesh", "Mesh|Utilities", "")
+
+public:
+	/** Points input */
+	UPROPERTY(meta = (DataflowInput))
+	TArray<FVector> Points;
+
+	/** Mesh output */
+	UPROPERTY(meta = (DataflowOutput))
+	TObjectPtr<UDynamicMesh> Mesh;
+
+	/** Mesh triangle count */
+	UPROPERTY(meta = (DataflowOutput))
+	int32 TriangleCount;
+
+	FPointsToMeshDataflowNode(const Dataflow::FNodeParameters& InParam, FGuid InGuid = FGuid::NewGuid())
+		: FDataflowNode(InParam, InGuid)
+	{
+		RegisterInputConnection(&Points);
+		RegisterOutputConnection(&Mesh);
+		RegisterOutputConnection(&TriangleCount);
+	}
+
+	virtual void Evaluate(Dataflow::FContext& Context, const FDataflowOutput* Out) const override;
+
+};
+
+
+/**
+ *
+ * Converts a BoundingBox into a DynamicMesh
+ *
+ */
+USTRUCT()
+struct FBoxToMeshDataflowNode : public FDataflowNode
+{
+	GENERATED_USTRUCT_BODY()
+	DATAFLOW_NODE_DEFINE_INTERNAL(FBoxToMeshDataflowNode, "BoxToMesh", "Mesh|Utilities", "")
+
+public:
+	/** BoundingBox input */
+	UPROPERTY(meta = (DataflowInput))
+	FBox Box = FBox(ForceInit);
+
+	/** Mesh output */
+	UPROPERTY(meta = (DataflowOutput))
+	TObjectPtr<UDynamicMesh> Mesh;
+
+	/** Mesh triangle count */
+	UPROPERTY(meta = (DataflowOutput))
+	int32 TriangleCount = 0;
+
+	FBoxToMeshDataflowNode(const Dataflow::FNodeParameters& InParam, FGuid InGuid = FGuid::NewGuid())
+		: FDataflowNode(InParam, InGuid)
+	{
+		RegisterInputConnection(&Box);
+		RegisterOutputConnection(&Mesh);
+		RegisterOutputConnection(&TriangleCount);
+	}
+
+	virtual void Evaluate(Dataflow::FContext& Context, const FDataflowOutput* Out) const override;
+
+};
+
+
+/**
+ *
+ * Collects information from the DynamicMesh and outputs it into a formatted string
+ *
+ */
+USTRUCT()
+struct FMeshInfoDataflowNode : public FDataflowNode
+{
+	GENERATED_USTRUCT_BODY()
+	DATAFLOW_NODE_DEFINE_INTERNAL(FMeshInfoDataflowNode, "MeshInfo", "Mesh|Utilities", "")
+
+public:
+	/** DynamicMesh for the information */
+	UPROPERTY(meta = (DataflowInput))
+	TObjectPtr<UDynamicMesh> Mesh;
+
+	/** Formatted output string */
+	UPROPERTY(meta = (DataflowOutput))
+	FString InfoString = FString("");
+
+	FMeshInfoDataflowNode(const Dataflow::FNodeParameters& InParam, FGuid InGuid = FGuid::NewGuid())
+		: FDataflowNode(InParam, InGuid)
+	{
+		RegisterInputConnection(&Mesh);
+		RegisterOutputConnection(&InfoString);
+	}
+
+	virtual void Evaluate(Dataflow::FContext& Context, const FDataflowOutput* Out) const override;
+
+};
+
+
+/**
+ *
+ * Converts a DynamicMesh to a Collection
+ *
+ */
+USTRUCT()
+struct FMeshToCollectionDataflowNode : public FDataflowNode
+{
+	GENERATED_USTRUCT_BODY()
+	DATAFLOW_NODE_DEFINE_INTERNAL(FMeshToCollectionDataflowNode, "MeshToCollection", "Mesh|Utilities", "")
+
+public:
+	/** DynamicMesh to convert */
+	UPROPERTY(meta = (DataflowInput))
+	TObjectPtr<UDynamicMesh> Mesh;
+
+	/** Output Collection */
+	UPROPERTY(meta = (DataflowOutput))
+	FManagedArrayCollection Collection;
+
+	FMeshToCollectionDataflowNode(const Dataflow::FNodeParameters& InParam, FGuid InGuid = FGuid::NewGuid())
+		: FDataflowNode(InParam, InGuid)
+	{
+		RegisterInputConnection(&Mesh);
+		RegisterOutputConnection(&Collection);
+	}
+
+	virtual void Evaluate(Dataflow::FContext& Context, const FDataflowOutput* Out) const override;
+
+};
+
+
+/**
+ *
+ * Converts a StaticMesh into a DynamicMesh
+ *
+ */
+USTRUCT()
+struct FStaticMeshToMeshDataflowNode : public FDataflowNode
+{
+	GENERATED_USTRUCT_BODY()
+	DATAFLOW_NODE_DEFINE_INTERNAL(FStaticMeshToMeshDataflowNode, "StaticMeshToMesh", "Mesh|Utilities", "")
+
+public:
+	/** StaticMesh to convert */
+	UPROPERTY(EditAnywhere, Category = "StaticMesh");
+	TObjectPtr<UStaticMesh> StaticMesh;
+
+	/** Output the HiRes representation, if set to true and HiRes doesn't exist it will output empty mesh */
+	UPROPERTY(EditAnywhere, Category = "StaticMesh", meta = (DisplayName = "Use HiRes"));
+	bool UseHiRes = true;
+
+	/** Specifies the LOD level to use */
+	UPROPERTY(EditAnywhere, Category = "StaticMesh", meta = (DisplayName = "LOD Level"));
+	int32 LODLevel = 0;
+
+	/** Output mesh */
+	UPROPERTY(meta = (DataflowOutput))
+	TObjectPtr<UDynamicMesh> Mesh;
+
+	FStaticMeshToMeshDataflowNode(const Dataflow::FNodeParameters& InParam, FGuid InGuid = FGuid::NewGuid())
+		: FDataflowNode(InParam, InGuid)
+	{
+		RegisterOutputConnection(&Mesh);
+	}
+
+	virtual void Evaluate(Dataflow::FContext& Context, const FDataflowOutput* Out) const override;
+
+};
+
+
+/**
+ *
+ * Transforms a mesh
+ *
+ */
+USTRUCT()
+struct FTransformDataflowNode : public FDataflowNode
+{
+	GENERATED_USTRUCT_BODY()
+	DATAFLOW_NODE_DEFINE_INTERNAL(FTransformDataflowNode, "Transform", "Math", "")
+
+public:
+	/** Translation */
+	UPROPERTY(EditAnywhere, Category = "Transform");
+	FVector Translate = FVector(0.0);
+
+	/** Rotation */
+	UPROPERTY(EditAnywhere, Category = "Transform");
+	FVector Rotate = FVector(0.0);
+
+	/** Scale */
+	UPROPERTY(EditAnywhere, Category = "Transform");
+	FVector Scale = FVector(1.0);
+
+	/** Shear */
+	UPROPERTY(EditAnywhere, Category = "Transform");
+	FVector Shear = FVector(0.0);
+
+	/** Uniform scale */
+	UPROPERTY(EditAnywhere, Category = "Transform");
+	float UniformScale = 1.f;
+
+	/** Pivot for the rotation */
+	UPROPERTY(EditAnywhere, Category = "Pivot");
+	FVector RotatePivot = FVector(0.0);
+
+	/** Pivot for the scale */
+	UPROPERTY(EditAnywhere, Category = "Pivot");
+	FVector ScalePivot = FVector(0.0);
+
+	/** Invert the transformation */
+	UPROPERTY(EditAnywhere, Category = "General");
+	bool InvertTransformation = false;
+
+	/** Output mesh */
+	UPROPERTY(meta = (DataflowInput, DataflowOutput))
+	TObjectPtr<UDynamicMesh> Mesh;
+
+	FTransformDataflowNode(const Dataflow::FNodeParameters& InParam, FGuid InGuid = FGuid::NewGuid())
+		: FDataflowNode(InParam, InGuid)
+	{
+		RegisterInputConnection(&Mesh);
+		RegisterOutputConnection(&Mesh);
+	}
+
+	virtual void Evaluate(Dataflow::FContext& Context, const FDataflowOutput* Out) const override;
+
+};
+
+
+/**
+ *
+ * Appends two meshes
+ *
+ */
+USTRUCT()
+struct FMeshAppendDataflowNode : public FDataflowNode
+{
+	GENERATED_USTRUCT_BODY()
+	DATAFLOW_NODE_DEFINE_INTERNAL(FMeshAppendDataflowNode, "MeshAppend", "Mesh|Utilities", "")
+
+public:
+	/** Mesh input */
+	UPROPERTY(meta = (DataflowInput))
+	TObjectPtr<UDynamicMesh> Mesh1;
+
+	/** Mesh input */
+	UPROPERTY(meta = (DataflowInput))
+	TObjectPtr<UDynamicMesh> Mesh2;
+
+	/** Output mesh */
+	UPROPERTY(meta = (DataflowOutput))
+	TObjectPtr<UDynamicMesh> Mesh;
+
+	FMeshAppendDataflowNode(const Dataflow::FNodeParameters& InParam, FGuid InGuid = FGuid::NewGuid())
+		: FDataflowNode(InParam, InGuid)
+	{
+		RegisterInputConnection(&Mesh1);
+		RegisterInputConnection(&Mesh2);
+		RegisterOutputConnection(&Mesh);
+	}
+
+	virtual void Evaluate(Dataflow::FContext& Context, const FDataflowOutput* Out) const override;
+
+};
+
+
+UENUM(BlueprintType)
+enum class EMeshBooleanOperationEnum : uint8
+{
+	Dataflow_MeshBoolean_Union UMETA(DisplayName = "Union"),
+	Dataflow_MeshBoolean_Intersect UMETA(DisplayName = "Intersect"),
+	Dataflow_MeshBoolean_Difference UMETA(DisplayName = "Difference"),
+	//~~~
+	//256th entry
+	Dataflow_Max                UMETA(Hidden)
+};
+
+/**
+ *
+ * Mesh boolean (Union, Intersect, Difference) between two meshes
+ *
+ */
+USTRUCT()
+struct FMeshBooleanDataflowNode : public FDataflowNode
+{
+	GENERATED_USTRUCT_BODY()
+	DATAFLOW_NODE_DEFINE_INTERNAL(FMeshBooleanDataflowNode, "MeshBoolean", "Mesh|Utilities", "")
+
+public:
+	/** Boolean operation */
+	UPROPERTY(EditAnywhere, Category = "Boolean");
+	EMeshBooleanOperationEnum Operation = EMeshBooleanOperationEnum::Dataflow_MeshBoolean_Intersect;
+
+	/** Mesh input */
+	UPROPERTY(meta = (DataflowInput))
+	TObjectPtr<UDynamicMesh> Mesh1;
+
+	/** Mesh input */
+	UPROPERTY(meta = (DataflowInput))
+	TObjectPtr<UDynamicMesh> Mesh2;
+
+	/** Output mesh */
+	UPROPERTY(meta = (DataflowOutput))
+	TObjectPtr<UDynamicMesh> Mesh;
+
+	FMeshBooleanDataflowNode(const Dataflow::FNodeParameters& InParam, FGuid InGuid = FGuid::NewGuid())
+		: FDataflowNode(InParam, InGuid)
+	{
+		RegisterInputConnection(&Mesh1);
+		RegisterInputConnection(&Mesh2);
+		RegisterOutputConnection(&Mesh);
+	}
+
+	virtual void Evaluate(Dataflow::FContext& Context, const FDataflowOutput* Out) const override;
+
+};
+
+
+/**
+ *
+ * Copies the same mesh with scale onto points
+ *
+ */
+USTRUCT()
+struct FMeshCopyToPointsDataflowNode : public FDataflowNode
+{
+	GENERATED_USTRUCT_BODY()
+	DATAFLOW_NODE_DEFINE_INTERNAL(FMeshCopyToPointsDataflowNode, "MeshCopyToPoints", "Mesh|Utilities", "")
+
+public:
+	/** Points to copy meshes onto */
+	UPROPERTY(meta = (DataflowInput))
+	TArray<FVector> Points;
+
+	/** Mesh to copy onto points */
+	UPROPERTY(meta = (DataflowInput))
+	TObjectPtr<UDynamicMesh> MeshToCopy;
+
+	/** Scale appied to the mesh */
+	UPROPERTY(EditAnywhere, Category = "Copy");
+	float Scale = 1.f;
+
+	/** Copied meshes */
+	UPROPERTY(meta = (DataflowOutput))
+	TObjectPtr<UDynamicMesh> Mesh;
+
+	FMeshCopyToPointsDataflowNode(const Dataflow::FNodeParameters& InParam, FGuid InGuid = FGuid::NewGuid())
+		: FDataflowNode(InParam, InGuid)
+	{
+		RegisterInputConnection(&Points);
+		RegisterInputConnection(&MeshToCopy);
+		RegisterOutputConnection(&Mesh);
+	}
+
+	virtual void Evaluate(Dataflow::FContext& Context, const FDataflowOutput* Out) const override;
+
+};
+
+
+UENUM(BlueprintType)
+enum class ECompareOperationEnum : uint8
+{
+	Dataflow_Compare_Equal UMETA(DisplayName = "=="),
+	Dataflow_Compare_Smaller UMETA(DisplayName = "<"),
+	Dataflow_Compare_SmallerOrEqual UMETA(DisplayName = "<="),
+	Dataflow_Compare_Greater UMETA(DisplayName = ">"),
+	Dataflow_Compare_GreaterOrEqual UMETA(DisplayName = ">="),
+	//~~~
+	//256th entry
+	Dataflow_Max                UMETA(Hidden)
+};
+
+/**
+ *
+ * Comparison between integers
+ *
+ */
+USTRUCT()
+struct FCompareIntDataflowNode : public FDataflowNode
+{
+	GENERATED_USTRUCT_BODY()
+	DATAFLOW_NODE_DEFINE_INTERNAL(FCompareIntDataflowNode, "CompareInt", "Math|Int", "")
+
+public:
+	/** Comparison operation */
+	UPROPERTY(EditAnywhere, Category = "Compare");
+	ECompareOperationEnum Operation = ECompareOperationEnum::Dataflow_Compare_Equal;
+
+	/** Int input */
+	UPROPERTY(EditAnywhere, Category = "Compare");
+	int32 IntA = 0;
+
+	/** Int input */
+	UPROPERTY(EditAnywhere, Category = "Compare");
+	int32 IntB = 0;
+
+	/** Boolean result of the comparison */
+	UPROPERTY(meta = (DataflowOutput));
+	bool Result;
+
+	FCompareIntDataflowNode(const Dataflow::FNodeParameters& InParam, FGuid InGuid = FGuid::NewGuid())
+		: FDataflowNode(InParam, InGuid)
+	{
+		RegisterInputConnection(&IntA);
+		RegisterInputConnection(&IntB);
+		RegisterOutputConnection(&Result);
+	}
+
+	virtual void Evaluate(Dataflow::FContext& Context, const FDataflowOutput* Out) const override;
+
+};
+
+
+/**
+ *
+ * Branch between two inputs based on boolean condition
+ *
+ */
+USTRUCT()
+struct FBranchDataflowNode : public FDataflowNode
+{
+	GENERATED_USTRUCT_BODY()
+	DATAFLOW_NODE_DEFINE_INTERNAL(FBranchDataflowNode, "Branch", "Utilites|FlowControl", "")
+
+public:
+	/** Mesh input */
+	UPROPERTY(meta = (DataflowInput))
+	TObjectPtr<UDynamicMesh> MeshA;
+
+	/** Mesh input */
+	UPROPERTY(meta = (DataflowInput))
+	TObjectPtr<UDynamicMesh> MeshB;
+
+	/** If true, Output = MeshA, otherwise Output = MeshB */
+	UPROPERTY(EditAnywhere, Category = "Branch");
+	bool Condition = false;
+
+	/** Output mesh */
+	UPROPERTY(meta = (DataflowOutput))
+	TObjectPtr<UDynamicMesh> Mesh;
+
+	FBranchDataflowNode(const Dataflow::FNodeParameters& InParam, FGuid InGuid = FGuid::NewGuid())
+		: FDataflowNode(InParam, InGuid)
+	{
+		RegisterInputConnection(&MeshA);
+		RegisterInputConnection(&MeshB);
+		RegisterInputConnection(&Condition);
+		RegisterOutputConnection(&Mesh);
+	}
+
+	virtual void Evaluate(Dataflow::FContext& Context, const FDataflowOutput* Out) const override;
+
+};
+
+
+/**
+ *
+ * Outputs Mesh data
+ *
+ */
+USTRUCT()
+struct FGetMeshDataDataflowNode : public FDataflowNode
+{
+	GENERATED_USTRUCT_BODY()
+	DATAFLOW_NODE_DEFINE_INTERNAL(FGetMeshDataDataflowNode, "GetMeshData", "Mesh|Utilities", "")
+
+public:
+	/** Mesh for the data */
+	UPROPERTY(meta = (DataflowInput))
+	TObjectPtr<UDynamicMesh> Mesh;
+
+	/** Number of vertices */
+	UPROPERTY(meta = (DataflowOutput))
+	int32 VertexCount = 0;
+
+	/** Number of edges */
+	UPROPERTY(meta = (DataflowOutput))
+	int32 EdgeCount = 0;
+
+	/** Number of triangles */
+	UPROPERTY(meta = (DataflowOutput))
+	int32 TriangleCount = 0;
+
+	FGetMeshDataDataflowNode(const Dataflow::FNodeParameters& InParam, FGuid InGuid = FGuid::NewGuid())
+		: FDataflowNode(InParam, InGuid)
+	{
+		RegisterInputConnection(&Mesh);
+		RegisterOutputConnection(&VertexCount);
+		RegisterOutputConnection(&EdgeCount);
+		RegisterOutputConnection(&TriangleCount);
+	}
+
+	virtual void Evaluate(Dataflow::FContext& Context, const FDataflowOutput* Out) const override;
+
+};
+
+
+/**
+ *
+ * Collects grooup and attribute information from the Collection and outputs it into a formatted string
+ *
+ */
+USTRUCT()
+struct FGetSchemaDataflowNode : public FDataflowNode
+{
+	GENERATED_USTRUCT_BODY()
+	DATAFLOW_NODE_DEFINE_INTERNAL(FGetSchemaDataflowNode, "GetSchema", "GeometryCollection|Utilities", "")
+
+public:
+	/** GeometryCollection  for the information */
+	UPROPERTY(meta = (DataflowInput))
+	FManagedArrayCollection Collection;
+
+	/** Formatted string containing the groups and attributes */
+	UPROPERTY(meta = (DataflowOutput))
+	FString String;
+
+	FGetSchemaDataflowNode(const Dataflow::FNodeParameters& InParam, FGuid InGuid = FGuid::NewGuid())
+		: FDataflowNode(InParam, InGuid)
+	{
+		RegisterInputConnection(&Collection);
+		RegisterOutputConnection(&String);
+	}
+
+	virtual void Evaluate(Dataflow::FContext& Context, const FDataflowOutput* Out) const override;
+
+};
+
+
+UENUM(BlueprintType)
+enum class EClusterSizeMethodEnum : uint8
+{
+	Dataflow_ClusterSizeMethod_ByNumber UMETA(DisplayName = "By Number"),
+	Dataflow_ClusterSizeMethod_ByFractionOfInput UMETA(DisplayName = "By Fraction Of Input"),
+	Dataflow_ClusterSizeMethod_BySize UMETA(DisplayName = "By Size"),
+	//~~~
+	//256th entry
+	Dataflow_Max                UMETA(Hidden)
+};
+
+/**
+ *
+ * Automatically group pieces of a fractured Collection into a specified number of clusters
+ *
+ */
+USTRUCT()
+struct FAutoClusterDataflowNode : public FDataflowNode
+{
+	GENERATED_USTRUCT_BODY()
+	DATAFLOW_NODE_DEFINE_INTERNAL(FAutoClusterDataflowNode, "AutoCluster", "GeometryCollection|Cluster", "")
+
+public:
+	/** How to choose the size of the clusters to create */
+	UPROPERTY(EditAnywhere, Category = "Cluster Size");
+	EClusterSizeMethodEnum ClusterSizeMethod = EClusterSizeMethodEnum::Dataflow_ClusterSizeMethod_ByNumber;
+
+	/** Use a Voronoi diagram with this many Voronoi sites as a guide for deciding cluster boundaries */
+	UPROPERTY(EditAnywhere, Category = "Cluster Size", meta = (DataflowInput, UIMin = 2, UIMax = 5000, EditCondition = "ClusterSizeMethod == EClusterSizeMethodEnum::Dataflow_ClusterSizeMethod_ByNumber"))
+	int32 ClusterSites = 10;
+
+	/** Choose the number of Voronoi sites used for clustering as a fraction of the number of child bones to process */
+	UPROPERTY(EditAnywhere, Category = "Cluster Size", meta = (DataflowInput, UIMin = 0.f, UIMax = 0.5f, EditCondition = "ClusterSizeMethod == EClusterSizeMethodEnum::Dataflow_ClusterSizeMethod_ByFractionOfInput"))
+	float ClusterFraction = 0.25;
+
+	/** Choose the Edge-Size of the cube used to groups bones under a cluster (in cm). */
+	UPROPERTY(EditAnywhere, Category = "ClusterSize", meta = (DataflowInput, DisplayName = "Cluster Size", UIMin = ".01", UIMax = "100", ClampMin = ".0001", ClampMax = "10000", EditCondition = "ClusterSizeMethod == EClusterSizeMethodEnum::Dataflow_ClusterSizeMethod_BySize"))
+	float SiteSize = 1;
+
+	/** If true, bones will only be added to the same cluster if they are physically connected (either directly, or via other bones in the same cluster) */
+	UPROPERTY(EditAnywhere, Category = "Auto Cluster")
+	bool AutoCluster = true;
+
+	/** If true, prevent the creation of clusters with only a single child. Either by merging into a neighboring cluster, or not creating the cluster. */
+	UPROPERTY(EditAnywhere, Category = "Auto Cluster")
+	bool AvoidIsolated = true;
+
+	/** Fractured GeometryCollection to cluster */
+	UPROPERTY(meta = (DataflowInput, DataflowOutput, DataflowPassthrough = "Collection", DataflowIntrinsic))
+	FManagedArrayCollection Collection;
+
+	/** Bone selection for the clustering */
+	UPROPERTY(meta = (DataflowInput, DisplayName = "TransformSelection"))
+	FDataflowTransformSelection TransformSelection;
+
+	FAutoClusterDataflowNode(const Dataflow::FNodeParameters& InParam, FGuid InGuid = FGuid::NewGuid())
+		: FDataflowNode(InParam, InGuid)
+	{
+		RegisterInputConnection(&Collection);
+		RegisterInputConnection(&TransformSelection);
+		RegisterInputConnection(&ClusterSites);
+		RegisterInputConnection(&ClusterFraction);
+		RegisterInputConnection(&SiteSize);
+		RegisterOutputConnection(&Collection, &Collection);
+	}
+
+	virtual void Evaluate(Dataflow::FContext& Context, const FDataflowOutput* Out) const override;
+
+};
+
+
+USTRUCT()
+struct FClusterFlattenDataflowNode : public FDataflowNode
+{
+	GENERATED_USTRUCT_BODY()
+	DATAFLOW_NODE_DEFINE_INTERNAL(FClusterFlattenDataflowNode, "Flatten", "GeometryCollection|Cluster", "")
+
+public:
+	// @todo(harsha) Support Selections
+
+	UPROPERTY(meta = (DataflowInput, DataflowOutput, DataflowPassthrough = "Collection", DataflowIntrinsic))
+	FManagedArrayCollection Collection;
+
+	FClusterFlattenDataflowNode(const Dataflow::FNodeParameters& InParam, FGuid InGuid = FGuid::NewGuid())
+		: FDataflowNode(InParam, InGuid)
+	{
+		RegisterInputConnection(&Collection);
+		RegisterOutputConnection(&Collection, &Collection);
+	}
+
+	virtual void Evaluate(Dataflow::FContext& Context, const FDataflowOutput* Out) const override;
+
+};
+
+
+USTRUCT()
+struct FRemoveOnBreakDataflowNode : public FDataflowNode
+{
+	GENERATED_USTRUCT_BODY()
+	DATAFLOW_NODE_DEFINE_INTERNAL(FRemoveOnBreakDataflowNode, "RemoveOnBreak", "Fracture|Utilities", "")
+
+public:
+	// @todo(harsha) Support Selections
+
+	UPROPERTY(meta = (DataflowInput, DataflowOutput, DataflowPassthrough = "Collection", DataflowIntrinsic))
+	FManagedArrayCollection Collection;
+
+	UPROPERTY(EditAnywhere, Category = "Removal", meta = (DataflowInput))
+	FVector2f PostBreakTimer{0.0, 0.0};
+
+	UPROPERTY(EditAnywhere, Category = "Removal", meta = (DataflowInput))
+	FVector2f RemovalTimer{0.0, 1.0};
+
+	UPROPERTY(EditAnywhere, Category = "Removal", meta = (DataflowInput))
+	bool ClusterCrumbling = false;
+
+	FRemoveOnBreakDataflowNode(const Dataflow::FNodeParameters& InParam, FGuid InGuid = FGuid::NewGuid())
+		: FDataflowNode(InParam, InGuid)
+	{
+		RegisterInputConnection(&Collection);
+		RegisterInputConnection(&PostBreakTimer);
+		RegisterInputConnection(&RemovalTimer);
+		RegisterInputConnection(&ClusterCrumbling);
+		RegisterOutputConnection(&Collection, &Collection);
+	}
+
+	virtual void Evaluate(Dataflow::FContext& Context, const FDataflowOutput* Out) const override;
+
+};
+
+
+UENUM(BlueprintType)
+enum class EAnchorStateEnum : uint8
+{
+	Dataflow_AnchorState_Anchored UMETA(DisplayName = "Anchored"),
+	Dataflow_AnchorState_NotAnchored UMETA(DisplayName = "Not Anchored"),
+	//~~~
+	//256th entry
+	Dataflow_Max                UMETA(Hidden)
+};
+
+/**
+ *
+ * Sets the anchored state on the selected bones in a Collection
+ *
+ */
+USTRUCT()
+struct FSetAnchorStateDataflowNode : public FDataflowNode
+{
+	GENERATED_USTRUCT_BODY()
+	DATAFLOW_NODE_DEFINE_INTERNAL(FSetAnchorStateDataflowNode, "SetAnchorState", "GeometryCollection", "")
+
+public:
+	/** What anchor state to set on selected bones */
+	UPROPERTY(EditAnywhere, Category = "Anchoring");
+	EAnchorStateEnum AnchorState = EAnchorStateEnum::Dataflow_AnchorState_Anchored;
+
+	/** If true, sets the non selected bones to opposite anchor state */
+	UPROPERTY(EditAnywhere, Category = "Anchoring")
+	bool SetNotSelectedBonesToOppositeState = false;
+
+	/** GeometryCollection to set anchor state on */
+	UPROPERTY(meta = (DataflowInput, DataflowOutput, DataflowPassthrough = "Collection", DataflowIntrinsic))
+	FManagedArrayCollection Collection;
+
+	/** Bone selection for setting the state on */
+	UPROPERTY(meta = (DataflowInput, DisplayName = "TransformSelection", DataflowIntrinsic))
+	FDataflowTransformSelection TransformSelection;
+
+	FSetAnchorStateDataflowNode(const Dataflow::FNodeParameters& InParam, FGuid InGuid = FGuid::NewGuid())
+		: FDataflowNode(InParam, InGuid)
+	{
+		RegisterInputConnection(&Collection);
+		RegisterInputConnection(&TransformSelection);
+		RegisterOutputConnection(&Collection, &Collection);
+	}
+
+	virtual void Evaluate(Dataflow::FContext& Context, const FDataflowOutput* Out) const override;
+
+};
+
+
+UENUM(BlueprintType)
+enum class EProximityMethodEnum : uint8
+{
+	/** Precise proximity mode looks for geometry with touching vertices or touching, coplanar, opposite - facing triangles.This works well with geometry fractured using our fracture tools. */
+	Dataflow_ProximityMethod_Precise UMETA(DisplayName = "Precise"),
+	/** Convex Hull proximity mode looks for geometry with overlapping convex hulls(with an optional offset) */
+	Dataflow_ProximityMethod_ConvexHull UMETA(DisplayName = "ConvexHull"),
+	//~~~
+	//256th entry
+	Dataflow_Max                UMETA(Hidden)
+};
+
+/**
+ *
+ * Update the proximity (contact) graph for the bones in a Collection
+ *
+ */
+USTRUCT()
+struct FProximityDataflowNode : public FDataflowNode
+{
+	GENERATED_USTRUCT_BODY()
+	DATAFLOW_NODE_DEFINE_INTERNAL(FProximityDataflowNode, "Proximity", "GeometryCollection", "")
+
+public:
+	/** Which method to use to decide whether a given piece of geometry is in proximity with another */
+	UPROPERTY(EditAnywhere, Category = "Proximity");
+	EProximityMethodEnum ProximityMethod = EProximityMethodEnum::Dataflow_ProximityMethod_Precise;
+
+	/** If hull-based proximity detection is enabled, amount to expand hulls when searching for overlapping neighbors */
+	UPROPERTY(EditAnywhere, Category = "Proximity", meta = (ClampMin = "0", EditCondition = "ProximityMethod == EProximityMethodEnum::Dataflow_ProximityMethod_ConvexHull"))
+	float DistanceThreshold = 1;
+
+	/** Whether to automatically transform the proximity graph into a connection graph to be used for simulation */
+	UPROPERTY(EditAnywhere, Category = "Proximity")
+	bool bUseAsConnectionGraph = false;
+
+	/** GeometryCollection to update the proximity graph on */
+	UPROPERTY(meta = (DataflowInput, DataflowOutput, DataflowPassthrough = "Collection", DataflowIntrinsic))
+	FManagedArrayCollection Collection;
+
+	FProximityDataflowNode(const Dataflow::FNodeParameters& InParam, FGuid InGuid = FGuid::NewGuid())
+		: FDataflowNode(InParam, InGuid)
+	{
+		RegisterInputConnection(&Collection);
+		RegisterOutputConnection(&Collection, &Collection);
+	}
+
+	virtual void Evaluate(Dataflow::FContext& Context, const FDataflowOutput* Out) const override;
+
+};
+
+
+namespace Dataflow
+{
+	void GeometryCollectionEngineNodes();
 }
 
diff --git a/Engine/Plugins/Experimental/GeometryCollectionPlugin/Source/GeometryCollectionNodes/Public/Dataflow/GeometryCollectionSelectionNodes.h b/Engine/Plugins/Experimental/GeometryCollectionPlugin/Source/GeometryCollectionNodes/Public/Dataflow/GeometryCollectionSelectionNodes.h
new file mode 100644
--- /dev/null
+++ b/Engine/Plugins/Experimental/GeometryCollectionPlugin/Source/GeometryCollectionNodes/Public/Dataflow/GeometryCollectionSelectionNodes.h
@@ -0,0 +1,694 @@
+// Copyright Epic Games, Inc. All Rights Reserved.
+
+#pragma once 
+
+#include "CoreMinimal.h"
+#include "Dataflow/DataflowEngine.h"
+#include "GeometryCollection/ManagedArrayCollection.h"
+#include "Dataflow/DataflowSelection.h"
+
+#include "GeometryCollectionSelectionNodes.generated.h"
+
+
+class FGeometryCollection;
+
+
+/**
+ *
+ * Selects all the bones for the Collection
+ *
+ */
+USTRUCT()
+struct FCollectionTransformSelectionAllDataflowNode : public FDataflowNode
+{
+	GENERATED_USTRUCT_BODY()
+	DATAFLOW_NODE_DEFINE_INTERNAL(FCollectionTransformSelectionAllDataflowNode, "CollectionTransformSelectAll", "GeometryCollection|Selection", "")
+
+public:
+	/** GeometryCollection for the selection */
+	UPROPERTY(meta = (DataflowInput, DataflowOutput))
+	FManagedArrayCollection Collection;
+
+	/** Array of the selected bone indicies */
+	UPROPERTY(meta = (DataflowOutput, DisplayName = "TransformSelection"))
+	FDataflowTransformSelection TransformSelection;
+
+	FCollectionTransformSelectionAllDataflowNode(const Dataflow::FNodeParameters& InParam, FGuid InGuid = FGuid::NewGuid())
+		: FDataflowNode(InParam, InGuid)
+	{
+		RegisterInputConnection(&Collection);
+		RegisterOutputConnection(&TransformSelection);
+		RegisterOutputConnection(&Collection);
+	}
+
+	virtual void Evaluate(Dataflow::FContext& Context, const FDataflowOutput* Out) const override;
+
+};
+
+
+UENUM(BlueprintType)
+enum class ESetOperationEnum : uint8
+{
+	Dataflow_SetOperation_AND UMETA(DisplayName = "AND"),
+	Dataflow_SetOperation_OR UMETA(DisplayName = "OR"),
+	Dataflow_SetOperation_XOR UMETA(DisplayName = "XOR"),
+	//~~~
+	//256th entry
+	Dataflow_Max                UMETA(Hidden)
+};
+
+/**
+ *
+ * Runs boolean operation on incoming TransformSelections
+ *
+ */
+USTRUCT()
+struct FCollectionTransformSelectionSetOperationDataflowNode : public FDataflowNode
+{
+	GENERATED_USTRUCT_BODY()
+	DATAFLOW_NODE_DEFINE_INTERNAL(FCollectionTransformSelectionSetOperationDataflowNode, "CollectionTransformSelectionSetOperation", "GeometryCollection|Selection", "")
+
+public:
+	/** Boolean operation */
+	UPROPERTY(EditAnywhere, Category = "Compare");
+	ESetOperationEnum Operation = ESetOperationEnum::Dataflow_SetOperation_AND;
+
+	/** Array of the selected bone indicies */
+	UPROPERTY(meta = (DataflowInput, DisplayName = "TransformSelectionA"))
+	FDataflowTransformSelection TransformSelectionA;
+
+	/** Array of the selected bone indicies */
+	UPROPERTY(meta = (DataflowInput, DisplayName = "TransformSelectionB"))
+	FDataflowTransformSelection TransformSelectionB;
+
+	/** Array of the selected bone indicies after operation*/
+	UPROPERTY(meta = (DataflowOutput, DisplayName = "TransformSelection"))
+	FDataflowTransformSelection TransformSelection;
+
+	FCollectionTransformSelectionSetOperationDataflowNode(const Dataflow::FNodeParameters& InParam, FGuid InGuid = FGuid::NewGuid())
+		: FDataflowNode(InParam, InGuid)
+	{
+		RegisterInputConnection(&TransformSelectionA);
+		RegisterInputConnection(&TransformSelectionB);
+		RegisterOutputConnection(&TransformSelection);
+	}
+
+	virtual void Evaluate(Dataflow::FContext& Context, const FDataflowOutput* Out) const override;
+
+};
+
+
+/**
+ *
+ * Generates a formatted string of the bones and the selection
+ *
+ */
+USTRUCT()
+struct FCollectionTransformSelectionInfoDataflowNode : public FDataflowNode
+{
+	GENERATED_USTRUCT_BODY()
+	DATAFLOW_NODE_DEFINE_INTERNAL(FCollectionTransformSelectionInfoDataflowNode, "CollectionTransformSelectionInfo", "GeometryCollection|Selection", "")
+
+public:
+	/** Array of the selected bone indicies */
+	UPROPERTY(meta = (DataflowInput, DisplayName = "TransformSelection"))
+	FDataflowTransformSelection TransformSelection;
+
+	/** GeometryCollection for the selection */
+	UPROPERTY(meta = (DataflowInput))
+	FManagedArrayCollection Collection;
+
+	/** Formatted string of the bones and selection */
+	UPROPERTY(meta = (DataflowOutput))
+	FString String;
+
+	FCollectionTransformSelectionInfoDataflowNode(const Dataflow::FNodeParameters& InParam, FGuid InGuid = FGuid::NewGuid())
+		: FDataflowNode(InParam, InGuid)
+	{
+		RegisterInputConnection(&TransformSelection);
+		RegisterInputConnection(&Collection);
+		RegisterOutputConnection(&String);
+	}
+
+	virtual void Evaluate(Dataflow::FContext& Context, const FDataflowOutput* Out) const override;
+
+};
+
+
+/**
+ *
+ * Generates an empty bone selection for the Collection
+ *
+ */
+USTRUCT()
+struct FCollectionTransformSelectionNoneDataflowNode : public FDataflowNode
+{
+	GENERATED_USTRUCT_BODY()
+	DATAFLOW_NODE_DEFINE_INTERNAL(FCollectionTransformSelectionNoneDataflowNode, "CollectionTransformSelectNone", "GeometryCollection|Selection", "")
+
+public:
+	/** GeometryCollection for the selection */
+	UPROPERTY(meta = (DataflowInput, DataflowOutput))
+	FManagedArrayCollection Collection;
+
+	/** Array of the selected bone indicies */
+	UPROPERTY(meta = (DataflowOutput, DisplayName = "TransformSelection"))
+	FDataflowTransformSelection TransformSelection;
+
+	FCollectionTransformSelectionNoneDataflowNode(const Dataflow::FNodeParameters& InParam, FGuid InGuid = FGuid::NewGuid())
+		: FDataflowNode(InParam, InGuid)
+	{
+		RegisterInputConnection(&Collection);
+		RegisterOutputConnection(&TransformSelection);
+		RegisterOutputConnection(&Collection);
+	}
+
+	virtual void Evaluate(Dataflow::FContext& Context, const FDataflowOutput* Out) const override;
+
+};
+
+
+/**
+ *
+ * Inverts the incoming selection of bones
+ *
+ */
+USTRUCT()
+struct FCollectionTransformSelectionInvertDataflowNode : public FDataflowNode
+{
+	GENERATED_USTRUCT_BODY()
+	DATAFLOW_NODE_DEFINE_INTERNAL(FCollectionTransformSelectionInvertDataflowNode, "CollectionTransformSelectInvert", "GeometryCollection|Selection", "")
+
+public:
+	/** Array of the selected bone indicies */
+	UPROPERTY(meta = (DataflowInput, DataflowOutput, DisplayName = "TransformSelection"))
+	FDataflowTransformSelection TransformSelection;
+
+	FCollectionTransformSelectionInvertDataflowNode(const Dataflow::FNodeParameters& InParam, FGuid InGuid = FGuid::NewGuid())
+		: FDataflowNode(InParam, InGuid)
+	{
+		RegisterInputConnection(&TransformSelection);
+		RegisterOutputConnection(&TransformSelection);
+	}
+
+	virtual void Evaluate(Dataflow::FContext& Context, const FDataflowOutput* Out) const override;
+
+};
+
+
+/**
+ *
+ * Selects bones randomly in the Collection
+ *
+ */
+USTRUCT()
+struct FCollectionTransformSelectionRandomDataflowNode : public FDataflowNode
+{
+	GENERATED_USTRUCT_BODY()
+	DATAFLOW_NODE_DEFINE_INTERNAL(FCollectionTransformSelectionRandomDataflowNode, "CollectionTransformSelectRandom", "GeometryCollection|Selection", "")
+
+public:
+	/** If true, it always generates the same result for the same RandomSeed */
+	UPROPERTY(EditAnywhere, Category = "Random")
+	bool Deterministic = false;
+
+	/** Seed for the random generation, only used if Deterministic is on */
+	UPROPERTY(EditAnywhere, Category = "Random", meta = (DataflowInput, EditCondition = "Deterministic"))
+	float RandomSeed = 0.f;
+
+	/** Bones get selected if RandomValue > RandomThreshold */
+	UPROPERTY(EditAnywhere, Category = "Random", meta = (DataflowInput, UIMin = 0.f, UIMax = 1.f))
+	float RandomThreshold = 0.5f;
+
+	/** GeometryCollection for the selection */
+	UPROPERTY(meta = (DataflowInput, DataflowOutput))
+	FManagedArrayCollection Collection;
+
+	/** Array of the selected bone indicies */
+	UPROPERTY(meta = (DataflowOutput, DisplayName = "TransformSelection"))
+	FDataflowTransformSelection TransformSelection;
+
+	FCollectionTransformSelectionRandomDataflowNode(const Dataflow::FNodeParameters& InParam, FGuid InGuid = FGuid::NewGuid())
+		: FDataflowNode(InParam, InGuid)
+	{
+		RandomSeed = FMath::FRandRange(-1e5, 1e5);
+		RegisterInputConnection(&Collection);
+		RegisterInputConnection(&RandomSeed);
+		RegisterInputConnection(&RandomThreshold);
+		RegisterOutputConnection(&TransformSelection);
+		RegisterOutputConnection(&Collection);
+	}
+
+	virtual void Evaluate(Dataflow::FContext& Context, const FDataflowOutput* Out) const override;
+
+};
+
+
+/**
+ *
+ * Selects the root bones in the Collection
+ *
+ */
+USTRUCT()
+struct FCollectionTransformSelectionRootDataflowNode : public FDataflowNode
+{
+	GENERATED_USTRUCT_BODY()
+	DATAFLOW_NODE_DEFINE_INTERNAL(FCollectionTransformSelectionRootDataflowNode, "CollectionTransformSelectRoot", "GeometryCollection|Selection", "")
+
+public:
+	/** GeometryCollection for the selection */
+	UPROPERTY(meta = (DataflowInput, DataflowOutput))
+	FManagedArrayCollection Collection;
+
+	/** Array of the selected bone indicies */
+	UPROPERTY(meta = (DataflowOutput, DisplayName = "TransformSelection"))
+	FDataflowTransformSelection TransformSelection;
+
+	FCollectionTransformSelectionRootDataflowNode(const Dataflow::FNodeParameters& InParam, FGuid InGuid = FGuid::NewGuid())
+		: FDataflowNode(InParam, InGuid)
+	{
+		RegisterInputConnection(&Collection);
+		RegisterOutputConnection(&TransformSelection);
+		RegisterOutputConnection(&Collection);
+	}
+
+	virtual void Evaluate(Dataflow::FContext& Context, const FDataflowOutput* Out) const override;
+
+};
+
+
+/**
+ * 
+ * Selects specified bones in the GeometryCollection by using a 
+ * space separated list
+ *
+ */
+USTRUCT(meta = (DataflowDeprecated, DataflowTerminal))
+struct FCollectionTransformSelectionCustomDataflowNode : public FDataflowNode
+{
+	GENERATED_USTRUCT_BODY()
+	DATAFLOW_NODE_DEFINE_INTERNAL(FCollectionTransformSelectionCustomDataflowNode, "CollectionTransformSelectCustom", "GeometryCollection|Selection", "")
+
+public:
+	/** GeometryCollection for the selection */
+	UPROPERTY(meta = (DataflowInput, DataflowOutput, DataflowPassthrough = "Collection", DataflowIntrinsic))
+	FManagedArrayCollection Collection;
+
+	/** Space separated list of bone indicies to specify the selection */
+	UPROPERTY(EditAnywhere, Category = "Selection")
+	FString BoneIndicies = FString();
+
+	/** Array of the selected bone indicies */
+	UPROPERTY(meta = (DataflowOutput, DisplayName = "TransformSelection"))
+	FDataflowTransformSelection TransformSelection;
+
+	FCollectionTransformSelectionCustomDataflowNode(const Dataflow::FNodeParameters& InParam, FGuid InGuid = FGuid::NewGuid())
+		: FDataflowNode(InParam, InGuid)
+	{
+		RegisterInputConnection(&Collection);
+		RegisterInputConnection(&BoneIndicies);
+		RegisterOutputConnection(&TransformSelection);
+		RegisterOutputConnection(&Collection, &Collection);
+	}
+
+	virtual void Evaluate(Dataflow::FContext& Context, const FDataflowOutput* Out) const override;
+
+};
+
+
+/**
+ *
+ * Selects the parents of the currently selected bones
+ *
+ */
+USTRUCT(meta = (DataflowDeprecated, DataflowTerminal))
+struct FCollectionTransformSelectionParentDataflowNode : public FDataflowNode
+{
+	GENERATED_USTRUCT_BODY()
+	DATAFLOW_NODE_DEFINE_INTERNAL(FCollectionTransformSelectionParentDataflowNode, "CollectionTransformSelectParent", "GeometryCollection|Selection", "")
+
+public:
+	/** Array of the selected bone indicies */
+	UPROPERTY(meta = (DataflowInput, DataflowOutput, DisplayName = "TransformSelection", DataflowPassthrough = "TransformSelection", DataflowIntrinsic))
+	FDataflowTransformSelection TransformSelection;
+
+	/** GeometryCollection for the selection */
+	UPROPERTY(meta = (DataflowInput, DataflowOutput, DataflowPassthrough = "Collection", DataflowIntrinsic))
+	FManagedArrayCollection Collection;
+
+	FCollectionTransformSelectionParentDataflowNode(const Dataflow::FNodeParameters& InParam, FGuid InGuid = FGuid::NewGuid())
+		: FDataflowNode(InParam, InGuid)
+	{
+		RegisterInputConnection(&TransformSelection);
+		RegisterInputConnection(&Collection);
+		RegisterOutputConnection(&TransformSelection, &TransformSelection);
+		RegisterOutputConnection(&Collection, &Collection);
+	}
+
+	virtual void Evaluate(Dataflow::FContext& Context, const FDataflowOutput* Out) const override;
+
+};
+
+
+/**
+ *
+ * Outputs the specified percentage of the incoming bone selection
+ *
+ */
+USTRUCT(meta = (DataflowDeprecated, DataflowTerminal))
+struct FCollectionTransformSelectionByPercentageDataflowNode : public FDataflowNode
+{
+	GENERATED_USTRUCT_BODY()
+	DATAFLOW_NODE_DEFINE_INTERNAL(FCollectionTransformSelectionByPercentageDataflowNode, "CollectionTransformSelectByPercentage", "GeometryCollection|Selection", "")
+
+public:
+	/** Array of the selected bone indicies */
+	UPROPERTY(meta = (DataflowInput, DataflowOutput, DisplayName = "TransformSelection", DataflowPassthrough = "TransformSelection", DataflowIntrinsic))
+	FDataflowTransformSelection TransformSelection;
+
+	/** GeometryCollection for the selection */
+	UPROPERTY(meta = (DataflowInput, DataflowOutput, DataflowPassthrough = "Collection", DataflowIntrinsic))
+	FManagedArrayCollection Collection;
+
+	/** Percentage to keep from the original selection */
+	UPROPERTY(EditAnywhere, Category = "Selection", meta = (UIMin = 0, UIMax = 100))
+	int32 Percentage = 100;
+
+	/** Sets the random generation to deterministic */
+	UPROPERTY(EditAnywhere, Category = "Random")
+	bool Deterministic = false;
+
+	/** Seed value for the random generation */
+	UPROPERTY(EditAnywhere, Category = "Random", meta = (DataflowInput, EditCondition = "Deterministic"))
+	float RandomSeed = 0.f;
+
+	FCollectionTransformSelectionByPercentageDataflowNode(const Dataflow::FNodeParameters& InParam, FGuid InGuid = FGuid::NewGuid())
+		: FDataflowNode(InParam, InGuid)
+	{
+		RandomSeed = FMath::RandRange(-100000, 100000);
+		RegisterInputConnection(&TransformSelection);
+		RegisterInputConnection(&Collection);
+		RegisterInputConnection(&Percentage);
+		RegisterInputConnection(&RandomSeed);
+		RegisterOutputConnection(&TransformSelection, &TransformSelection);
+		RegisterOutputConnection(&Collection, &Collection);
+	}
+
+	virtual void Evaluate(Dataflow::FContext& Context, const FDataflowOutput* Out) const override;
+
+};
+
+
+/**
+ *
+ * Selects the children of the incoming bone selection
+ *
+ */
+USTRUCT(meta = (DataflowDeprecated, DataflowTerminal))
+struct FCollectionTransformSelectionChildrenDataflowNode : public FDataflowNode
+{
+	GENERATED_USTRUCT_BODY()
+	DATAFLOW_NODE_DEFINE_INTERNAL(FCollectionTransformSelectionChildrenDataflowNode, "CollectionTransformSelectChildren", "GeometryCollection|Selection", "")
+
+public:
+	/** Array of the selected bone indicies */
+	UPROPERTY(meta = (DataflowInput, DataflowOutput, DisplayName = "TransformSelection", DataflowPassthrough = "TransformSelection", DataflowIntrinsic))
+	FDataflowTransformSelection TransformSelection;
+
+	/** GeometryCollection for the selection */
+	UPROPERTY(meta = (DataflowInput, DataflowOutput, DataflowPassthrough = "Collection", DataflowIntrinsic))
+	FManagedArrayCollection Collection;
+
+	FCollectionTransformSelectionChildrenDataflowNode(const Dataflow::FNodeParameters& InParam, FGuid InGuid = FGuid::NewGuid())
+		: FDataflowNode(InParam, InGuid)
+	{
+		RegisterInputConnection(&TransformSelection);
+		RegisterInputConnection(&Collection);
+		RegisterOutputConnection(&TransformSelection, &TransformSelection);
+		RegisterOutputConnection(&Collection, &Collection);
+	}
+
+	virtual void Evaluate(Dataflow::FContext& Context, const FDataflowOutput* Out) const override;
+
+};
+
+
+/**
+ *
+ * Selects the siblings of the incoming bone selection
+ *
+ */
+USTRUCT(meta = (DataflowDeprecated, DataflowTerminal))
+struct FCollectionTransformSelectionSiblingsDataflowNode : public FDataflowNode
+{
+	GENERATED_USTRUCT_BODY()
+	DATAFLOW_NODE_DEFINE_INTERNAL(FCollectionTransformSelectionSiblingsDataflowNode, "CollectionTransformSelectSiblings", "GeometryCollection|Selection", "")
+
+public:
+	/** Array of the selected bone indicies */
+	UPROPERTY(meta = (DataflowInput, DataflowOutput, DisplayName = "TransformSelection", DataflowPassthrough = "TransformSelection", DataflowIntrinsic))
+	FDataflowTransformSelection TransformSelection;
+
+	/** GeometryCollection for the selection */
+	UPROPERTY(meta = (DataflowInput, DataflowOutput, DataflowPassthrough = "Collection", DataflowIntrinsic))
+	FManagedArrayCollection Collection;
+
+	FCollectionTransformSelectionSiblingsDataflowNode(const Dataflow::FNodeParameters& InParam, FGuid InGuid = FGuid::NewGuid())
+		: FDataflowNode(InParam, InGuid)
+	{
+		RegisterInputConnection(&TransformSelection);
+		RegisterInputConnection(&Collection);
+		RegisterOutputConnection(&TransformSelection, &TransformSelection);
+		RegisterOutputConnection(&Collection, &Collection);
+	}
+
+	virtual void Evaluate(Dataflow::FContext& Context, const FDataflowOutput* Out) const override;
+
+};
+
+
+/**
+ *
+ * Selects the level of the incoming bone selection
+ *
+ */
+USTRUCT(meta = (DataflowDeprecated, DataflowTerminal))
+struct FCollectionTransformSelectionLevelDataflowNode : public FDataflowNode
+{
+	GENERATED_USTRUCT_BODY()
+	DATAFLOW_NODE_DEFINE_INTERNAL(FCollectionTransformSelectionLevelDataflowNode, "CollectionTransformSelectLevel", "GeometryCollection|Selection", "")
+
+public:
+	/** Array of the selected bone indicies */
+	UPROPERTY(meta = (DataflowInput, DataflowOutput, DisplayName = "TransformSelection", DataflowPassthrough = "TransformSelection", DataflowIntrinsic))
+	FDataflowTransformSelection TransformSelection;
+
+	/** GeometryCollection for the selection */
+	UPROPERTY(meta = (DataflowInput, DataflowOutput, DataflowPassthrough = "Collection", DataflowIntrinsic))
+	FManagedArrayCollection Collection;
+
+	FCollectionTransformSelectionLevelDataflowNode(const Dataflow::FNodeParameters& InParam, FGuid InGuid = FGuid::NewGuid())
+		: FDataflowNode(InParam, InGuid)
+	{
+		RegisterInputConnection(&TransformSelection);
+		RegisterInputConnection(&Collection);
+		RegisterOutputConnection(&TransformSelection, &TransformSelection);
+		RegisterOutputConnection(&Collection, &Collection);
+	}
+
+	virtual void Evaluate(Dataflow::FContext& Context, const FDataflowOutput* Out) const override;
+
+};
+
+
+/**
+ *
+ * Selects the contact(s) of the incoming bone selection
+ *
+ */
+USTRUCT(meta = (DataflowDeprecated, DataflowTerminal))
+struct FCollectionTransformSelectionContactDataflowNode : public FDataflowNode
+{
+	GENERATED_USTRUCT_BODY()
+	DATAFLOW_NODE_DEFINE_INTERNAL(FCollectionTransformSelectionContactDataflowNode, "CollectionTransformSelectContact", "GeometryCollection|Selection", "")
+
+public:
+	/** Array of the selected bone indicies */
+	UPROPERTY(meta = (DataflowInput, DataflowOutput, DisplayName = "TransformSelection", DataflowPassthrough = "TransformSelection", DataflowIntrinsic))
+	FDataflowTransformSelection TransformSelection;
+
+	/** GeometryCollection for the selection */
+	UPROPERTY(meta = (DataflowInput, DataflowOutput, DataflowPassthrough = "Collection", DataflowIntrinsic))
+	FManagedArrayCollection Collection;
+
+	FCollectionTransformSelectionContactDataflowNode(const Dataflow::FNodeParameters& InParam, FGuid InGuid = FGuid::NewGuid())
+		: FDataflowNode(InParam, InGuid)
+	{
+		RegisterInputConnection(&TransformSelection);
+		RegisterInputConnection(&Collection);
+		RegisterOutputConnection(&TransformSelection, &TransformSelection);
+		RegisterOutputConnection(&Collection, &Collection);
+	}
+
+	virtual void Evaluate(Dataflow::FContext& Context, const FDataflowOutput* Out) const override;
+
+};
+
+
+/**
+ *
+ * Selects the leaves in the GeometryCollection
+ *
+ */
+USTRUCT(meta = (DataflowDeprecated, DataflowTerminal))
+struct FCollectionTransformSelectionLeafDataflowNode : public FDataflowNode
+{
+	GENERATED_USTRUCT_BODY()
+	DATAFLOW_NODE_DEFINE_INTERNAL(FCollectionTransformSelectionLeafDataflowNode, "CollectionTransformSelectLeaf", "GeometryCollection|Selection", "")
+
+public:
+	/** GeometryCollection for the selection */
+	UPROPERTY(meta = (DataflowInput, DataflowOutput, DataflowPassthrough = "Collection", DataflowIntrinsic))
+	FManagedArrayCollection Collection;
+
+	/** Array of the selected bone indicies */
+	UPROPERTY(meta = (DataflowOutput, DisplayName = "TransformSelection"))
+	FDataflowTransformSelection TransformSelection;
+
+	FCollectionTransformSelectionLeafDataflowNode(const Dataflow::FNodeParameters& InParam, FGuid InGuid = FGuid::NewGuid())
+		: FDataflowNode(InParam, InGuid)
+	{
+		RegisterInputConnection(&Collection);
+		RegisterOutputConnection(&TransformSelection);
+		RegisterOutputConnection(&Collection, &Collection);
+	}
+
+	virtual void Evaluate(Dataflow::FContext& Context, const FDataflowOutput* Out) const override;
+
+};
+
+
+/**
+ *
+ * Selects the clusters in the GeometryCollection
+ *
+ */
+USTRUCT(meta = (DataflowDeprecated, DataflowTerminal))
+struct FCollectionTransformSelectionClusterDataflowNode : public FDataflowNode
+{
+	GENERATED_USTRUCT_BODY()
+	DATAFLOW_NODE_DEFINE_INTERNAL(FCollectionTransformSelectionClusterDataflowNode, "CollectionTransformSelectCluster", "GeometryCollection|Selection", "")
+
+public:
+	/** GeometryCollection for the selection */
+	UPROPERTY(meta = (DataflowInput, DataflowOutput, DataflowPassthrough = "Collection", DataflowIntrinsic))
+	FManagedArrayCollection Collection;
+
+	/** Array of the selected bone indicies */
+	UPROPERTY(meta = (DataflowOutput, DisplayName = "TransformSelection"))
+	FDataflowTransformSelection TransformSelection;
+
+	FCollectionTransformSelectionClusterDataflowNode(const Dataflow::FNodeParameters& InParam, FGuid InGuid = FGuid::NewGuid())
+		: FDataflowNode(InParam, InGuid)
+	{
+		RegisterInputConnection(&Collection);
+		RegisterOutputConnection(&TransformSelection);
+		RegisterOutputConnection(&Collection, &Collection);
+	}
+
+	virtual void Evaluate(Dataflow::FContext& Context, const FDataflowOutput* Out) const override;
+
+};
+
+
+/**
+ *
+ * Selects pieces based on their size
+ *
+ */
+USTRUCT(meta = (DataflowDeprecated, DataflowTerminal))
+struct FCollectionTransformSelectionBySizeDataflowNode : public FDataflowNode
+{
+	GENERATED_USTRUCT_BODY()
+	DATAFLOW_NODE_DEFINE_INTERNAL(FCollectionTransformSelectionBySizeDataflowNode, "CollectionTransformSelectBySize", "GeometryCollection|Selection", "")
+
+public:
+	/** GeometryCollection for the selection */
+	UPROPERTY(meta = (DataflowInput, DataflowOutput, DataflowPassthrough = "Collection", DataflowIntrinsic))
+	FManagedArrayCollection Collection;
+
+	/** Minimum size for the selection */
+	UPROPERTY(EditAnywhere, Category = "Size", meta = (UIMin = 0.f, UIMax = 1e9))
+	float SizeMin = 0.f;
+
+	/** Maximum size for the selection */
+	UPROPERTY(EditAnywhere, Category = "Size", meta = (UIMin = 0.f, UIMax = 1e9))
+	float SizeMax = 1000.f;
+
+	/** Array of the selected bone indicies */
+	UPROPERTY(meta = (DataflowOutput, DisplayName = "TransformSelection"))
+	FDataflowTransformSelection TransformSelection;
+
+	FCollectionTransformSelectionBySizeDataflowNode(const Dataflow::FNodeParameters& InParam, FGuid InGuid = FGuid::NewGuid())
+		: FDataflowNode(InParam, InGuid)
+	{
+		RegisterInputConnection(&Collection);
+		RegisterInputConnection(&SizeMin);
+		RegisterInputConnection(&SizeMax);
+		RegisterOutputConnection(&TransformSelection);
+		RegisterOutputConnection(&Collection, &Collection);
+	}
+
+	virtual void Evaluate(Dataflow::FContext& Context, const FDataflowOutput* Out) const override;
+
+};
+
+
+/**
+ *
+ * Selects pieces based on their volume
+ *
+ */
+USTRUCT(meta = (DataflowDeprecated, DataflowTerminal))
+struct FCollectionTransformSelectionByVolumeDataflowNode : public FDataflowNode
+{
+	GENERATED_USTRUCT_BODY()
+	DATAFLOW_NODE_DEFINE_INTERNAL(FCollectionTransformSelectionByVolumeDataflowNode, "CollectionTransformSelectByVolume", "GeometryCollection|Selection", "")
+
+public:
+	/** GeometryCollection for the selection */
+	UPROPERTY(meta = (DataflowInput, DataflowOutput, DataflowPassthrough = "Collection", DataflowIntrinsic))
+	FManagedArrayCollection Collection;
+
+	/** Minimum volume for the selection */
+	UPROPERTY(EditAnywhere, Category = "Volume", meta = (UIMin = 0.f, UIMax = 1e9))
+	float VolumeMin = 0.f;
+
+	/** Maximum volume for the selection */
+	UPROPERTY(EditAnywhere, Category = "Volume", meta = (UIMin = 0.f, UIMax = 1e9))
+	float VolumeMax = 1000.f;
+
+	/** Array of the selected bone indicies */
+	UPROPERTY(meta = (DataflowOutput, DisplayName = "TransformSelection"))
+	FDataflowTransformSelection TransformSelection;
+
+	FCollectionTransformSelectionByVolumeDataflowNode(const Dataflow::FNodeParameters& InParam, FGuid InGuid = FGuid::NewGuid())
+		: FDataflowNode(InParam, InGuid)
+	{
+		RegisterInputConnection(&Collection);
+		RegisterInputConnection(&VolumeMin);
+		RegisterInputConnection(&VolumeMax);
+		RegisterOutputConnection(&TransformSelection);
+		RegisterOutputConnection(&Collection, &Collection);
+	}
+
+	virtual void Evaluate(Dataflow::FContext& Context, const FDataflowOutput* Out) const override;
+
+};
+
+
+namespace Dataflow
+{
+	void GeometryCollectionSelectionNodes();
+}
+
diff --git a/Engine/Source/Runtime/Experimental/Dataflow/Core/Private/Dataflow/DataflowNode.cpp b/Engine/Source/Runtime/Experimental/Dataflow/Core/Private/Dataflow/DataflowNode.cpp
--- a/Engine/Source/Runtime/Experimental/Dataflow/Core/Private/Dataflow/DataflowNode.cpp
+++ b/Engine/Source/Runtime/Experimental/Dataflow/Core/Private/Dataflow/DataflowNode.cpp
@@ -8,7 +8,10 @@
 
 #include UE_INLINE_GENERATED_CPP_BY_NAME(DataflowNode)
 
-
+const FName FDataflowNode::DataflowInput = TEXT("DataflowInput");
+const FName FDataflowNode::DataflowOutput = TEXT("DataflowOutput");
+const FName FDataflowNode::DataflowPassthrough = TEXT("DataflowPassthrough");
+const FName FDataflowNode::DataflowIntrinsic = TEXT("DataflowIntrinsic");
 
 //
 // Inputs
@@ -293,7 +296,7 @@ bool FDataflowNode::ValidateConnections()
 				const FProperty* Property = *PropertyIt;
 				FName PropName(Property->GetName());
 #if WITH_EDITORONLY_DATA
-				if (Property->HasMetaData(TEXT("DataflowInput")))
+				if (Property->HasMetaData(FDataflowNode::DataflowInput))
 				{
 					if (!FindInput(PropName))
 					{
@@ -301,7 +304,7 @@ bool FDataflowNode::ValidateConnections()
 						bValid = false;
 					}
 				}
-				if (Property->HasMetaData(TEXT("DataflowOutput")))
+				if (Property->HasMetaData(FDataflowNode::DataflowOutput))
 				{
 					const FDataflowOutput* OutputConnection = FindOutput(PropName);
 					if(!OutputConnection)
@@ -311,12 +314,12 @@ bool FDataflowNode::ValidateConnections()
 					}
 
 					// Validate passthrough connections if they exist
-					if (const FString* PassthroughName = Property->FindMetaData(TEXT("Passthrough")))
+					if (const FString* PassthroughName = Property->FindMetaData(FDataflowNode::DataflowPassthrough))
 					{
 						void* PassthroughConnectionAddress = OutputConnection->GetPassthroughRealAddress();
 						if(PassthroughConnectionAddress == nullptr)
 						{
-							UE_LOG(LogChaos, Warning, TEXT("Missing Passthrough registration for (%s:%s)"), *GetName().ToString(),*PropName.ToString());
+							UE_LOG(LogChaos, Warning, TEXT("Missing DataflowPassthrough registration for (%s:%s)"), *GetName().ToString(),*PropName.ToString());
 							bValid = false;
 						}
 
@@ -325,25 +328,25 @@ bool FDataflowNode::ValidateConnections()
 
 						if(PassthroughConnectionInputFromArg != PassthroughConnectionInput)
 						{
-							UE_LOG(LogChaos, Warning, TEXT("Mismatch in declared and registered Passthrough connection; (%s:%s vs %s)"), *GetName().ToString(), *PropName.ToString(), *PassthroughConnectionInputFromArg->GetName().ToString());
+							UE_LOG(LogChaos, Warning, TEXT("Mismatch in declared and registered DataflowPassthrough connection; (%s:%s vs %s)"), *GetName().ToString(), *PropName.ToString(), *PassthroughConnectionInputFromArg->GetName().ToString());
 							bValid = false;
 						}
 
 						if(!PassthroughConnectionInput)
 						{
-							UE_LOG(LogChaos, Warning, TEXT("Incorrect Passthrough Connection set for (%s:%s)"), *GetName().ToString(),*PropName.ToString());
+							UE_LOG(LogChaos, Warning, TEXT("Incorrect DataflowPassthrough Connection set for (%s:%s)"), *GetName().ToString(),*PropName.ToString());
 							bValid = false;
 						}
 
 						else if(OutputConnection->GetType() != PassthroughConnectionInput->GetType())
 						{
-							UE_LOG(LogChaos, Warning, TEXT("Passthrough connection types mismatch for (%s:%s)"), *GetName().ToString(),*PropName.ToString());
+							UE_LOG(LogChaos, Warning, TEXT("DataflowPassthrough connection types mismatch for (%s:%s)"), *GetName().ToString(),*PropName.ToString());
 							bValid = false;
 						}
 					}
 					else if(OutputConnection->GetPassthroughRealAddress()) 
 					{
-						UE_LOG(LogChaos, Warning, TEXT("Missing Passthrough decleration for (%s:%s)"), *GetName().ToString(),*PropName.ToString());
+						UE_LOG(LogChaos, Warning, TEXT("Missing DataflowPassthrough decleration for (%s:%s)"), *GetName().ToString(),*PropName.ToString());
 						bValid = false;
 					}
 				}
@@ -354,5 +357,158 @@ bool FDataflowNode::ValidateConnections()
 	return bValid;
 }
 
+FString FDataflowNode::GetToolTip()
+{
+	if (const TUniquePtr<FStructOnScope> ScriptOnStruct = TUniquePtr<FStructOnScope>(NewStructOnScope()))
+	{
+#if WITH_EDITORONLY_DATA
+		if (const UStruct* Struct = ScriptOnStruct->GetStruct())
+		{
+			FString OutStr, InputsStr, OutputsStr;
+			
+			FText StructText = Struct->GetToolTipText();
+
+			OutStr.Appendf(TEXT("%s\n\n%s\n"), *GetDisplayName().ToString(), *StructText.ToString());
+			
+			for (TFieldIterator<FProperty> PropertyIt(Struct); PropertyIt; ++PropertyIt)
+			{
+				const FProperty* Property = *PropertyIt;
+
+				if (Property->HasMetaData(TEXT("Tooltip")))
+				{
+					FString ToolTipStr = Property->GetToolTipText(true).ToString();
+					if (ToolTipStr.Len() > 0)
+					{
+						TArray<FString> OutArr;
+						ToolTipStr.ParseIntoArray(OutArr, TEXT(":\r\n"));
+
+						if (Property->HasMetaData(FDataflowNode::DataflowInput) &&
+							Property->HasMetaData(FDataflowNode::DataflowOutput) &&
+							Property->HasMetaData(FDataflowNode::DataflowPassthrough))
+						{
+							if (Property->HasMetaData(FDataflowNode::DataflowIntrinsic))
+							{
+								InputsStr.Appendf(TEXT("    %s [Intrinsic] - %s\n"), *Property->GetName(), *OutArr[1]);
+							}
+							else
+							{
+								InputsStr.Appendf(TEXT("    %s - %s\n"), *Property->GetName(), *OutArr[1]);
+							}
+
+							OutputsStr.Appendf(TEXT("    %s [Passthrough] - %s\n"), *Property->GetName(), *OutArr[1]);
+						}					
+						else if (Property->HasMetaData(FDataflowNode::DataflowInput))
+						{
+							if (Property->HasMetaData(FDataflowNode::DataflowIntrinsic))
+							{
+								InputsStr.Appendf(TEXT("    %s [Intrinsic] - %s\n"), *Property->GetName(), *OutArr[1]);
+							}
+							else
+							{
+								InputsStr.Appendf(TEXT("    %s - %s\n"), *Property->GetName(), *OutArr[1]);
+							}
+						}
+						else if (Property->HasMetaData(FDataflowNode::DataflowOutput))
+						{
+							OutputsStr.Appendf(TEXT("    %s - %s\n"), *Property->GetName(), *OutArr[1]);
+						}
+						
+						if (!Property->HasMetaData(FDataflowNode::DataflowInput) && !Property->HasMetaData(FDataflowNode::DataflowOutput))
+						{
+							InputsStr.Appendf(TEXT("    %s - %s\n"), *Property->GetName(), *OutArr[1]);
+						}
+					}
+				}
+			}
+
+			if (InputsStr.Len() > 0)
+			{
+				OutStr.Appendf(TEXT("\n Input(s) :\n % s"), *InputsStr);
+			}
+			if (OutputsStr.Len() > 0)
+			{
+				OutStr.Appendf(TEXT("\n Output(s):\n%s"), *OutputsStr);
+			}
+
+			return OutStr;
+		}
+#endif
+	}
+
+	return "";
+}
+
+FString FDataflowNode::GetPinToolTip(const FName& PropertyName)
+{
+	if (const TUniquePtr<FStructOnScope> ScriptOnStruct = TUniquePtr<FStructOnScope>(NewStructOnScope()))
+	{
+		if (const UStruct* Struct = ScriptOnStruct->GetStruct())
+		{
+			for (TFieldIterator<FProperty> PropertyIt(Struct); PropertyIt; ++PropertyIt)
+			{
+				const FProperty* Property = *PropertyIt;
+
+#if WITH_EDITORONLY_DATA
+				if (Property->GetName() == PropertyName.ToString())
+				{
+					if (Property->HasMetaData(TEXT("Tooltip")))
+					{
+						FString ToolTipStr = Property->GetToolTipText(true).ToString();
+						if (ToolTipStr.Len() > 0)
+						{
+							TArray<FString> OutArr;
+							int32 NumElems = ToolTipStr.ParseIntoArray(OutArr, TEXT(":\r\n"));
+
+							if (NumElems == 2)
+							{
+								return OutArr[1];
+							}
+						}
+					}
+				}
+#endif
+			}
+		}
+	}
+
+	return "";
+}
+
+TArray<FString> FDataflowNode::GetPinMetaData(const FName& PropertyName)
+{
+	if (const TUniquePtr<FStructOnScope> ScriptOnStruct = TUniquePtr<FStructOnScope>(NewStructOnScope()))
+	{
+		if (const UStruct* Struct = ScriptOnStruct->GetStruct())
+		{
+			for (TFieldIterator<FProperty> PropertyIt(Struct); PropertyIt; ++PropertyIt)
+			{
+				const FProperty* Property = *PropertyIt;
+
+#if WITH_EDITORONLY_DATA
+				if (Property->GetName() == PropertyName.ToString())
+				{
+					TArray<FString> MetaDataStrArr;
+					if (Property->HasMetaData(FDataflowNode::DataflowPassthrough))
+					{
+						MetaDataStrArr.Add("Passthrough");
+					}
+					if (Property->HasMetaData(FDataflowNode::DataflowIntrinsic))
+					{
+						MetaDataStrArr.Add("Intrinsic");
+					}
+
+					return MetaDataStrArr;
+				}
+#endif
+			}
+		}
+	}
+
+	return TArray<FString>();
+}
+
+
+
+
 
 
diff --git a/Engine/Source/Runtime/Experimental/Dataflow/Core/Private/Dataflow/DataflowNodeFactory.cpp b/Engine/Source/Runtime/Experimental/Dataflow/Core/Private/Dataflow/DataflowNodeFactory.cpp
--- a/Engine/Source/Runtime/Experimental/Dataflow/Core/Private/Dataflow/DataflowNodeFactory.cpp
+++ b/Engine/Source/Runtime/Experimental/Dataflow/Core/Private/Dataflow/DataflowNodeFactory.cpp
@@ -16,6 +16,8 @@ namespace Dataflow
 			TUniquePtr<FDataflowNode> Node = ClassMap[Param.Type](Param);
 			if(Node->IsValid())
 			{
+				ParametersMap[Param.Type].ToolTip = Node->GetToolTip();
+
 				return Graph.AddNode(MoveTemp(Node));
 			}
 			
diff --git a/Engine/Source/Runtime/Experimental/Dataflow/Core/Private/Dataflow/DataflowSelection.cpp b/Engine/Source/Runtime/Experimental/Dataflow/Core/Private/Dataflow/DataflowSelection.cpp
new file mode 100644
--- /dev/null
+++ b/Engine/Source/Runtime/Experimental/Dataflow/Core/Private/Dataflow/DataflowSelection.cpp
@@ -0,0 +1,99 @@
+// Copyright Epic Games, Inc. All Rights Reserved.
+
+#include "Dataflow/DataflowSelection.h"
+
+void FDataflowSelection::Initialize(int32 NumBits, bool Value) 
+{ 
+	SelectionArray.AddUninitialized(NumBits); 
+	SelectionArray.SetRange(0, NumBits, Value); 
+}
+
+void FDataflowSelection::Initialize(const FDataflowSelection& Other)
+{
+	Initialize(Other.Num(), false);
+
+	for (int32 Idx = 0; Idx < Other.Num(); ++Idx)
+	{
+		if (Other.IsSelected(Idx))
+		{
+			SelectionArray[Idx] = true;
+		}
+	}
+}
+
+void FDataflowSelection::AsArray(TArray<int32>& SelectionArr) const
+{
+	SelectionArr.Empty();
+
+	TBitArray<>::FConstIterator It(SelectionArray);
+	while (It)
+	{
+		int32 Index = It.GetIndex();
+
+		if (SelectionArray[Index])
+		{
+			SelectionArr.Add(Index);
+		}
+
+		++It;
+	}
+}
+
+void FDataflowSelection::SetFromArray(const TArray<int32>& SelectionArr)
+{
+	SelectionArray.Init(false, SelectionArray.Num());
+
+	for (int32 Elem : SelectionArr)
+	{
+		SelectionArray[Elem] = true;
+	}
+}
+
+void FDataflowSelection::AND(const FDataflowSelection& Other, FDataflowSelection& Result) const
+{ 
+	Result.SelectionArray = TBitArray<>::BitwiseAND(SelectionArray, Other.SelectionArray, EBitwiseOperatorFlags::MaxSize);
+}
+
+void FDataflowSelection::OR(const FDataflowSelection& Other, FDataflowSelection& Result) const
+{
+	Result.SelectionArray = TBitArray<>::BitwiseOR(SelectionArray, Other.SelectionArray, EBitwiseOperatorFlags::MaxSize);
+}
+
+void FDataflowSelection::XOR(const FDataflowSelection& Other, FDataflowSelection& Result) const
+{
+	Result.SelectionArray = TBitArray<>::BitwiseXOR(SelectionArray, Other.SelectionArray, EBitwiseOperatorFlags::MaxSize);
+}
+
+int32 FDataflowSelection::NumSelected()
+{
+	int32 NumSelectedBits = 0;
+
+	TBitArray<>::FConstIterator It(SelectionArray);
+	while (It)
+	{		 
+		if (It.GetValue())
+		{
+			NumSelectedBits++;
+		}
+
+		++It;
+	}
+
+	return NumSelectedBits;
+}
+
+bool FDataflowSelection::AnySelected()
+{
+	TBitArray<>::FConstIterator It(SelectionArray);
+	while (It)
+	{
+		if (It.GetValue())
+		{
+			return true;
+		}
+
+		++It;
+	}
+
+	return false;
+}
diff --git a/Engine/Source/Runtime/Experimental/Dataflow/Core/Public/Dataflow/DataflowNode.h b/Engine/Source/Runtime/Experimental/Dataflow/Core/Public/Dataflow/DataflowNode.h
--- a/Engine/Source/Runtime/Experimental/Dataflow/Core/Public/Dataflow/DataflowNode.h
+++ b/Engine/Source/Runtime/Experimental/Dataflow/Core/Public/Dataflow/DataflowNode.h
@@ -75,7 +75,9 @@ struct DATAFLOWCORE_API FDataflowNode
 	virtual FName GetDisplayName() const { return ""; }
 	virtual FName GetCategory() const { return ""; }
 	virtual FString GetTags() const { return ""; }
-	virtual FString GetToolTip() const { return ""; }
+	virtual FString GetToolTip();
+	FString GetPinToolTip(const FName& PropertyName);
+	TArray<FString> GetPinMetaData(const FName& PropertyName);
 	virtual TArray<Dataflow::FRenderingParameter> GetRenderParameters() const { return GetRenderParametersImpl(); }
 
 	//
@@ -104,7 +106,10 @@ struct DATAFLOWCORE_API FDataflowNode
 	const FDataflowOutput* FindOutput(FName Name) const;
 	const FDataflowOutput* FindOutput(const void* Reference) const;
 
-
+	static const FName DataflowInput;
+	static const FName DataflowOutput;
+	static const FName DataflowPassthrough;
+	static const FName DataflowIntrinsic;
 
 	//
 	//  Struct Support
@@ -243,7 +248,6 @@ public:																				\
 	virtual FName GetDisplayName() const override { return TYPE::StaticDisplay(); }	\
 	virtual FName GetCategory() const override { return TYPE::StaticCategory(); }	\
 	virtual FString GetTags() const override { return TYPE::StaticTags(); }			\
-	virtual FString GetToolTip() const override { return TYPE::StaticToolTip(); }	\
 	virtual const UScriptStruct* TypedScriptStruct() const override					\
 		{return TYPE::StaticStruct();}												\
 	TYPE() {}																		\
diff --git a/Engine/Source/Runtime/Experimental/Dataflow/Core/Public/Dataflow/DataflowSelection.h b/Engine/Source/Runtime/Experimental/Dataflow/Core/Public/Dataflow/DataflowSelection.h
new file mode 100644
--- /dev/null
+++ b/Engine/Source/Runtime/Experimental/Dataflow/Core/Public/Dataflow/DataflowSelection.h
@@ -0,0 +1,35 @@
+// Copyright Epic Games, Inc. All Rights Reserved.
+
+#pragma once
+
+#include "DataflowSelection.generated.h"
+
+USTRUCT()
+struct DATAFLOWCORE_API FDataflowSelection
+{
+	GENERATED_USTRUCT_BODY()
+
+	void Initialize(int32 NumBits, bool Value);
+	void Initialize(const FDataflowSelection& Other);
+	int32 Num() const { return SelectionArray.Num(); }
+	int32 NumSelected();
+	bool AnySelected();
+	bool IsSelected(int32 Idx) const { return SelectionArray[Idx]; }
+	void SetSelected(int32 Idx) { SelectionArray[Idx] = true; }
+	void SetNotSelected(int32 Idx) { SelectionArray[Idx] = false; }
+	void AsArray(TArray<int32>& SelectionArr) const;
+	void SetFromArray(const TArray<int32>& SelectionArr);
+	void AND(const FDataflowSelection& Other, FDataflowSelection& Result) const;
+	void OR(const FDataflowSelection& Other, FDataflowSelection& Result) const;
+	void XOR(const FDataflowSelection& Other, FDataflowSelection& Result) const;
+	void Invert() { SelectionArray.BitwiseNOT(); }
+
+private:
+	TBitArray<> SelectionArray;
+};
+
+USTRUCT()
+struct DATAFLOWCORE_API FDataflowTransformSelection : public FDataflowSelection
+{
+	GENERATED_USTRUCT_BODY()
+};
diff --git a/Engine/Source/Runtime/Experimental/Dataflow/Engine/Private/Dataflow/DataflowEdNode.cpp b/Engine/Source/Runtime/Experimental/Dataflow/Engine/Private/Dataflow/DataflowEdNode.cpp
--- a/Engine/Source/Runtime/Experimental/Dataflow/Engine/Private/Dataflow/DataflowEdNode.cpp
+++ b/Engine/Source/Runtime/Experimental/Dataflow/Engine/Private/Dataflow/DataflowEdNode.cpp
@@ -199,6 +199,56 @@ FLinearColor UDataflowEdNode::GetNodeBodyTintColor() const
 	return FLinearColor(0.f, 0.f, 0.f);
 }
 
+FText UDataflowEdNode::GetTooltipText() const
+{
+	if (DataflowGraph)
+	{
+		if (DataflowNodeGuid.IsValid())
+		{
+			if (TSharedPtr<FDataflowNode> DataflowNode = DataflowGraph->FindBaseNode(DataflowNodeGuid))
+			{
+				return FText::FromString(DataflowNode->GetToolTip());
+			}
+		}
+	}
+
+	return FText::FromString("");
+
+}
+
+void UDataflowEdNode::GetPinHoverText(const UEdGraphPin& Pin, FString& HoverTextOut) const
+{
+	if (DataflowGraph)
+	{
+		if (DataflowNodeGuid.IsValid())
+		{
+			if (TSharedPtr<FDataflowNode> DataflowNode = DataflowGraph->FindBaseNode(DataflowNodeGuid))
+			{		
+				FString MetaDataStr;
+				
+				TArray<FString> PinMetaData = DataflowNode->GetPinMetaData(Pin.PinName);
+			
+				if (Pin.Direction == EGPD_Input && PinMetaData.Contains(FDataflowNode::DataflowIntrinsic.ToString()))
+				{
+					MetaDataStr = "[Intrinsic]";
+				}
+				if (Pin.Direction == EGPD_Output && PinMetaData.Contains(FDataflowNode::DataflowPassthrough.ToString()))
+				{
+					MetaDataStr = "[Passthrough]";
+				}
+
+				FString NameStr = Pin.PinName.ToString();
+				if (MetaDataStr.Len() > 0)
+				{
+					NameStr.Appendf(TEXT(" %s"), *MetaDataStr);
+				}
+
+				HoverTextOut.Appendf(TEXT("%s\n%s\n\n%s"), *NameStr, *Pin.PinType.PinCategory.ToString(), *DataflowNode->GetPinToolTip(Pin.PinName));
+			}
+		}
+	}
+}
+
 #endif
 
 
diff --git a/Engine/Source/Runtime/Experimental/Dataflow/Engine/Public/Dataflow/DataflowEdNode.h b/Engine/Source/Runtime/Experimental/Dataflow/Engine/Public/Dataflow/DataflowEdNode.h
--- a/Engine/Source/Runtime/Experimental/Dataflow/Engine/Public/Dataflow/DataflowEdNode.h
+++ b/Engine/Source/Runtime/Experimental/Dataflow/Engine/Public/Dataflow/DataflowEdNode.h
@@ -42,6 +42,8 @@ public:
 #if WITH_EDITOR
 	virtual FLinearColor GetNodeTitleColor() const override;
 	virtual FLinearColor GetNodeBodyTintColor() const override;
+	virtual FText GetTooltipText() const override;
+	virtual void GetPinHoverText(const UEdGraphPin& Pin, FString& HoverTextOut) const override;
 #endif // WITH_EDITOR
 	// End of UEdGraphNode interface
 
