[GeometryCollectionPlugin]
* Attempt at fixing compile error happening on non-editor platforms using this plugin.

diff --git a/Engine/Plugins/Experimental/GeometryCollectionPlugin/Source/GeometryCollectionNodes/Private/Dataflow/GeometryCollectionNodes.cpp b/Engine/Plugins/Experimental/GeometryCollectionPlugin/Source/GeometryCollectionNodes/Private/Dataflow/GeometryCollectionNodes.cpp
--- a/Engine/Plugins/Experimental/GeometryCollectionPlugin/Source/GeometryCollectionNodes/Private/Dataflow/GeometryCollectionNodes.cpp
+++ b/Engine/Plugins/Experimental/GeometryCollectionPlugin/Source/GeometryCollectionNodes/Private/Dataflow/GeometryCollectionNodes.cpp
@@ -1220,6 +1220,7 @@ void FStaticMeshToMeshDataflowNode::Evaluate(Dataflow::FContext& Context, const
 {
 	if (Out->IsA<TObjectPtr<UDynamicMesh>>(&Mesh))
 	{
+#if WITH_EDITORONLY_DATA
 		if (FMeshDescription* MeshDescription = UseHiRes ? StaticMesh->GetHiResMeshDescription() : StaticMesh->GetMeshDescription(LODLevel))
 		{
 			TObjectPtr<UDynamicMesh> DynamicMesh = NewObject<UDynamicMesh>();
@@ -1234,6 +1235,7 @@ void FStaticMeshToMeshDataflowNode::Evaluate(Dataflow::FContext& Context, const
 			SetValue<TObjectPtr<UDynamicMesh>>(Context, DynamicMesh, &Mesh);
 		}
 		else
+#endif
 		{
 			SetValue<TObjectPtr<UDynamicMesh>>(Context, NewObject<UDynamicMesh>(), &Mesh);
 		}
