Removed use of exponential notation since it isn't supported by C++ UHT.

diff --git a/Engine/Plugins/Experimental/GeometryCollectionPlugin/Source/GeometryCollectionNodes/Public/Dataflow/GeometryCollectionSelectionNodes.h b/Engine/Plugins/Experimental/GeometryCollectionPlugin/Source/GeometryCollectionNodes/Public/Dataflow/GeometryCollectionSelectionNodes.h
--- a/Engine/Plugins/Experimental/GeometryCollectionPlugin/Source/GeometryCollectionNodes/Public/Dataflow/GeometryCollectionSelectionNodes.h
+++ b/Engine/Plugins/Experimental/GeometryCollectionPlugin/Source/GeometryCollectionNodes/Public/Dataflow/GeometryCollectionSelectionNodes.h
@@ -618,11 +618,11 @@ public:
 	FManagedArrayCollection Collection;
 
 	/** Minimum size for the selection */
-	UPROPERTY(EditAnywhere, Category = "Size", meta = (UIMin = 0.f, UIMax = 1e9))
+	UPROPERTY(EditAnywhere, Category = "Size", meta = (UIMin = 0.f, UIMax = 1000000000.f))
 	float SizeMin = 0.f;
 
 	/** Maximum size for the selection */
-	UPROPERTY(EditAnywhere, Category = "Size", meta = (UIMin = 0.f, UIMax = 1e9))
+	UPROPERTY(EditAnywhere, Category = "Size", meta = (UIMin = 0.f, UIMax = 1000000000.f))
 	float SizeMax = 1000.f;
 
 	/** Array of the selected bone indicies */
@@ -661,11 +661,11 @@ public:
 	FManagedArrayCollection Collection;
 
 	/** Minimum volume for the selection */
-	UPROPERTY(EditAnywhere, Category = "Volume", meta = (UIMin = 0.f, UIMax = 1e9))
+	UPROPERTY(EditAnywhere, Category = "Volume", meta = (UIMin = 0.f, UIMax = 1000000000.f))
 	float VolumeMin = 0.f;
 
 	/** Maximum volume for the selection */
-	UPROPERTY(EditAnywhere, Category = "Volume", meta = (UIMin = 0.f, UIMax = 1e9))
+	UPROPERTY(EditAnywhere, Category = "Volume", meta = (UIMin = 0.f, UIMax = 1000000000.f))
 	float VolumeMax = 1000.f;
 
 	/** Array of the selected bone indicies */
