Dataflow : fix geometry collection terminal node

diff --git a/Engine/Plugins/Experimental/GeometryCollectionPlugin/Source/GeometryCollectionNodes/Public/Dataflow/Nodes/GeometryCollectionAssetNodes.h b/Engine/Plugins/Experimental/GeometryCollectionPlugin/Source/GeometryCollectionNodes/Public/Dataflow/Nodes/GeometryCollectionAssetNodes.h
--- a/Engine/Plugins/Experimental/GeometryCollectionPlugin/Source/GeometryCollectionNodes/Public/Dataflow/Nodes/GeometryCollectionAssetNodes.h
+++ b/Engine/Plugins/Experimental/GeometryCollectionPlugin/Source/GeometryCollectionNodes/Public/Dataflow/Nodes/GeometryCollectionAssetNodes.h
@@ -20,11 +20,11 @@ struct FGeometryCollectionTerminalDataflowNode : public FDataflowTerminalNode
 
 public:
 
-	UPROPERTY(meta = (DataflowInput, DataflowOutput, Passthrough = "Collection", DisplayName = "Collection"))
+	UPROPERTY(meta = (DataflowInput, DataflowOutput, DataflowPassthrough = "Collection", DisplayName = "Collection"))
 	FManagedArrayCollection Collection;
 
 	/** Materials array to use for this asset */
-	UPROPERTY(meta = (DataflowInput, DataflowOutput, Passthrough = "Materials", DisplayName = "Materials"))
+	UPROPERTY(meta = (DataflowInput, DataflowOutput, DataflowPassthrough = "Materials", DisplayName = "Materials"))
 	TArray<TObjectPtr<UMaterial>> Materials;
 
 	FGeometryCollectionTerminalDataflowNode(const Dataflow::FNodeParameters& InParam, FGuid InGuid = FGuid::NewGuid())
