Dataflow : Fix slection and gc asset node
- fix crash in selection nodes
- optimize selection node initialization from another selection
- fix node for creating GC from geometry sources to make sure we have only one root and a level attribute

diff --git a/Engine/Plugins/Experimental/GeometryCollectionPlugin/Source/GeometryCollectionNodes/Private/Dataflow/Nodes/GeometryCollectionAssetNodes.cpp b/Engine/Plugins/Experimental/GeometryCollectionPlugin/Source/GeometryCollectionNodes/Private/Dataflow/Nodes/GeometryCollectionAssetNodes.cpp
--- a/Engine/Plugins/Experimental/GeometryCollectionPlugin/Source/GeometryCollectionNodes/Private/Dataflow/Nodes/GeometryCollectionAssetNodes.cpp
+++ b/Engine/Plugins/Experimental/GeometryCollectionPlugin/Source/GeometryCollectionNodes/Private/Dataflow/Nodes/GeometryCollectionAssetNodes.cpp
@@ -6,6 +6,8 @@
 #include "GeometryCollection/GeometryCollection.h"
 #include "GeometryCollection/GeometryCollectionEngineConversion.h"
 #include "GeometryCollection/GeometryCollectionObject.h"
+#include "GeometryCollection/GeometryCollectionClusteringUtility.h"
+#include "GeometryCollection/Facades/CollectionHierarchyFacade.h"
 
 namespace Dataflow
 {
@@ -126,8 +128,17 @@ void FCreateGeometryCollectionFromSourcesDataflowNode::Evaluate(Dataflow::FConte
 		OutCollection.ReindexMaterials();
 	}
 
-	// we have to make a copy since we have generated a FGeometryCollection which is inherited from FManagedArrayCollection
+	// make sure we have only one root
+	if (FGeometryCollectionClusteringUtility::ContainsMultipleRootBones(&OutCollection))
+	{
+		FGeometryCollectionClusteringUtility::ClusterAllBonesUnderNewRoot(&OutCollection);
+	}
 
+	// make sure we have a level attribute
+	Chaos::Facades::FCollectionHierarchyFacade HierarchyFacade(OutCollection);
+	HierarchyFacade.GenerateLevelAttribute();
+
+	// we have to make a copy since we have generated a FGeometryCollection which is inherited from FManagedArrayCollection
 	SetValue(Context, static_cast<const FManagedArrayCollection&>(OutCollection), &Collection);
 	SetValue(Context, MoveTemp(OutMaterials), &Materials);
 }
diff --git a/Engine/Source/Runtime/Experimental/Dataflow/Core/Private/Dataflow/DataflowSelection.cpp b/Engine/Source/Runtime/Experimental/Dataflow/Core/Private/Dataflow/DataflowSelection.cpp
--- a/Engine/Source/Runtime/Experimental/Dataflow/Core/Private/Dataflow/DataflowSelection.cpp
+++ b/Engine/Source/Runtime/Experimental/Dataflow/Core/Private/Dataflow/DataflowSelection.cpp
@@ -10,20 +10,12 @@ void FDataflowSelection::Initialize(int32 NumBits, bool Value)
 
 void FDataflowSelection::Initialize(const FDataflowSelection& Other)
 {
-	Initialize(Other.Num(), false);
-
-	for (int32 Idx = 0; Idx < Other.Num(); ++Idx)
-	{
-		if (Other.IsSelected(Idx))
-		{
-			SelectionArray[Idx] = true;
-		}
-	}
+	SelectionArray = Other.SelectionArray;
 }
 
 void FDataflowSelection::AsArray(TArray<int32>& SelectionArr) const
 {
-	SelectionArr.Empty();
+	SelectionArr.Reset(SelectionArray.Num());
 
 	TBitArray<>::FConstIterator It(SelectionArray);
 	while (It)
@@ -41,7 +33,7 @@ void FDataflowSelection::AsArray(TArray<int32>& SelectionArr) const
 
 void FDataflowSelection::SetFromArray(const TArray<int32>& SelectionArr)
 {
-	SelectionArray.Init(false, SelectionArray.Num());
+	SelectionArray.Init(false, SelectionArr.Num());
 
 	for (int32 Elem : SelectionArr)
 	{
@@ -66,6 +58,7 @@ void FDataflowSelection::XOR(const FDataflowSelection& Other, FDataflowSelection
 
 int32 FDataflowSelection::NumSelected() const
 {
+	// todo : use SelectionArray.CountSetBits instead 
 	int32 NumSelectedBits = 0;
 
 	TBitArray<>::FConstIterator It(SelectionArray);
