Dataflow:
- Fixed uninitialized property errors

diff --git a/Engine/Plugins/Experimental/GeometryCollectionPlugin/Source/GeometryCollectionNodes/Public/Dataflow/GeometryCollectionNodes.h b/Engine/Plugins/Experimental/GeometryCollectionPlugin/Source/GeometryCollectionNodes/Public/Dataflow/GeometryCollectionNodes.h
--- a/Engine/Plugins/Experimental/GeometryCollectionPlugin/Source/GeometryCollectionNodes/Public/Dataflow/GeometryCollectionNodes.h
+++ b/Engine/Plugins/Experimental/GeometryCollectionPlugin/Source/GeometryCollectionNodes/Public/Dataflow/GeometryCollectionNodes.h
@@ -1537,7 +1537,7 @@ public:
 
 	/** Mesh triangle count */
 	UPROPERTY(meta = (DataflowOutput))
-	int32 TriangleCount;
+	int32 TriangleCount = 0;
 
 	FPointsToMeshDataflowNode(const Dataflow::FNodeParameters& InParam, FGuid InGuid = FGuid::NewGuid())
 		: FDataflowNode(InParam, InGuid)
@@ -1921,7 +1921,7 @@ public:
 
 	/** Boolean result of the comparison */
 	UPROPERTY(meta = (DataflowOutput));
-	bool Result;
+	bool Result = false;
 
 	FCompareIntDataflowNode(const Dataflow::FNodeParameters& InParam, FGuid InGuid = FGuid::NewGuid())
 		: FDataflowNode(InParam, InGuid)
