Dataflow:
- Moved GeometryCollection Dataflow nodes into separate implemetation files by category
- Changed Dataflow node category to "Terminal" on terminal nodes
- Added DataflowGeoemtryCollection, DataflowFlesh categories to nodes
- Addded black default color to terminal category and changed default color for non-specified categories to yellow
- Added comments to all GeometryCollection nodes
- Updated existing Dataflow nodes

diff --git a/Engine/Plugins/Experimental/ChaosFlesh/Source/ChaosFleshNodes/Public/Dataflow/ChaosFleshEngineAssetNodes.h b/Engine/Plugins/Experimental/ChaosFlesh/Source/ChaosFleshNodes/Public/Dataflow/ChaosFleshEngineAssetNodes.h
--- a/Engine/Plugins/Experimental/ChaosFlesh/Source/ChaosFleshNodes/Public/Dataflow/ChaosFleshEngineAssetNodes.h
+++ b/Engine/Plugins/Experimental/ChaosFlesh/Source/ChaosFleshNodes/Public/Dataflow/ChaosFleshEngineAssetNodes.h
@@ -12,7 +12,7 @@
 
 #include "ChaosFleshEngineAssetNodes.generated.h"
 
-USTRUCT()
+USTRUCT(meta = (DataflowFlesh))
 struct FGetFleshAssetDataflowNode : public FDataflowNode
 {
 	GENERATED_USTRUCT_BODY()
@@ -34,11 +34,11 @@ public:
 };
 
 
-USTRUCT()
+USTRUCT(meta = (DataflowFlesh, DataflowTerminal))
 struct FFleshAssetTerminalDataflowNode : public FDataflowTerminalNode
 {
 	GENERATED_USTRUCT_BODY()
-	DATAFLOW_NODE_DEFINE_INTERNAL(FFleshAssetTerminalDataflowNode, "FleshAssetTerminal", "Flesh", "")
+	DATAFLOW_NODE_DEFINE_INTERNAL(FFleshAssetTerminalDataflowNode, "FleshAssetTerminal", "Terminal", "")
 	DATAFLOW_NODE_RENDER_TYPE(FGeometryCollection::StaticType(), "Collection")
 
 public:
@@ -60,7 +60,7 @@ public:
 
 
 
-USTRUCT()
+USTRUCT(meta = (DataflowFlesh))
 struct FImportFleshDataflowNode : public FDataflowNode
 {
 	GENERATED_USTRUCT_BODY()
@@ -85,7 +85,7 @@ public:
 };
 
 
-USTRUCT()
+USTRUCT(meta = (DataflowFlesh))
 struct FConstructTetGridNode : public FDataflowNode
 {
 	GENERATED_USTRUCT_BODY()
@@ -114,7 +114,7 @@ public:
 
 
 
-USTRUCT()
+USTRUCT(meta = (DataflowFlesh))
 struct FComputeFleshMassNode : public FDataflowNode
 {
 	GENERATED_USTRUCT_BODY()
@@ -144,7 +144,7 @@ public:
 * vertices, and origin & insertion vertex fields.  Fiber directions should smoothly follow the geometry
 * oriented from the origin vertices pointing to the insertion vertices.
 */
-USTRUCT()
+USTRUCT(meta = (DataflowFlesh))
 struct FComputeFiberFieldNode : public FDataflowNode
 {
 	GENERATED_USTRUCT_BODY()
diff --git a/Engine/Plugins/Experimental/ChaosFlesh/Source/ChaosFleshNodes/Public/Dataflow/ChaosFleshFiberDirectionInitializationNodes.h b/Engine/Plugins/Experimental/ChaosFlesh/Source/ChaosFleshNodes/Public/Dataflow/ChaosFleshFiberDirectionInitializationNodes.h
--- a/Engine/Plugins/Experimental/ChaosFlesh/Source/ChaosFleshNodes/Public/Dataflow/ChaosFleshFiberDirectionInitializationNodes.h
+++ b/Engine/Plugins/Experimental/ChaosFlesh/Source/ChaosFleshNodes/Public/Dataflow/ChaosFleshFiberDirectionInitializationNodes.h
@@ -8,7 +8,7 @@
 
 #include "ChaosFleshFiberDirectionInitializationNodes.generated.h"
 
-USTRUCT()
+USTRUCT(meta = (DataflowFlesh))
 struct FGenerateFiberDirectionsDataflowNode : public FDataflowNode
 {
 	GENERATED_USTRUCT_BODY()
diff --git a/Engine/Plugins/Experimental/ChaosFlesh/Source/ChaosFleshNodes/Public/Dataflow/ChaosFleshKinematicInitializationNodes.h b/Engine/Plugins/Experimental/ChaosFlesh/Source/ChaosFleshNodes/Public/Dataflow/ChaosFleshKinematicInitializationNodes.h
--- a/Engine/Plugins/Experimental/ChaosFlesh/Source/ChaosFleshNodes/Public/Dataflow/ChaosFleshKinematicInitializationNodes.h
+++ b/Engine/Plugins/Experimental/ChaosFlesh/Source/ChaosFleshNodes/Public/Dataflow/ChaosFleshKinematicInitializationNodes.h
@@ -10,7 +10,7 @@
 
 class USkeletalMesh;
 
-USTRUCT()
+USTRUCT(meta = (DataflowFlesh))
 struct FKinematicTetrahedralBindingsDataflowNode : public FDataflowNode
 {
 	GENERATED_USTRUCT_BODY()
@@ -41,7 +41,7 @@ public:
 };
 
 
-USTRUCT()
+USTRUCT(meta = (DataflowFlesh))
 struct FKinematicInitializationDataflowNode : public FDataflowNode
 {
 	GENERATED_USTRUCT_BODY()
@@ -84,7 +84,7 @@ public:
 
 
 
-USTRUCT()
+USTRUCT(meta = (DataflowFlesh))
 struct FSelectionSetDataflowNode : public FDataflowNode
 {
 	GENERATED_USTRUCT_BODY()
@@ -110,7 +110,7 @@ public:
 
 
 
-USTRUCT()
+USTRUCT(meta = (DataflowFlesh))
 struct FSetVerticesKinematicDataflowNode : public FDataflowNode
 {
 	GENERATED_USTRUCT_BODY()
@@ -136,7 +136,7 @@ public:
 	virtual void Evaluate(Dataflow::FContext& Context, const FDataflowOutput* Out) const override;
 };
 
-USTRUCT()
+USTRUCT(meta = (DataflowFlesh))
 struct FBinVerticesDataflowNode : public FDataflowNode
 {
 	GENERATED_USTRUCT_BODY()
diff --git a/Engine/Plugins/Experimental/ChaosFlesh/Source/ChaosFleshNodes/Public/Dataflow/ChaosFleshRenderInitializationNodes.h b/Engine/Plugins/Experimental/ChaosFlesh/Source/ChaosFleshNodes/Public/Dataflow/ChaosFleshRenderInitializationNodes.h
--- a/Engine/Plugins/Experimental/ChaosFlesh/Source/ChaosFleshNodes/Public/Dataflow/ChaosFleshRenderInitializationNodes.h
+++ b/Engine/Plugins/Experimental/ChaosFlesh/Source/ChaosFleshNodes/Public/Dataflow/ChaosFleshRenderInitializationNodes.h
@@ -11,7 +11,7 @@
 class USkeletalMesh;
 
 
-USTRUCT()
+USTRUCT(meta = (DataflowFlesh))
 struct FBindForRenderToSkeletalMeshDataflowNode : public FDataflowNode
 {
 	GENERATED_USTRUCT_BODY()
diff --git a/Engine/Plugins/Experimental/ChaosFlesh/Source/ChaosFleshNodes/Public/Dataflow/ChaosFleshTetrahedralNodes.h b/Engine/Plugins/Experimental/ChaosFlesh/Source/ChaosFleshNodes/Public/Dataflow/ChaosFleshTetrahedralNodes.h
--- a/Engine/Plugins/Experimental/ChaosFlesh/Source/ChaosFleshNodes/Public/Dataflow/ChaosFleshTetrahedralNodes.h
+++ b/Engine/Plugins/Experimental/ChaosFlesh/Source/ChaosFleshNodes/Public/Dataflow/ChaosFleshTetrahedralNodes.h
@@ -11,7 +11,7 @@
 class UStaticMesh;
 
 
-USTRUCT()
+USTRUCT(meta = (DataflowFlesh))
 struct FGenerateTetrahedralCollectionDataflowNodes : public FDataflowNode
 {
 	GENERATED_USTRUCT_BODY()
diff --git a/Engine/Plugins/Experimental/Fracture/Source/FractureEngine/Private/FractureEngineSelection.cpp b/Engine/Plugins/Experimental/Fracture/Source/FractureEngine/Private/FractureEngineSelection.cpp
--- a/Engine/Plugins/Experimental/Fracture/Source/FractureEngine/Private/FractureEngineSelection.cpp
+++ b/Engine/Plugins/Experimental/Fracture/Source/FractureEngine/Private/FractureEngineSelection.cpp
@@ -16,12 +16,46 @@ static void ConvertSelectionSetToSelectionArr(const TSet<int32>& SelectionSet, T
 	}
 }
 
-void FFractureEngineSelection::SelectParent(const FGeometryCollection& GeometryCollection, TArray<int32>& SelectedBones)
+void FFractureEngineSelection::GetRootBones(const FManagedArrayCollection& Collection, TArray<int32>& RootBonesOut)
 {
-	if (GeometryCollection.HasGroup(FGeometryCollection::TransformGroup) &&
-		GeometryCollection.HasAttribute("Parent", FGeometryCollection::TransformGroup))
+	if (Collection.HasGroup(FGeometryCollection::TransformGroup) &&
+		Collection.HasAttribute("Parent", FGeometryCollection::TransformGroup))
+	{
+		const TManagedArray<int32>& Parents = Collection.GetAttribute<int32>("Parent", FGeometryCollection::TransformGroup);
+
+		for (int32 Idx = 0; Idx < Parents.Num(); ++Idx)
+		{
+			if (Parents[Idx] == FGeometryCollection::Invalid)
+			{
+				RootBonesOut.Add(Idx);
+			}
+		}
+	}
+}
+
+void FFractureEngineSelection::GetRootBones(const FManagedArrayCollection& Collection, FDataflowTransformSelection& TransformSelection)
+{
+	if (Collection.HasGroup(FGeometryCollection::TransformGroup) &&
+		Collection.HasAttribute("Parent", FGeometryCollection::TransformGroup))
+	{
+		const TManagedArray<int32>& Parents = Collection.GetAttribute<int32>("Parent", FGeometryCollection::TransformGroup);
+
+		for (int32 Idx = 0; Idx < Parents.Num(); ++Idx)
+		{
+			if (Parents[Idx] == FGeometryCollection::Invalid)
+			{
+				TransformSelection.SetSelected(Idx);
+			}
+		}
+	}
+}
+
+void FFractureEngineSelection::SelectParent(const FManagedArrayCollection& Collection, TArray<int32>& SelectedBones)
+{
+	if (Collection.HasGroup(FGeometryCollection::TransformGroup) &&
+		Collection.HasAttribute("Parent", FGeometryCollection::TransformGroup))
 	{
-		const TManagedArray<int32>& Parents = GeometryCollection.GetAttribute<int32>("Parent", FGeometryCollection::TransformGroup);
+		const TManagedArray<int32>& Parents = Collection.GetAttribute<int32>("Parent", FGeometryCollection::TransformGroup);
 
 		TSet<int32> NewSelection;
 		for (int32 Bone : SelectedBones)
@@ -37,22 +71,22 @@ void FFractureEngineSelection::SelectParent(const FGeometryCollection& GeometryC
 	}
 }
 
-void FFractureEngineSelection::SelectParent(const FGeometryCollection& GeometryCollection, FDataflowTransformSelection& TransformSelection)
+void FFractureEngineSelection::SelectParent(const FManagedArrayCollection& Collection, FDataflowTransformSelection& TransformSelection)
 {
 	TArray<int32> SelectionArr;
 	TransformSelection.AsArray(SelectionArr);
 
-	SelectParent(GeometryCollection, SelectionArr);
+	SelectParent(Collection, SelectionArr);
 
 	TransformSelection.SetFromArray(SelectionArr);
 }
 
-void FFractureEngineSelection::SelectChildren(const FGeometryCollection& GeometryCollection, TArray<int32>& SelectedBones)
+void FFractureEngineSelection::SelectChildren(const FManagedArrayCollection& Collection, TArray<int32>& SelectedBones)
 {
-	if (GeometryCollection.HasGroup(FGeometryCollection::TransformGroup) &&
-		GeometryCollection.HasAttribute("Children", FGeometryCollection::TransformGroup))
+	if (Collection.HasGroup(FGeometryCollection::TransformGroup) &&
+		Collection.HasAttribute("Children", FGeometryCollection::TransformGroup))
 	{
-		const TManagedArray<TSet<int32>>& Children = GeometryCollection.GetAttribute<TSet<int32>>("Children", FGeometryCollection::TransformGroup);
+		const TManagedArray<TSet<int32>>& Children = Collection.GetAttribute<TSet<int32>>("Children", FGeometryCollection::TransformGroup);
 
 		TSet<int32> NewSelection;
 		for (int32 Bone : SelectedBones)
@@ -72,24 +106,24 @@ void FFractureEngineSelection::SelectChildren(const FGeometryCollection& Geometr
 	}
 }
 
-void FFractureEngineSelection::SelectChildren(const FGeometryCollection& GeometryCollection, FDataflowTransformSelection& TransformSelection)
+void FFractureEngineSelection::SelectChildren(const FManagedArrayCollection& Collection, FDataflowTransformSelection& TransformSelection)
 {
 	TArray<int32> SelectionArr;
 	TransformSelection.AsArray(SelectionArr);
 
-	SelectChildren(GeometryCollection, SelectionArr);
+	SelectChildren(Collection, SelectionArr);
 
 	TransformSelection.SetFromArray(SelectionArr);
 }
 
-void FFractureEngineSelection::SelectSiblings(const FGeometryCollection& GeometryCollection, TArray<int32>& SelectedBones)
+void FFractureEngineSelection::SelectSiblings(const FManagedArrayCollection& Collection, TArray<int32>& SelectedBones)
 {
-	if (GeometryCollection.HasGroup(FGeometryCollection::TransformGroup) &&
-		GeometryCollection.HasAttribute("Parent", FGeometryCollection::TransformGroup) &&
-		GeometryCollection.HasAttribute("Children", FGeometryCollection::TransformGroup))
+	if (Collection.HasGroup(FGeometryCollection::TransformGroup) &&
+		Collection.HasAttribute("Parent", FGeometryCollection::TransformGroup) &&
+		Collection.HasAttribute("Children", FGeometryCollection::TransformGroup))
 	{
-		const TManagedArray<int32>& Parents = GeometryCollection.GetAttribute<int32>("Parent", FGeometryCollection::TransformGroup);
-		const TManagedArray<TSet<int32>>& Children = GeometryCollection.GetAttribute<TSet<int32>>("Children", FGeometryCollection::TransformGroup);
+		const TManagedArray<int32>& Parents = Collection.GetAttribute<int32>("Parent", FGeometryCollection::TransformGroup);
+		const TManagedArray<TSet<int32>>& Children = Collection.GetAttribute<TSet<int32>>("Children", FGeometryCollection::TransformGroup);
 
 		TSet<int32> NewSelection;
 		for (int32 Bone : SelectedBones)
@@ -108,28 +142,28 @@ void FFractureEngineSelection::SelectSiblings(const FGeometryCollection& Geometr
 	}
 }
 
-void FFractureEngineSelection::SelectSiblings(const FGeometryCollection& GeometryCollection, FDataflowTransformSelection& TransformSelection)
+void FFractureEngineSelection::SelectSiblings(const FManagedArrayCollection& Collection, FDataflowTransformSelection& TransformSelection)
 {
 	TArray<int32> SelectionArr;
 	TransformSelection.AsArray(SelectionArr);
 
-	SelectSiblings(GeometryCollection, SelectionArr);
+	SelectSiblings(Collection, SelectionArr);
 
 	TransformSelection.SetFromArray(SelectionArr);
 }
 
-void FFractureEngineSelection::SelectLevel(const FGeometryCollection& GeometryCollection, TArray<int32>& SelectedBones)
+void FFractureEngineSelection::SelectLevel(const FManagedArrayCollection& Collection, TArray<int32>& SelectedBones)
 {
-	if (GeometryCollection.HasGroup(FGeometryCollection::TransformGroup) &&
-		GeometryCollection.HasAttribute("Level", FGeometryCollection::TransformGroup))
+	if (Collection.HasGroup(FGeometryCollection::TransformGroup) &&
+		Collection.HasAttribute("Level", FGeometryCollection::TransformGroup))
 	{
-		const TManagedArray<int32>& Levels = GeometryCollection.GetAttribute<int32>("Levels", FGeometryCollection::TransformGroup);
+		const TManagedArray<int32>& Levels = Collection.GetAttribute<int32>("Levels", FGeometryCollection::TransformGroup);
 
 		TSet<int32> NewSelection;
 		for (int32 Bone : SelectedBones)
 		{
 			int32 Level = Levels[Bone];
-			for (int32 TransformIdx = 0; TransformIdx < GeometryCollection.NumElements(FTransformCollection::TransformGroup); ++TransformIdx)
+			for (int32 TransformIdx = 0; TransformIdx < Collection.NumElements(FTransformCollection::TransformGroup); ++TransformIdx)
 			{
 				if (Levels[TransformIdx] == Level)
 				{
@@ -142,12 +176,12 @@ void FFractureEngineSelection::SelectLevel(const FGeometryCollection& GeometryCo
 	}
 }
 
-void FFractureEngineSelection::SelectLevel(const FGeometryCollection& GeometryCollection, FDataflowTransformSelection& TransformSelection)
+void FFractureEngineSelection::SelectLevel(const FManagedArrayCollection& Collection, FDataflowTransformSelection& TransformSelection)
 {
 	TArray<int32> SelectionArr;
 	TransformSelection.AsArray(SelectionArr);
 
-	SelectLevel(GeometryCollection, SelectionArr);
+	SelectLevel(Collection, SelectionArr);
 
 	TransformSelection.SetFromArray(SelectionArr);
 }
@@ -293,7 +327,7 @@ static void RandomShuffleArray(TArray<int32>& InArray, const bool Deterministic,
 	}
 }
 
-void FFractureEngineSelection::SelectByPercentage(const FGeometryCollection& GeometryCollection, TArray<int32>& SelectedBones, const int32 Percentage, const bool Deterministic, const float RandomSeed)
+void FFractureEngineSelection::SelectByPercentage(TArray<int32>& SelectedBones, const int32 Percentage, const bool Deterministic, const float RandomSeed)
 {
 	RandomShuffleArray(SelectedBones, Deterministic, RandomSeed);
 
@@ -301,12 +335,12 @@ void FFractureEngineSelection::SelectByPercentage(const FGeometryCollection& Geo
 	SelectedBones.SetNum(NewNumElements);
 }
 
-void FFractureEngineSelection::SelectByPercentage(const FGeometryCollection& GeometryCollection, FDataflowTransformSelection& TransformSelection, const int32 Percentage, const bool Deterministic, const float RandomSeed)
+void FFractureEngineSelection::SelectByPercentage(FDataflowTransformSelection& TransformSelection, const int32 Percentage, const bool Deterministic, const float RandomSeed)
 {
 	TArray<int32> SelectionArr;
 	TransformSelection.AsArray(SelectionArr);
 
-	SelectByPercentage(GeometryCollection, SelectionArr, Percentage, Deterministic, RandomSeed);
+	SelectByPercentage(SelectionArr, Percentage, Deterministic, RandomSeed);
 
 	TransformSelection.SetFromArray(SelectionArr);
 }
diff --git a/Engine/Plugins/Experimental/Fracture/Source/FractureEngine/Public/FractureEngineSelection.h b/Engine/Plugins/Experimental/Fracture/Source/FractureEngine/Public/FractureEngineSelection.h
--- a/Engine/Plugins/Experimental/Fracture/Source/FractureEngine/Public/FractureEngineSelection.h
+++ b/Engine/Plugins/Experimental/Fracture/Source/FractureEngine/Public/FractureEngineSelection.h
@@ -8,17 +8,20 @@
 class FRACTUREENGINE_API FFractureEngineSelection
 {
 public:
-	static void SelectParent(const FGeometryCollection& GeometryCollection, TArray<int32>& SelectedBones);
-	static void SelectParent(const FGeometryCollection& GeometryCollection, FDataflowTransformSelection& TransformSelection);
+	static void GetRootBones(const FManagedArrayCollection& Collection, TArray<int32>& RootBonesOut);
+	static void GetRootBones(const FManagedArrayCollection& Collection, FDataflowTransformSelection& TransformSelection);
 
-	static void SelectChildren(const FGeometryCollection& GeometryCollection, TArray<int32>& SelectedBones);
-	static void SelectChildren(const FGeometryCollection& GeometryCollection, FDataflowTransformSelection& TransformSelection);
+	static void SelectParent(const FManagedArrayCollection& Collection, TArray<int32>& SelectedBones);
+	static void SelectParent(const FManagedArrayCollection& Collection, FDataflowTransformSelection& TransformSelection);
 
-	static void SelectSiblings(const FGeometryCollection& GeometryCollection, TArray<int32>& SelectedBones);
-	static void SelectSiblings(const FGeometryCollection& GeometryCollection, FDataflowTransformSelection& TransformSelection);
+	static void SelectChildren(const FManagedArrayCollection& Collection, TArray<int32>& SelectedBones);
+	static void SelectChildren(const FManagedArrayCollection& Collection, FDataflowTransformSelection& TransformSelection);
 
-	static void SelectLevel(const FGeometryCollection& GeometryCollection, TArray<int32>& SelectedBones);
-	static void SelectLevel(const FGeometryCollection& GeometryCollection, FDataflowTransformSelection& TransformSelection);
+	static void SelectSiblings(const FManagedArrayCollection& Collection, TArray<int32>& SelectedBones);
+	static void SelectSiblings(const FManagedArrayCollection& Collection, FDataflowTransformSelection& TransformSelection);
+
+	static void SelectLevel(const FManagedArrayCollection& Collection, TArray<int32>& SelectedBones);
+	static void SelectLevel(const FManagedArrayCollection& Collection, FDataflowTransformSelection& TransformSelection);
 
 	static void SelectContact(FGeometryCollection& GeometryCollection, TArray<int32>& SelectedBones);
 	static void SelectContact(FGeometryCollection& GeometryCollection, FDataflowTransformSelection& TransformSelection);
@@ -29,8 +32,8 @@ public:
 	static void SelectCluster(const FGeometryCollection& GeometryCollection, TArray<int32>& SelectedBones);
 	static void SelectCluster(const FGeometryCollection& GeometryCollection, FDataflowTransformSelection& TransformSelection);
 
-	static void SelectByPercentage(const FGeometryCollection& GeometryCollection, TArray<int32>& SelectedBones, const int32 Percentage, const bool Deterministic, const float RandomSeed);
-	static void SelectByPercentage(const FGeometryCollection& GeometryCollection, FDataflowTransformSelection& TransformSelection, const int32 Percentage, const bool Deterministic, const float RandomSeed);
+	static void SelectByPercentage(TArray<int32>& SelectedBones, const int32 Percentage, const bool Deterministic, const float RandomSeed);
+	static void SelectByPercentage(FDataflowTransformSelection& TransformSelection, const int32 Percentage, const bool Deterministic, const float RandomSeed);
 
 	static void SelectBySize(FGeometryCollection& GeometryCollection, TArray<int32>& SelectedBones, const float SizeMin, const float SizeMax);
 	static void SelectBySize(FGeometryCollection& GeometryCollection, FDataflowTransformSelection& TransformSelection, const float SizeMin, const float SizeMax);
diff --git a/Engine/Plugins/Experimental/GeometryCollectionPlugin/Source/GeometryCollectionNodes/Private/Dataflow/GeometryCollectionClusteringNodes.cpp b/Engine/Plugins/Experimental/GeometryCollectionPlugin/Source/GeometryCollectionNodes/Private/Dataflow/GeometryCollectionClusteringNodes.cpp
new file mode 100644
--- /dev/null
+++ b/Engine/Plugins/Experimental/GeometryCollectionPlugin/Source/GeometryCollectionNodes/Private/Dataflow/GeometryCollectionClusteringNodes.cpp
@@ -0,0 +1,108 @@
+// Copyright Epic Games, Inc. All Rights Reserved.
+
+#include "Dataflow/GeometryCollectionClusteringNodes.h"
+#include "Dataflow/DataflowCore.h"
+
+#include "Engine/StaticMesh.h"
+#include "GeometryCollection/GeometryCollectionObject.h"
+#include "GeometryCollection/ManagedArrayCollection.h"
+#include "GeometryCollection/GeometryCollection.h"
+#include "GeometryCollection/GeometryCollectionEngineUtility.h"
+#include "GeometryCollection/GeometryCollectionEngineConversion.h"
+#include "Logging/LogMacros.h"
+#include "Templates/SharedPointer.h"
+#include "UObject/UnrealTypePrivate.h"
+#include "DynamicMeshToMeshDescription.h"
+#include "MeshDescriptionToDynamicMesh.h"
+#include "StaticMeshAttributes.h"
+#include "DynamicMeshEditor.h"
+#include "Operations/MeshBoolean.h"
+
+#include "EngineGlobals.h"
+#include "GeometryCollection/GeometryCollectionAlgo.h"
+#include "GeometryCollection/GeometryCollectionClusteringUtility.h"
+#include "GeometryCollection/GeometryCollectionConvexUtility.h"
+#include "Voronoi/Voronoi.h"
+#include "PlanarCut.h"
+#include "GeometryCollection/GeometryCollectionProximityUtility.h"
+#include "FractureEngineClustering.h"
+#include "FractureEngineSelection.h"
+
+#include UE_INLINE_GENERATED_CPP_BY_NAME(GeometryCollectionClusteringNodes)
+
+namespace Dataflow
+{
+
+	void GeometryCollectionClusteringNodes()
+	{
+		static const FLinearColor CDefaultNodeBodyTintColor = FLinearColor(0.f, 0.f, 0.f, 0.5f);
+
+		DATAFLOW_NODE_REGISTER_CREATION_FACTORY(FAutoClusterDataflowNode);
+		DATAFLOW_NODE_REGISTER_CREATION_FACTORY(FClusterFlattenDataflowNode);
+
+		// GeometryCollection|Cluster
+		DATAFLOW_NODE_REGISTER_CREATION_FACTORY_NODE_COLORS_BY_CATEGORY("GeometryCollection|Cluster", FLinearColor(.25f, 0.45f, 0.8f), CDefaultNodeBodyTintColor);
+	}
+}
+
+
+void FAutoClusterDataflowNode::Evaluate(Dataflow::FContext& Context, const FDataflowOutput* Out) const
+{
+	if (Out->IsA<FManagedArrayCollection>(&Collection))
+	{
+		const FManagedArrayCollection& InCollection = GetValue<FManagedArrayCollection>(Context, &Collection);
+		const FDataflowTransformSelection& InTransformSelection = GetValue<FDataflowTransformSelection>(Context, &TransformSelection);
+
+		if (TUniquePtr<FGeometryCollection> GeomCollection = TUniquePtr<FGeometryCollection>(InCollection.NewCopy<FGeometryCollection>()))
+		{
+			EClusterSizeMethodEnum InClusterSizeMethod = ClusterSizeMethod;
+			int32 InClusterSites = GetValue<int32>(Context, &ClusterSites);
+			float InClusterFraction = GetValue<float>(Context, &ClusterFraction);
+			float InSiteSize = GetValue<float>(Context, &SiteSize);
+			bool InAutoCluster = AutoCluster;
+			bool InAvoidIsolated = AvoidIsolated;
+
+			TArray<int32> SelectedBones;
+			InTransformSelection.AsArray(SelectedBones);
+
+			FFractureEngineClustering::AutoCluster(*GeomCollection,
+				SelectedBones,
+				(EFractureEngineClusterSizeMethod)InClusterSizeMethod,
+				InClusterSites,
+				InClusterFraction,
+				InSiteSize,
+				InAutoCluster,
+				InAvoidIsolated);
+
+			SetValue<FManagedArrayCollection>(Context, (const FManagedArrayCollection&)(*GeomCollection), &Collection);
+		}
+	}
+}
+
+
+void FClusterFlattenDataflowNode::Evaluate(Dataflow::FContext& Context, const FDataflowOutput* Out) const
+{
+	if (Out->IsA<FManagedArrayCollection>(&Collection))
+	{
+		const FManagedArrayCollection& InCollection = GetValue<FManagedArrayCollection>(Context, &Collection);
+		if (TUniquePtr<FGeometryCollection> GeomCollection = TUniquePtr<FGeometryCollection>(InCollection.NewCopy<FGeometryCollection>()))
+		{
+			FGeometryCollectionClusteringUtility::UpdateHierarchyLevelOfChildren(GeomCollection.Get(), -1);
+
+			const TManagedArray<int32>& Levels = GeomCollection->GetAttribute<int32>("Level", FGeometryCollection::TransformGroup);
+
+			// Populate Selected Bones in an Array
+			// @todo(harsha) Implement with Selection
+			// For every bone in selected array: [ClusterIndex]
+			int32 ClusterIndex = 0;
+			TArray<int32> LeafBones;
+			FGeometryCollectionClusteringUtility::GetLeafBones(GeomCollection.Get(), ClusterIndex, true, LeafBones);
+			FGeometryCollectionClusteringUtility::ClusterBonesUnderExistingNode(GeomCollection.Get(), ClusterIndex, LeafBones);
+			FGeometryCollectionClusteringUtility::RemoveDanglingClusters(GeomCollection.Get());
+			// End for
+
+			FGeometryCollectionClusteringUtility::UpdateHierarchyLevelOfChildren(GeomCollection.Get(), -1);
+			SetValue<FManagedArrayCollection>(Context, (const FManagedArrayCollection&)(*GeomCollection), &Collection);
+		}
+	}
+}
diff --git a/Engine/Plugins/Experimental/GeometryCollectionPlugin/Source/GeometryCollectionNodes/Private/Dataflow/GeometryCollectionFracturingNodes.cpp b/Engine/Plugins/Experimental/GeometryCollectionPlugin/Source/GeometryCollectionNodes/Private/Dataflow/GeometryCollectionFracturingNodes.cpp
new file mode 100644
--- /dev/null
+++ b/Engine/Plugins/Experimental/GeometryCollectionPlugin/Source/GeometryCollectionNodes/Private/Dataflow/GeometryCollectionFracturingNodes.cpp
@@ -0,0 +1,263 @@
+// Copyright Epic Games, Inc. All Rights Reserved.
+
+#include "Dataflow/GeometryCollectionFracturingNodes.h"
+#include "Dataflow/DataflowCore.h"
+
+#include "Engine/StaticMesh.h"
+#include "GeometryCollection/GeometryCollectionObject.h"
+#include "GeometryCollection/ManagedArrayCollection.h"
+#include "GeometryCollection/GeometryCollection.h"
+#include "GeometryCollection/GeometryCollectionEngineUtility.h"
+#include "GeometryCollection/GeometryCollectionEngineConversion.h"
+#include "Logging/LogMacros.h"
+#include "Templates/SharedPointer.h"
+#include "UObject/UnrealTypePrivate.h"
+#include "DynamicMeshToMeshDescription.h"
+#include "MeshDescriptionToDynamicMesh.h"
+#include "StaticMeshAttributes.h"
+#include "DynamicMeshEditor.h"
+#include "Operations/MeshBoolean.h"
+
+#include "EngineGlobals.h"
+#include "GeometryCollection/GeometryCollectionAlgo.h"
+#include "GeometryCollection/GeometryCollectionClusteringUtility.h"
+#include "GeometryCollection/GeometryCollectionConvexUtility.h"
+#include "Voronoi/Voronoi.h"
+#include "PlanarCut.h"
+#include "GeometryCollection/GeometryCollectionProximityUtility.h"
+#include "FractureEngineClustering.h"
+#include "FractureEngineSelection.h"
+
+#include UE_INLINE_GENERATED_CPP_BY_NAME(GeometryCollectionFracturingNodes)
+
+namespace Dataflow
+{
+
+	void GeometryCollectionFracturingNodes()
+	{
+		static const FLinearColor CDefaultNodeBodyTintColor = FLinearColor(0.f, 0.f, 0.f, 0.5f);
+
+		DATAFLOW_NODE_REGISTER_CREATION_FACTORY(FUniformScatterPointsDataflowNode);
+		DATAFLOW_NODE_REGISTER_CREATION_FACTORY(FRadialScatterPointsDataflowNode);
+		DATAFLOW_NODE_REGISTER_CREATION_FACTORY(FVoronoiFractureDataflowNode);
+		DATAFLOW_NODE_REGISTER_CREATION_FACTORY(FPlaneCutterDataflowNode);
+
+		// GeometryCollection|Fracture
+		DATAFLOW_NODE_REGISTER_CREATION_FACTORY_NODE_COLORS_BY_CATEGORY("GeometryCollection|Fracture", FLinearColor(1.f, 1.f, 0.8f), CDefaultNodeBodyTintColor);
+	}
+}
+
+void FUniformScatterPointsDataflowNode::Evaluate(Dataflow::FContext& Context, const FDataflowOutput* Out) const
+{
+	if (Out->IsA<TArray<FVector>>(&Points))
+	{
+		const FBox& BBox = GetValue<FBox>(Context, &BoundingBox);
+		if (BBox.GetVolume() > 0.f)
+		{
+			FRandomStream RandStream(GetValue<float>(Context, &RandomSeed));
+
+			const FVector Extent(BBox.Max - BBox.Min);
+			const int32 NumPoints = RandStream.RandRange(GetValue<int32>(Context, &MinNumberOfPoints), GetValue<int32>(Context, &MaxNumberOfPoints));
+
+			TArray<FVector> PointsArr;
+			PointsArr.Reserve(NumPoints);
+			for (int32 Idx = 0; Idx < NumPoints; ++Idx)
+			{
+				PointsArr.Emplace(BBox.Min + FVector(RandStream.FRand(), RandStream.FRand(), RandStream.FRand()) * Extent);
+			}
+
+			SetValue<TArray<FVector>>(Context, PointsArr, &Points);
+		}
+		else
+		{
+			// ERROR: Invalid BoundingBox input
+			SetValue<TArray<FVector>>(Context, TArray<FVector>(), &Points);
+		}
+	}
+}
+
+void FRadialScatterPointsDataflowNode::Evaluate(Dataflow::FContext& Context, const FDataflowOutput* Out) const
+{
+	if (Out->IsA<TArray<FVector>>(&Points))
+	{
+		const FVector::FReal RadialStep = GetValue<float>(Context, &Radius) / GetValue<int32>(Context, &RadialSteps);
+		const FVector::FReal AngularStep = 2 * PI / GetValue<int32>(Context, &AngularSteps);
+
+		FRandomStream RandStream(GetValue<float>(Context, &RandomSeed));
+		FVector UpVector(GetValue<FVector>(Context, &Normal));
+		UpVector.Normalize();
+		FVector BasisX, BasisY;
+		UpVector.FindBestAxisVectors(BasisX, BasisY);
+
+		TArray<FVector> PointsArr;
+
+		FVector::FReal Len = RadialStep * .5;
+		for (int32 ii = 0; ii < GetValue<int32>(Context, &RadialSteps); ++ii, Len += RadialStep)
+		{
+			FVector::FReal Angle = FMath::DegreesToRadians(GetValue<float>(Context, &AngleOffset));
+			for (int32 kk = 0; kk < AngularSteps; ++kk, Angle += AngularStep)
+			{
+				FVector RotatingOffset = Len * (FMath::Cos(Angle) * BasisX + FMath::Sin(Angle) * BasisY);
+				PointsArr.Emplace(GetValue<FVector>(Context, &Center) + RotatingOffset + (RandStream.VRand() * RandStream.FRand() * Variability));
+			}
+		}
+
+		SetValue<TArray<FVector>>(Context, PointsArr, &Points);
+	}
+}
+
+
+static float GetMaxVertexMovement(float Grout, float Amplitude, int OctaveNumber, float Persistence)
+{
+	float MaxDisp = Grout;
+	float AmplitudeScaled = Amplitude;
+	for (int32 OctaveIdx = 0; OctaveIdx < OctaveNumber; OctaveIdx++, AmplitudeScaled *= Persistence)
+	{
+		MaxDisp += FMath::Abs(AmplitudeScaled);
+	}
+	return MaxDisp;
+}
+
+void FVoronoiFractureDataflowNode::Evaluate(Dataflow::FContext& Context, const FDataflowOutput* Out) const
+{
+	if (Out->IsA<FManagedArrayCollection>(&Collection))
+	{
+		const FManagedArrayCollection& InCollection = GetValue<FManagedArrayCollection>(Context, &Collection);
+		if (TUniquePtr<FGeometryCollection> GeomCollection = TUniquePtr<FGeometryCollection>(InCollection.NewCopy<FGeometryCollection>()))
+		{
+			const TArray<FVector>& Sites = GetValue<TArray<FVector>>(Context, &Points);
+			if (Sites.Num() > 0)
+			{
+				//
+				// Compute BoundingBox for ManagedArrayIn
+				//
+				FBox BoundingBox(ForceInit);
+
+				if (InCollection.HasAttribute("Transform", FGeometryCollection::TransformGroup) &&
+					InCollection.HasAttribute("Parent", FGeometryCollection::TransformGroup) &&
+					InCollection.HasAttribute("TransformIndex", FGeometryCollection::GeometryGroup) &&
+					InCollection.HasAttribute("BoundingBox", FGeometryCollection::GeometryGroup))
+				{
+					const TManagedArray<FTransform>& Transforms = InCollection.GetAttribute<FTransform>("Transform", FGeometryCollection::TransformGroup);
+					const TManagedArray<int32>& ParentIndices = InCollection.GetAttribute<int32>("Parent", FGeometryCollection::TransformGroup);
+					const TManagedArray<int32>& TransformIndices = InCollection.GetAttribute<int32>("TransformIndex", FGeometryCollection::GeometryGroup);
+					const TManagedArray<FBox>& BoundingBoxes = InCollection.GetAttribute<FBox>("BoundingBox", FGeometryCollection::GeometryGroup);
+
+					TArray<FMatrix> TmpGlobalMatrices;
+					GeometryCollectionAlgo::GlobalMatrices(Transforms, ParentIndices, TmpGlobalMatrices);
+
+					if (TmpGlobalMatrices.Num() > 0)
+					{
+						for (int32 BoxIdx = 0; BoxIdx < BoundingBoxes.Num(); ++BoxIdx)
+						{
+							const int32 TransformIndex = TransformIndices[BoxIdx];
+							BoundingBox += BoundingBoxes[BoxIdx].TransformBy(TmpGlobalMatrices[TransformIndex]);
+						}
+					}
+
+					//
+					// Compute Voronoi Bounds
+					//
+					FBox VoronoiBounds = BoundingBox;
+					VoronoiBounds += FBox(Sites);
+
+					float GroutVal = GetValue<float>(Context, &Grout);
+					float AmplitudeVal = GetValue<float>(Context, &Amplitude);
+					int32 OctaveNumberVal = GetValue<int32>(Context, &OctaveNumber);
+					float PersistenceVal = GetValue<float>(Context, &Persistence);
+
+					VoronoiBounds = VoronoiBounds.ExpandBy(GetMaxVertexMovement(GroutVal, AmplitudeVal, OctaveNumberVal, PersistenceVal) + KINDA_SMALL_NUMBER);
+
+					//
+					// Voronoi Fracture
+					//
+					FNoiseSettings NoiseSettings;
+					NoiseSettings.Amplitude = AmplitudeVal;
+					NoiseSettings.Frequency = GetValue<float>(Context, &Frequency);
+					NoiseSettings.Octaves = OctaveNumberVal;
+					NoiseSettings.PointSpacing = GetValue<float>(Context, &PointSpacing);
+					NoiseSettings.Lacunarity = GetValue<float>(Context, &Lacunarity);
+					NoiseSettings.Persistence = GetValue<float>(Context, &Persistence);;
+
+					FVoronoiDiagram Voronoi(Sites, VoronoiBounds, .1f);
+
+					FPlanarCells VoronoiPlanarCells = FPlanarCells(Sites, Voronoi);
+					VoronoiPlanarCells.InternalSurfaceMaterials.NoiseSettings = NoiseSettings;
+
+					const TArrayView<const int32>& TransformIndicesArray(TransformIndices.GetConstArray());
+
+					float CollisionSampleSpacingVal = GetValue<float>(Context, &CollisionSampleSpacing);
+					float RandomSeedVal = GetValue<float>(Context, &RandomSeed);
+
+					int ResultGeometryIndex = CutMultipleWithPlanarCells(VoronoiPlanarCells, *GeomCollection, TransformIndicesArray, GroutVal, CollisionSampleSpacingVal, RandomSeedVal, FTransform().Identity);
+
+					SetValue<FManagedArrayCollection>(Context, (const FManagedArrayCollection&)(*GeomCollection), &Collection);
+				}
+			}
+		}
+	}
+}
+
+
+void FPlaneCutterDataflowNode::Evaluate(Dataflow::FContext& Context, const FDataflowOutput* Out) const
+{
+	if (Out->IsA<FManagedArrayCollection>(&Collection))
+	{
+		const FManagedArrayCollection& InCollection = GetValue<FManagedArrayCollection>(Context, &Collection);
+		if (TUniquePtr<FGeometryCollection> GeomCollection = TUniquePtr<FGeometryCollection>(InCollection.NewCopy<FGeometryCollection>()))
+		{
+			TArray<FPlane> CuttingPlanes;
+			TArray<FTransform> CuttingPlaneTransforms;
+
+			float RandomSeedVal = GetValue<float>(Context, &RandomSeed);
+			FRandomStream RandStream(RandomSeedVal);
+
+			FBox Bounds = GetValue<FBox>(Context, &BoundingBox);
+			const FVector Extent(Bounds.Max - Bounds.Min);
+
+			CuttingPlaneTransforms.Reserve(CuttingPlaneTransforms.Num() + NumPlanes);
+			for (int32 ii = 0; ii < NumPlanes; ++ii)
+			{
+				FVector Position(Bounds.Min + FVector(RandStream.FRand(), RandStream.FRand(), RandStream.FRand()) * Extent);
+				CuttingPlaneTransforms.Emplace(FTransform(FRotator(RandStream.FRand() * 360.0f, RandStream.FRand() * 360.0f, 0.0f), Position));
+			}
+
+			for (const FTransform& Transform : CuttingPlaneTransforms)
+			{
+				CuttingPlanes.Add(FPlane(Transform.GetLocation(), Transform.GetUnitAxis(EAxis::Z)));
+			}
+
+			FInternalSurfaceMaterials InternalSurfaceMaterials;
+			FNoiseSettings NoiseSettings;
+
+			float AmplitudeVal = GetValue<float>(Context, &Amplitude);
+			if (AmplitudeVal > 0.0f)
+			{
+				NoiseSettings.Amplitude = AmplitudeVal;
+				NoiseSettings.Frequency = GetValue<float>(Context, &Frequency);
+				NoiseSettings.Lacunarity = GetValue<float>(Context, &Lacunarity);
+				NoiseSettings.Persistence = GetValue<float>(Context, &Persistence);
+				NoiseSettings.Octaves = GetValue<int32>(Context, &OctaveNumber);
+				NoiseSettings.PointSpacing = GetValue<float>(Context, &PointSpacing);
+
+				InternalSurfaceMaterials.NoiseSettings = NoiseSettings;
+			}
+
+			if (GeomCollection->HasAttribute("TransformIndex", "Geometry"))
+			{
+				const TManagedArray<int32>& TransformIndices = GeomCollection->GetAttribute<int32>("TransformIndex", "Geometry");
+				const TArrayView<const int32>& TransformIndicesArray(TransformIndices.GetConstArray());
+
+				float CollisionSampleSpacingVal = GetValue<float>(Context, &CollisionSampleSpacing);
+				float GroutVal = GetValue<float>(Context, &Grout);
+
+				int ResultGeometryIndex = CutMultipleWithMultiplePlanes(CuttingPlanes, InternalSurfaceMaterials, *GeomCollection, TransformIndicesArray, GroutVal, CollisionSampleSpacingVal, RandomSeedVal, FTransform().Identity);
+
+				SetValue<FManagedArrayCollection>(Context, (const FManagedArrayCollection&)(*GeomCollection), &Collection);
+			}
+		}
+	}
+}
+
+
+
diff --git a/Engine/Plugins/Experimental/GeometryCollectionPlugin/Source/GeometryCollectionNodes/Private/Dataflow/GeometryCollectionMeshNodes.cpp b/Engine/Plugins/Experimental/GeometryCollectionPlugin/Source/GeometryCollectionNodes/Private/Dataflow/GeometryCollectionMeshNodes.cpp
new file mode 100644
--- /dev/null
+++ b/Engine/Plugins/Experimental/GeometryCollectionPlugin/Source/GeometryCollectionNodes/Private/Dataflow/GeometryCollectionMeshNodes.cpp
@@ -0,0 +1,342 @@
+// Copyright Epic Games, Inc. All Rights Reserved.
+
+#include "Dataflow/GeometryCollectionMeshNodes.h"
+#include "Dataflow/DataflowCore.h"
+
+#include "Engine/StaticMesh.h"
+#include "GeometryCollection/GeometryCollectionObject.h"
+#include "GeometryCollection/ManagedArrayCollection.h"
+#include "GeometryCollection/GeometryCollection.h"
+#include "GeometryCollection/GeometryCollectionEngineUtility.h"
+#include "GeometryCollection/GeometryCollectionEngineConversion.h"
+#include "Logging/LogMacros.h"
+#include "Templates/SharedPointer.h"
+#include "UObject/UnrealTypePrivate.h"
+#include "DynamicMeshToMeshDescription.h"
+#include "MeshDescriptionToDynamicMesh.h"
+#include "StaticMeshAttributes.h"
+#include "DynamicMeshEditor.h"
+#include "Operations/MeshBoolean.h"
+
+#include "EngineGlobals.h"
+#include "GeometryCollection/GeometryCollectionAlgo.h"
+#include "GeometryCollection/GeometryCollectionClusteringUtility.h"
+#include "GeometryCollection/GeometryCollectionConvexUtility.h"
+#include "Voronoi/Voronoi.h"
+#include "PlanarCut.h"
+#include "GeometryCollection/GeometryCollectionProximityUtility.h"
+#include "FractureEngineClustering.h"
+#include "FractureEngineSelection.h"
+#include "DynamicMesh/DynamicMesh3.h"
+#include "UDynamicMesh.h"
+
+#include UE_INLINE_GENERATED_CPP_BY_NAME(GeometryCollectionMeshNodes)
+
+namespace Dataflow
+{
+
+	void GeometryCollectionMeshNodes()
+	{
+		static const FLinearColor CDefaultNodeBodyTintColor = FLinearColor(0.f, 0.f, 0.f, 0.5f);
+
+		DATAFLOW_NODE_REGISTER_CREATION_FACTORY(FPointsToMeshDataflowNode);
+		DATAFLOW_NODE_REGISTER_CREATION_FACTORY(FBoxToMeshDataflowNode);
+		DATAFLOW_NODE_REGISTER_CREATION_FACTORY(FMeshInfoDataflowNode);
+		DATAFLOW_NODE_REGISTER_CREATION_FACTORY(FMeshToCollectionDataflowNode);
+		DATAFLOW_NODE_REGISTER_CREATION_FACTORY(FStaticMeshToMeshDataflowNode);
+		DATAFLOW_NODE_REGISTER_CREATION_FACTORY(FMeshAppendDataflowNode);
+		DATAFLOW_NODE_REGISTER_CREATION_FACTORY(FMeshBooleanDataflowNode);
+		DATAFLOW_NODE_REGISTER_CREATION_FACTORY(FMeshCopyToPointsDataflowNode);
+		DATAFLOW_NODE_REGISTER_CREATION_FACTORY(FGetMeshDataDataflowNode);
+
+		// Mesh
+		DATAFLOW_NODE_REGISTER_CREATION_FACTORY_NODE_COLORS_BY_CATEGORY("Mesh", FLinearColor(1.f, 0.16f, 0.05f), CDefaultNodeBodyTintColor);
+	}
+}
+
+
+void FPointsToMeshDataflowNode::Evaluate(Dataflow::FContext& Context, const FDataflowOutput* Out) const
+{
+	if (Out->IsA<TObjectPtr<UDynamicMesh>>(&Mesh) || Out->IsA<int32>(&TriangleCount))
+	{
+		const TArray<FVector>& PointsArr = GetValue<TArray<FVector>>(Context, &Points);
+
+		if (PointsArr.Num() > 0)
+		{
+			TObjectPtr<UDynamicMesh> DynamicMesh = NewObject<UDynamicMesh>();
+			DynamicMesh->Reset();
+
+			UE::Geometry::FDynamicMesh3& DynMesh = DynamicMesh->GetMeshRef();
+
+			for (auto& Point : PointsArr)
+			{
+				DynMesh.AppendVertex(Point);
+			}
+
+			SetValue<TObjectPtr<UDynamicMesh>>(Context, DynamicMesh, &Mesh);
+			SetValue<int32>(Context, DynamicMesh->GetTriangleCount(), &TriangleCount);
+		}
+		else
+		{
+			SetValue<TObjectPtr<UDynamicMesh>>(Context, NewObject<UDynamicMesh>(), &Mesh);
+			SetValue<int32>(Context, 0, &TriangleCount);
+		}
+	}
+}
+
+void FBoxToMeshDataflowNode::Evaluate(Dataflow::FContext& Context, const FDataflowOutput* Out) const
+{
+	if (Out->IsA<TObjectPtr<UDynamicMesh>>(&Mesh) || Out->IsA<int32>(&TriangleCount))
+	{
+		TObjectPtr<UDynamicMesh> DynamicMesh = NewObject<UDynamicMesh>();
+		DynamicMesh->Reset();
+
+		UE::Geometry::FDynamicMesh3& DynMesh = DynamicMesh->GetMeshRef();
+
+		FBox InBox = GetValue<FBox>(Context, &Box);
+		FVector Min = InBox.Min;
+		FVector Max = InBox.Max;
+
+		// Add vertices
+		int Vertex0 = DynMesh.AppendVertex(Min);
+		int Vertex1 = DynMesh.AppendVertex(FVector(Min.X, Max.Y, Min.Z));
+		int Vertex2 = DynMesh.AppendVertex(FVector(Max.X, Max.Y, Min.Z));
+		int Vertex3 = DynMesh.AppendVertex(FVector(Max.X, Min.Y, Min.Z));
+		int Vertex4 = DynMesh.AppendVertex(FVector(Min.X, Min.Y, Max.Z));
+		int Vertex5 = DynMesh.AppendVertex(FVector(Min.X, Max.Y, Max.Z));
+		int Vertex6 = DynMesh.AppendVertex(Max);
+		int Vertex7 = DynMesh.AppendVertex(FVector(Max.X, Min.Y, Max.Z));
+
+		// Add triangles
+		int GroupID = 0;
+		DynMesh.AppendTriangle(Vertex0, Vertex1, Vertex3, GroupID);
+		DynMesh.AppendTriangle(Vertex1, Vertex2, Vertex3, GroupID);
+		DynMesh.AppendTriangle(Vertex3, Vertex6, Vertex7, GroupID);
+		DynMesh.AppendTriangle(Vertex3, Vertex2, Vertex6, GroupID);
+		DynMesh.AppendTriangle(Vertex7, Vertex4, Vertex0, GroupID);
+		DynMesh.AppendTriangle(Vertex0, Vertex3, Vertex7, GroupID);
+		DynMesh.AppendTriangle(Vertex0, Vertex4, Vertex5, GroupID);
+		DynMesh.AppendTriangle(Vertex0, Vertex5, Vertex1, GroupID);
+		DynMesh.AppendTriangle(Vertex1, Vertex5, Vertex6, GroupID);
+		DynMesh.AppendTriangle(Vertex6, Vertex2, Vertex1, GroupID);
+		DynMesh.AppendTriangle(Vertex4, Vertex6, Vertex5, GroupID);
+		DynMesh.AppendTriangle(Vertex4, Vertex7, Vertex6, GroupID);
+
+		SetValue<TObjectPtr<UDynamicMesh>>(Context, DynamicMesh, &Mesh);
+		SetValue<int32>(Context, DynamicMesh->GetTriangleCount(), &TriangleCount);
+	}
+}
+
+void FMeshInfoDataflowNode::Evaluate(Dataflow::FContext& Context, const FDataflowOutput* Out) const
+{
+	if (Out->IsA<FString>(&InfoString))
+	{
+		UE::Geometry::FDynamicMesh3& DynMesh = (GetValue<TObjectPtr<UDynamicMesh>>(Context, &Mesh))->GetMeshRef();
+
+		SetValue<FString>(Context, DynMesh.MeshInfoString(), &InfoString);
+	}
+}
+
+void FMeshToCollectionDataflowNode::Evaluate(Dataflow::FContext& Context, const FDataflowOutput* Out) const
+{
+	if (Out->IsA<FManagedArrayCollection>(&Collection))
+	{
+		const UE::Geometry::FDynamicMesh3& DynMesh = (GetValue<TObjectPtr<UDynamicMesh>>(Context, &Mesh))->GetMeshRef();
+		if (DynMesh.VertexCount() > 0)
+		{
+			FMeshDescription MeshDescription;
+			FStaticMeshAttributes Attributes(MeshDescription);
+			Attributes.Register();
+
+			FDynamicMeshToMeshDescription Converter;
+			Converter.Convert(&DynMesh, MeshDescription, true);
+
+			FGeometryCollection NewGeometryCollection = FGeometryCollection();
+			FGeometryCollectionEngineConversion::AppendMeshDescription(&MeshDescription, FString(TEXT("TEST")), 0, FTransform().Identity, &NewGeometryCollection);
+
+			FManagedArrayCollection Col = FManagedArrayCollection();
+			NewGeometryCollection.CopyTo(&Col);
+
+			SetValue<FManagedArrayCollection>(Context, Col, &Collection);
+		}
+		else
+		{
+			SetValue<FManagedArrayCollection>(Context, FManagedArrayCollection(), &Collection);
+		}
+	}
+}
+
+void FStaticMeshToMeshDataflowNode::Evaluate(Dataflow::FContext& Context, const FDataflowOutput* Out) const
+{
+#if WITH_EDITORONLY_DATA
+	if (Out->IsA<TObjectPtr<UDynamicMesh>>(&Mesh))
+	{
+		if (FMeshDescription* MeshDescription = UseHiRes ? StaticMesh->GetHiResMeshDescription() : StaticMesh->GetMeshDescription(LODLevel))
+		{
+			TObjectPtr<UDynamicMesh> DynamicMesh = NewObject<UDynamicMesh>();
+			DynamicMesh->Reset();
+
+			UE::Geometry::FDynamicMesh3& DynMesh = DynamicMesh->GetMeshRef();
+			{
+				FMeshDescriptionToDynamicMesh ConverterToDynamicMesh;
+				ConverterToDynamicMesh.Convert(MeshDescription, DynMesh);
+			}
+
+			SetValue<TObjectPtr<UDynamicMesh>>(Context, DynamicMesh, &Mesh);
+		}
+		else
+		{
+			SetValue<TObjectPtr<UDynamicMesh>>(Context, NewObject<UDynamicMesh>(), &Mesh);
+		}
+	}
+#endif
+}
+
+
+void FMeshAppendDataflowNode::Evaluate(Dataflow::FContext& Context, const FDataflowOutput* Out) const
+{
+	if (Out->IsA<TObjectPtr<UDynamicMesh>>(&Mesh))
+	{
+		const UE::Geometry::FDynamicMesh3& DynMesh1 = (GetValue<TObjectPtr<UDynamicMesh>>(Context, &Mesh1))->GetMeshRef();
+		const UE::Geometry::FDynamicMesh3& DynMesh2 = (GetValue<TObjectPtr<UDynamicMesh>>(Context, &Mesh2))->GetMeshRef();
+
+		if (DynMesh1.VertexCount() > 0 || DynMesh2.VertexCount() > 0)
+		{
+			TObjectPtr<UDynamicMesh> DynamicMesh = NewObject<UDynamicMesh>();
+			DynamicMesh->Reset();
+
+			UE::Geometry::FDynamicMesh3& ResultDynMesh = DynamicMesh->GetMeshRef();
+
+			UE::Geometry::FDynamicMeshEditor MeshEditor(&ResultDynMesh);
+
+			UE::Geometry::FMeshIndexMappings IndexMaps1;
+			MeshEditor.AppendMesh(&DynMesh1, IndexMaps1);
+
+			UE::Geometry::FMeshIndexMappings IndexMaps2;
+			MeshEditor.AppendMesh(&DynMesh2, IndexMaps2);
+
+			SetValue<TObjectPtr<UDynamicMesh>>(Context, DynamicMesh, &Mesh);
+		}
+		else
+		{
+			SetValue<TObjectPtr<UDynamicMesh>>(Context, NewObject<UDynamicMesh>(), &Mesh);
+		}
+	}
+}
+
+void FMeshBooleanDataflowNode::Evaluate(Dataflow::FContext& Context, const FDataflowOutput* Out) const
+{
+	if (Out->IsA<TObjectPtr<UDynamicMesh>>(&Mesh))
+	{
+		const UE::Geometry::FDynamicMesh3& DynMesh1 = (GetValue<TObjectPtr<UDynamicMesh>>(Context, &Mesh1))->GetMeshRef();
+		const UE::Geometry::FDynamicMesh3& DynMesh2 = (GetValue<TObjectPtr<UDynamicMesh>>(Context, &Mesh2))->GetMeshRef();
+
+		if (DynMesh1.VertexCount() > 0 && DynMesh2.VertexCount() > 0)
+		{
+			// Get output
+			TObjectPtr<UDynamicMesh> DynamicMesh = NewObject<UDynamicMesh>();
+			DynamicMesh->Reset();
+			UE::Geometry::FDynamicMesh3& ResultDynMesh = DynamicMesh->GetMeshRef();
+
+			UE::Geometry::FMeshBoolean::EBooleanOp BoolOp = UE::Geometry::FMeshBoolean::EBooleanOp::Intersect;
+			if (Operation == EMeshBooleanOperationEnum::Dataflow_MeshBoolean_Intersect)
+			{
+				BoolOp = UE::Geometry::FMeshBoolean::EBooleanOp::Intersect;
+			}
+			else if (Operation == EMeshBooleanOperationEnum::Dataflow_MeshBoolean_Union)
+			{
+				BoolOp = UE::Geometry::FMeshBoolean::EBooleanOp::Union;
+			}
+			else if (Operation == EMeshBooleanOperationEnum::Dataflow_MeshBoolean_Difference)
+			{
+				BoolOp = UE::Geometry::FMeshBoolean::EBooleanOp::Difference;
+			}
+
+			UE::Geometry::FMeshBoolean Boolean(&DynMesh1, &DynMesh2, &ResultDynMesh, BoolOp);
+			Boolean.bSimplifyAlongNewEdges = true;
+			Boolean.PreserveUVsOnlyForMesh = 0; // slight warping of the autogenerated cell UVs generally doesn't matter
+			Boolean.bWeldSharedEdges = false;
+			Boolean.bTrackAllNewEdges = true;
+			if (!Boolean.Compute())
+			{
+				SetValue<TObjectPtr<UDynamicMesh>>(Context, NewObject<UDynamicMesh>(), &Mesh);
+				return;
+			}
+
+			SetValue<TObjectPtr<UDynamicMesh>>(Context, DynamicMesh, &Mesh);
+		}
+		else
+		{
+			SetValue<TObjectPtr<UDynamicMesh>>(Context, NewObject<UDynamicMesh>(), &Mesh);
+		}
+	}
+}
+
+void FMeshCopyToPointsDataflowNode::Evaluate(Dataflow::FContext& Context, const FDataflowOutput* Out) const
+{
+	if (Out->IsA<TObjectPtr<UDynamicMesh>>(&Mesh))
+	{
+		const TArray<FVector>& InPoints = GetValue<TArray<FVector>>(Context, &Points);
+		const UE::Geometry::FDynamicMesh3& InDynMeshToCopy = (GetValue<TObjectPtr<UDynamicMesh>>(Context, &MeshToCopy))->GetMeshRef();
+
+		if (InPoints.Num() > 0 && InDynMeshToCopy.VertexCount() > 0)
+		{
+			TObjectPtr<UDynamicMesh> DynamicMesh = NewObject<UDynamicMesh>();
+			DynamicMesh->Reset();
+
+			UE::Geometry::FDynamicMesh3& ResultDynMesh = DynamicMesh->GetMeshRef();
+			UE::Geometry::FDynamicMeshEditor MeshEditor(&ResultDynMesh);
+
+			for (auto& Point : InPoints)
+			{
+				UE::Geometry::FDynamicMesh3 DynMeshTemp(InDynMeshToCopy);
+				UE::Geometry::FRefCountVector VertexRefCounts = DynMeshTemp.GetVerticesRefCounts();
+
+				UE::Geometry::FRefCountVector::IndexIterator ItVertexID = VertexRefCounts.BeginIndices();
+				const UE::Geometry::FRefCountVector::IndexIterator ItEndVertexID = VertexRefCounts.EndIndices();
+
+				while (ItVertexID != ItEndVertexID)
+				{
+					DynMeshTemp.SetVertex(*ItVertexID, Scale * DynMeshTemp.GetVertex(*ItVertexID) + Point);
+					++ItVertexID;
+				}
+
+				UE::Geometry::FMeshIndexMappings IndexMaps;
+				MeshEditor.AppendMesh(&DynMeshTemp, IndexMaps);
+			}
+
+			SetValue<TObjectPtr<UDynamicMesh>>(Context, DynamicMesh, &Mesh);
+		}
+		else
+		{
+			SetValue<TObjectPtr<UDynamicMesh>>(Context, NewObject<UDynamicMesh>(), &Mesh);
+		}
+	}
+}
+
+
+void FGetMeshDataDataflowNode::Evaluate(Dataflow::FContext& Context, const FDataflowOutput* Out) const
+{
+	if (Out->IsA<int32>(&VertexCount))
+	{
+		TObjectPtr<UDynamicMesh> DynamicMesh = GetValue<TObjectPtr<UDynamicMesh>>(Context, &Mesh);
+		const UE::Geometry::FDynamicMesh3& DynMesh = DynamicMesh->GetMeshRef();
+
+		SetValue<int32>(Context, DynMesh.VertexCount(), &VertexCount);
+	}
+	else if (Out->IsA<int32>(&EdgeCount))
+	{
+		TObjectPtr<UDynamicMesh> DynamicMesh = GetValue<TObjectPtr<UDynamicMesh>>(Context, &Mesh);
+		const UE::Geometry::FDynamicMesh3& DynMesh = DynamicMesh->GetMeshRef();
+
+		SetValue<int32>(Context, DynMesh.EdgeCount(), &EdgeCount);
+	}
+	else if (Out->IsA<int32>(&TriangleCount))
+	{
+		TObjectPtr<UDynamicMesh> DynamicMesh = GetValue<TObjectPtr<UDynamicMesh>>(Context, &Mesh);
+		const UE::Geometry::FDynamicMesh3& DynMesh = DynamicMesh->GetMeshRef();
+
+		SetValue<int32>(Context, DynMesh.TriangleCount(), &TriangleCount);
+	}
+}
+
diff --git a/Engine/Plugins/Experimental/GeometryCollectionPlugin/Source/GeometryCollectionNodes/Private/Dataflow/GeometryCollectionNodes.cpp b/Engine/Plugins/Experimental/GeometryCollectionPlugin/Source/GeometryCollectionNodes/Private/Dataflow/GeometryCollectionNodes.cpp
--- a/Engine/Plugins/Experimental/GeometryCollectionPlugin/Source/GeometryCollectionNodes/Private/Dataflow/GeometryCollectionNodes.cpp
+++ b/Engine/Plugins/Experimental/GeometryCollectionPlugin/Source/GeometryCollectionNodes/Private/Dataflow/GeometryCollectionNodes.cpp
@@ -55,8 +55,6 @@ namespace Dataflow
 		DATAFLOW_NODE_REGISTER_CREATION_FACTORY(FFloatToStringDataflowNode);
 		DATAFLOW_NODE_REGISTER_CREATION_FACTORY(FMakePointsDataflowNode);
 		DATAFLOW_NODE_REGISTER_CREATION_FACTORY(FMakeBoxDataflowNode);
-		DATAFLOW_NODE_REGISTER_CREATION_FACTORY(FUniformScatterPointsDataflowNode);
-		DATAFLOW_NODE_REGISTER_CREATION_FACTORY(FRadialScatterPointsDataflowNode);
 		DATAFLOW_NODE_REGISTER_CREATION_FACTORY(FMakeLiteralFloatDataflowNode);
 		DATAFLOW_NODE_REGISTER_CREATION_FACTORY(FMakeLiteralIntDataflowNode);
 		DATAFLOW_NODE_REGISTER_CREATION_FACTORY(FMakeLiteralBoolDataflowNode);
@@ -65,7 +63,6 @@ namespace Dataflow
 		DATAFLOW_NODE_REGISTER_CREATION_FACTORY(FBoolToStringDataflowNode);
 		DATAFLOW_NODE_REGISTER_CREATION_FACTORY(FExpandVectorDataflowNode);
 		DATAFLOW_NODE_REGISTER_CREATION_FACTORY(FIntToFloatDataflowNode);
-		DATAFLOW_NODE_REGISTER_CREATION_FACTORY(FVoronoiFractureDataflowNode);
 		DATAFLOW_NODE_REGISTER_CREATION_FACTORY(FStringAppendDataflowNode);
 		DATAFLOW_NODE_REGISTER_CREATION_FACTORY(FRandomFloatDataflowNode);
 		DATAFLOW_NODE_REGISTER_CREATION_FACTORY(FRandomFloatInRangeDataflowNode);
@@ -75,7 +72,6 @@ namespace Dataflow
 		DATAFLOW_NODE_REGISTER_CREATION_FACTORY(FDegreesToRadiansDataflowNode);
 		DATAFLOW_NODE_REGISTER_CREATION_FACTORY(FExplodedViewDataflowNode);
 		DATAFLOW_NODE_REGISTER_CREATION_FACTORY(FCreateNonOverlappingConvexHullsDataflowNode);
-		DATAFLOW_NODE_REGISTER_CREATION_FACTORY(FPlaneCutterDataflowNode);
 		DATAFLOW_NODE_REGISTER_CREATION_FACTORY(FHashStringDataflowNode);
 		DATAFLOW_NODE_REGISTER_CREATION_FACTORY(FHashVectorDataflowNode);
 		DATAFLOW_NODE_REGISTER_CREATION_FACTORY(FFloatToIntDataflowNode);
@@ -84,21 +80,10 @@ namespace Dataflow
 		DATAFLOW_NODE_REGISTER_CREATION_FACTORY(FGetNumArrayElementsDataflowNode);
 		DATAFLOW_NODE_REGISTER_CREATION_FACTORY(FGetBoundingBoxesDataflowNode);
 		DATAFLOW_NODE_REGISTER_CREATION_FACTORY(FGetCentroidsDataflowNode);
-		DATAFLOW_NODE_REGISTER_CREATION_FACTORY(FPointsToMeshDataflowNode);
-		DATAFLOW_NODE_REGISTER_CREATION_FACTORY(FBoxToMeshDataflowNode);
-		DATAFLOW_NODE_REGISTER_CREATION_FACTORY(FMeshInfoDataflowNode);
-		DATAFLOW_NODE_REGISTER_CREATION_FACTORY(FMeshToCollectionDataflowNode);
-		DATAFLOW_NODE_REGISTER_CREATION_FACTORY(FStaticMeshToMeshDataflowNode);
 		DATAFLOW_NODE_REGISTER_CREATION_FACTORY(FTransformDataflowNode);
-		DATAFLOW_NODE_REGISTER_CREATION_FACTORY(FMeshAppendDataflowNode);
-		DATAFLOW_NODE_REGISTER_CREATION_FACTORY(FMeshBooleanDataflowNode);
-		DATAFLOW_NODE_REGISTER_CREATION_FACTORY(FMeshCopyToPointsDataflowNode);
 		DATAFLOW_NODE_REGISTER_CREATION_FACTORY(FCompareIntDataflowNode);
 		DATAFLOW_NODE_REGISTER_CREATION_FACTORY(FBranchDataflowNode);
-		DATAFLOW_NODE_REGISTER_CREATION_FACTORY(FGetMeshDataDataflowNode);
 		DATAFLOW_NODE_REGISTER_CREATION_FACTORY(FGetSchemaDataflowNode);
-		DATAFLOW_NODE_REGISTER_CREATION_FACTORY(FAutoClusterDataflowNode);
-		DATAFLOW_NODE_REGISTER_CREATION_FACTORY(FClusterFlattenDataflowNode);
 		DATAFLOW_NODE_REGISTER_CREATION_FACTORY(FRemoveOnBreakDataflowNode);
 		DATAFLOW_NODE_REGISTER_CREATION_FACTORY(FSetAnchorStateDataflowNode);
 		DATAFLOW_NODE_REGISTER_CREATION_FACTORY(FProximityDataflowNode);
@@ -117,8 +102,6 @@ namespace Dataflow
 		DATAFLOW_NODE_REGISTER_CREATION_FACTORY_NODE_COLORS_BY_CATEGORY("Math", FLinearColor(0.f, 0.4f, 0.8f), CDefaultNodeBodyTintColor);
 		// Generators
 		DATAFLOW_NODE_REGISTER_CREATION_FACTORY_NODE_COLORS_BY_CATEGORY("Generators", FLinearColor(.7f, 0.7f, 0.7f), CDefaultNodeBodyTintColor);
-		// Mesh
-		DATAFLOW_NODE_REGISTER_CREATION_FACTORY_NODE_COLORS_BY_CATEGORY("Mesh", FLinearColor(1.f, 0.16f, 0.05f), CDefaultNodeBodyTintColor);
 	}
 }
 
@@ -176,8 +159,9 @@ void FAppendCollectionAssetsDataflowNode::Evaluate(Dataflow::FContext& Context,
 {
 	if (Out->IsA<DataType>(&Collection1))
 	{
-		DataType InCollection1 = GetValue<DataType>(Context, &Collection1);
-		DataType InCollection2 = GetValue<DataType>(Context, &Collection2);
+		FManagedArrayCollection InCollection1 = GetValue<DataType>(Context, &Collection1);
+		const FManagedArrayCollection& InCollection2 = GetValue<DataType>(Context, &Collection2);
+
 		InCollection1.Append(InCollection2);
 		SetValue<DataType>(Context, InCollection1, &Collection1);
 	}
@@ -280,7 +264,7 @@ void FMakeLiteralStringDataflowNode::Evaluate(Dataflow::FContext& Context, const
 	}
 }
 
-void ComputeBoundingBox(FBoundingBoxDataflowNode::DataType& Collection, FBox& BoundingBox)
+void ComputeBoundingBox(const FManagedArrayCollection& Collection, FBox& BoundingBox)
 {
 	if (Collection.HasAttribute("Transform", FGeometryCollection::TransformGroup) &&
 		Collection.HasAttribute("Parent", FGeometryCollection::TransformGroup) &&
@@ -310,7 +294,7 @@ void FBoundingBoxDataflowNode::Evaluate(Dataflow::FContext& Context, const FData
 {
 	if (Out->IsA<FBox>(&BoundingBox))
 	{
-		DataType InCollection = GetValue<DataType>(Context, &Collection);
+		const FManagedArrayCollection& InCollection = GetValue<FManagedArrayCollection>(Context, &Collection);
 
 		FBox BBox(ForceInit);
 
@@ -393,59 +377,6 @@ void FMakeBoxDataflowNode::Evaluate(Dataflow::FContext& Context, const FDataflow
 	}
 }
 
-void FUniformScatterPointsDataflowNode::Evaluate(Dataflow::FContext& Context, const FDataflowOutput* Out) const
-{
-	if (Out->IsA<TArray<FVector>>(&Points))
-	{
-		FBox BBox = GetValue<FBox>(Context, &BoundingBox);
-		if (BBox.GetVolume() > 0.f)
-		{
-			FRandomStream RandStream(GetValue<float>(Context, &RandomSeed));
-
-			const FVector Extent(BBox.Max - BBox.Min);
-			const int32 NumPoints = RandStream.RandRange(GetValue<int32>(Context, &MinNumberOfPoints), GetValue<int32>(Context, &MaxNumberOfPoints));
-
-			TArray<FVector> PointsArr;
-			PointsArr.Reserve(NumPoints);
-			for (int32 Idx = 0; Idx < NumPoints; ++Idx)
-			{
-				PointsArr.Emplace(BBox.Min + FVector(RandStream.FRand(), RandStream.FRand(), RandStream.FRand()) * Extent);
-			}
-
-			SetValue<TArray<FVector>>(Context, PointsArr, &Points);
-		}
-	}
-}
-
-void FRadialScatterPointsDataflowNode::Evaluate(Dataflow::FContext& Context, const FDataflowOutput* Out) const
-{
-	if (Out->IsA<TArray<FVector>>(&Points))
-	{
-		const FVector::FReal RadialStep = GetValue<float>(Context, &Radius) / GetValue<int32>(Context, &RadialSteps);
-		const FVector::FReal AngularStep = 2 * PI / GetValue<int32>(Context, &AngularSteps);
-
-		FRandomStream RandStream(GetValue<float>(Context, &RandomSeed));
-		FVector UpVector(GetValue<FVector>(Context, &Normal));
-		UpVector.Normalize();
-		FVector BasisX, BasisY;
-		UpVector.FindBestAxisVectors(BasisX, BasisY);
-
-		TArray<FVector> PointsArr;
-
-		FVector::FReal Len = RadialStep * .5;
-		for (int32 ii = 0; ii < GetValue<int32>(Context, &RadialSteps); ++ii, Len += RadialStep)
-		{
-			FVector::FReal Angle = FMath::DegreesToRadians(GetValue<float>(Context, &AngleOffset));
-			for (int32 kk = 0; kk < AngularSteps; ++kk, Angle += AngularStep)
-			{
-				FVector RotatingOffset = Len * (FMath::Cos(Angle) * BasisX + FMath::Sin(Angle) * BasisY);
-				PointsArr.Emplace(GetValue<FVector>(Context, &Center) + RotatingOffset + (RandStream.VRand() * RandStream.FRand() * Variability));
-			}
-		}
-
-		SetValue<TArray<FVector>>(Context, PointsArr, &Points);
-	}
-}
 
 void FMakeLiteralFloatDataflowNode::Evaluate(Dataflow::FContext& Context, const FDataflowOutput* Out) const
 {
@@ -524,96 +455,6 @@ void FIntToFloatDataflowNode::Evaluate(Dataflow::FContext& Context, const FDataf
 	}
 }
 
-float GetMaxVertexMovement(float Grout, float Amplitude, int OctaveNumber, float Persistence)
-{
-	float MaxDisp = Grout;
-	float AmplitudeScaled = Amplitude;
-	for (int32 OctaveIdx = 0; OctaveIdx < OctaveNumber; OctaveIdx++, AmplitudeScaled *= Persistence)
-	{
-		MaxDisp += FMath::Abs(AmplitudeScaled);
-	}
-	return MaxDisp;
-}
-
-void FVoronoiFractureDataflowNode::Evaluate(Dataflow::FContext& Context, const FDataflowOutput* Out) const
-{
-	if (Out->IsA<FManagedArrayCollection>(&Collection))
-	{
-		FManagedArrayCollection CollectionPtr = GetValue<FManagedArrayCollection>(Context, &Collection);
-		if (FGeometryCollection* GeomCollection = CollectionPtr.NewCopy<FGeometryCollection>())
-		{
-			TArray<FVector> Sites = GetValue<TArray<FVector>>(Context, &Points);
-			if (Sites.Num() > 0)
-			{
-				//
-				// Compute BoundingBox for ManagedArrayIn
-				//
-				FBox BoundingBox(ForceInit);
-
-				if (GeomCollection->HasAttribute("Transform", FGeometryCollection::TransformGroup) &&
-					GeomCollection->HasAttribute("Parent", FGeometryCollection::TransformGroup) &&
-					GeomCollection->HasAttribute("TransformIndex", FGeometryCollection::GeometryGroup) &&
-					GeomCollection->HasAttribute("BoundingBox", FGeometryCollection::GeometryGroup))
-				{
-					const TManagedArray<FTransform>& Transforms = GeomCollection->GetAttribute<FTransform>("Transform", FGeometryCollection::TransformGroup);
-					const TManagedArray<int32>& ParentIndices = GeomCollection->GetAttribute<int32>("Parent", FGeometryCollection::TransformGroup);
-					const TManagedArray<int32>& TransformIndices = GeomCollection->GetAttribute<int32>("TransformIndex", FGeometryCollection::GeometryGroup);
-					const TManagedArray<FBox>& BoundingBoxes = GeomCollection->GetAttribute<FBox>("BoundingBox", FGeometryCollection::GeometryGroup);
-
-					TArray<FMatrix> TmpGlobalMatrices;
-					GeometryCollectionAlgo::GlobalMatrices(Transforms, ParentIndices, TmpGlobalMatrices);
-
-					if (TmpGlobalMatrices.Num() > 0)
-					{
-						for (int32 BoxIdx = 0; BoxIdx < BoundingBoxes.Num(); ++BoxIdx)
-						{
-							const int32 TransformIndex = TransformIndices[BoxIdx];
-							BoundingBox += BoundingBoxes[BoxIdx].TransformBy(TmpGlobalMatrices[TransformIndex]);
-						}
-					}
-
-					//
-					// Compute Voronoi Bounds
-					//
-					FBox VoronoiBounds = BoundingBox;
-					VoronoiBounds += FBox(Sites);
-
-					float GroutVal = GetValue<float>(Context, &Grout);
-					float AmplitudeVal = GetValue<float>(Context, &Amplitude);
-					int32 OctaveNumberVal = GetValue<int32>(Context, &OctaveNumber);
-					float PersistenceVal = GetValue<float>(Context, &Persistence);
-
-					VoronoiBounds = VoronoiBounds.ExpandBy(GetMaxVertexMovement(GroutVal, AmplitudeVal, OctaveNumberVal, PersistenceVal) + KINDA_SMALL_NUMBER);
-
-					//
-					// Voronoi Fracture
-					//
-					FNoiseSettings NoiseSettings;
-					NoiseSettings.Amplitude = AmplitudeVal;
-					NoiseSettings.Frequency = GetValue<float>(Context, &Frequency);
-					NoiseSettings.Octaves = OctaveNumberVal;
-					NoiseSettings.PointSpacing = GetValue<float>(Context, &PointSpacing);
-					NoiseSettings.Lacunarity = GetValue<float>(Context, &Lacunarity);
-					NoiseSettings.Persistence = GetValue<float>(Context, &Persistence);;
-
-					FVoronoiDiagram Voronoi(Sites, VoronoiBounds, .1f);
-
-					FPlanarCells VoronoiPlanarCells = FPlanarCells(Sites, Voronoi);
-					VoronoiPlanarCells.InternalSurfaceMaterials.NoiseSettings = NoiseSettings;
-
-					const TArrayView<const int32>& TransformIndicesArray(TransformIndices.GetConstArray());
-
-					float CollisionSampleSpacingVal = GetValue<float>(Context, &CollisionSampleSpacing);
-					float RandomSeedVal = GetValue<float>(Context, &RandomSeed);
-
-					int ResultGeometryIndex = CutMultipleWithPlanarCells(VoronoiPlanarCells, *GeomCollection, TransformIndicesArray, GroutVal, CollisionSampleSpacingVal, RandomSeedVal, FTransform().Identity);
-
-					SetValue<FManagedArrayCollection>(Context, (const FManagedArrayCollection&)(*GeomCollection), &Collection);
-				}
-			}
-		}
-	}
-}
 
 void FStringAppendDataflowNode::Evaluate(Dataflow::FContext& Context, const FDataflowOutput* Out) const
 {
@@ -780,8 +621,9 @@ void FExplodedViewDataflowNode::Evaluate(Dataflow::FContext& Context, const FDat
 {
 	if (Out->IsA<FManagedArrayCollection>(&Collection))
 	{
-		FManagedArrayCollection CollectionPtr = GetValue<FManagedArrayCollection>(Context, &Collection);
-		if (FGeometryCollection* GeomCollection = CollectionPtr.NewCopy<FGeometryCollection>())
+		const FManagedArrayCollection& InCollection = GetValue<FManagedArrayCollection>(Context, &Collection);
+
+		if (TUniquePtr<FGeometryCollection> GeomCollection = TUniquePtr<FGeometryCollection>(InCollection.NewCopy<FGeometryCollection>()))
 		{
 			GeomCollection->AddAttribute<FVector3f>("ExplodedVector", FGeometryCollection::TransformGroup, FManagedArrayCollection::FConstructionParameters(FName()));
 			check(GeomCollection->HasAttribute("ExplodedVector", FGeometryCollection::TransformGroup));
@@ -792,7 +634,7 @@ void FExplodedViewDataflowNode::Evaluate(Dataflow::FContext& Context, const FDat
 			const TManagedArray<FBox>& BoundingBox = GeomCollection->GetAttribute<FBox>("BoundingBox", FGeometryCollection::GeometryGroup);
 
 			// Make sure we have valid "Level"
-			AddAdditionalAttributesIfRequired(GeomCollection);
+			AddAdditionalAttributesIfRequired(GeomCollection.Get());
 
 			const TManagedArray<int32>& Levels = GeomCollection->GetAttribute<int32>("Level", FTransformCollection::TransformGroup);
 			const TManagedArray<int32>& Parent = GeomCollection->GetAttribute<int32>("Parent", FTransformCollection::TransformGroup);
@@ -820,7 +662,7 @@ void FExplodedViewDataflowNode::Evaluate(Dataflow::FContext& Context, const FDat
 				ExplodedVectors[Idx] = FVector3f::ZeroVector;
 				FVector GeoCenter;
 
-				if (GetValidGeoCenter(GeomCollection, TransformToGeometryIndex, Transforms, Children, BoundingBox, Idx, GeoCenter))
+				if (GetValidGeoCenter(GeomCollection.Get(), TransformToGeometryIndex, Transforms, Children, BoundingBox, Idx, GeoCenter))
 				{
 					TransformedCenters[Idx] = GeoCenter;
 					if ((ViewFractureLevel < 0) || Levels[Idx] == ViewFractureLevel)
@@ -864,79 +706,21 @@ void FCreateNonOverlappingConvexHullsDataflowNode::Evaluate(Dataflow::FContext&
 {
 	if (Out->IsA<FManagedArrayCollection>(&Collection))
 	{
-		FManagedArrayCollection CollectionPtr = GetValue<FManagedArrayCollection>(Context, &Collection);
-		if (FGeometryCollection* GeomCollection = CollectionPtr.NewCopy<FGeometryCollection>())
+		const FManagedArrayCollection& InCollection = GetValue<FManagedArrayCollection>(Context, &Collection);
+
+		if (TUniquePtr<FGeometryCollection> GeomCollection = TUniquePtr<FGeometryCollection>(InCollection.NewCopy<FGeometryCollection>()))
 		{
 			float CanRemoveFractionVal = GetValue<float>(Context, &CanRemoveFraction);
 			float CanExceedFractionVal = GetValue<float>(Context, &CanExceedFraction);
 			float SimplificationDistanceThresholdVal = GetValue<float>(Context, &SimplificationDistanceThreshold);
 
-			FGeometryCollectionConvexUtility::FGeometryCollectionConvexData ConvexData = FGeometryCollectionConvexUtility::CreateNonOverlappingConvexHullData(GeomCollection, CanRemoveFractionVal, SimplificationDistanceThresholdVal, CanExceedFractionVal);
+			FGeometryCollectionConvexUtility::FGeometryCollectionConvexData ConvexData = FGeometryCollectionConvexUtility::CreateNonOverlappingConvexHullData(GeomCollection.Get(), CanRemoveFractionVal, SimplificationDistanceThresholdVal, CanExceedFractionVal);
 
 			SetValue<FManagedArrayCollection>(Context, (const FManagedArrayCollection&)(*GeomCollection), &Collection);
 		}
 	}
 }
 
-void FPlaneCutterDataflowNode::Evaluate(Dataflow::FContext& Context, const FDataflowOutput* Out) const
-{
-	if (Out->IsA<FManagedArrayCollection>(&Collection))
-	{
-		FManagedArrayCollection CollectionPtr = GetValue<FManagedArrayCollection>(Context, &Collection);
-		if (FGeometryCollection* GeomCollection = CollectionPtr.NewCopy<FGeometryCollection>())
-		{
-			TArray<FPlane> CuttingPlanes;
-			TArray<FTransform> CuttingPlaneTransforms;
-
-			float RandomSeedVal = GetValue<float>(Context, &RandomSeed);
-			FRandomStream RandStream(RandomSeedVal);
-
-			FBox Bounds = GetValue<FBox>(Context, &BoundingBox);
-			const FVector Extent(Bounds.Max - Bounds.Min);
-
-			CuttingPlaneTransforms.Reserve(CuttingPlaneTransforms.Num() + NumPlanes);
-			for (int32 ii = 0; ii < NumPlanes; ++ii)
-			{
-				FVector Position(Bounds.Min + FVector(RandStream.FRand(), RandStream.FRand(), RandStream.FRand()) * Extent);
-				CuttingPlaneTransforms.Emplace(FTransform(FRotator(RandStream.FRand() * 360.0f, RandStream.FRand() * 360.0f, 0.0f), Position));
-			}
-
-			for (const FTransform& Transform : CuttingPlaneTransforms)
-			{
-				CuttingPlanes.Add(FPlane(Transform.GetLocation(), Transform.GetUnitAxis(EAxis::Z)));
-			}
-
-			FInternalSurfaceMaterials InternalSurfaceMaterials;
-			FNoiseSettings NoiseSettings;
-
-			float AmplitudeVal = GetValue<float>(Context, &Amplitude);
-			if (AmplitudeVal > 0.0f)
-			{
-				NoiseSettings.Amplitude = AmplitudeVal;
-				NoiseSettings.Frequency = GetValue<float>(Context, &Frequency);
-				NoiseSettings.Lacunarity = GetValue<float>(Context, &Lacunarity);
-				NoiseSettings.Persistence = GetValue<float>(Context, &Persistence);
-				NoiseSettings.Octaves = GetValue<int32>(Context, &OctaveNumber);
-				NoiseSettings.PointSpacing = GetValue<float>(Context, &PointSpacing);
-
-				InternalSurfaceMaterials.NoiseSettings = NoiseSettings;
-			}
-
-			if (GeomCollection->HasAttribute("TransformIndex", "Geometry"))
-			{
-				const TManagedArray<int32>& TransformIndices = GeomCollection->GetAttribute<int32>("TransformIndex", "Geometry");
-				const TArrayView<const int32>& TransformIndicesArray(TransformIndices.GetConstArray());
-
-				float CollisionSampleSpacingVal = GetValue<float>(Context, &CollisionSampleSpacing);
-				float GroutVal = GetValue<float>(Context, &Grout);
-
-				int ResultGeometryIndex = CutMultipleWithMultiplePlanes(CuttingPlanes, InternalSurfaceMaterials, *GeomCollection, TransformIndicesArray, GroutVal, CollisionSampleSpacingVal, RandomSeedVal, FTransform().Identity);
-
-				SetValue<FManagedArrayCollection>(Context, (const FManagedArrayCollection&)(*GeomCollection), &Collection);
-			}
-		}
-	}
-}
 
 void FHashStringDataflowNode::Evaluate(Dataflow::FContext& Context, const FDataflowOutput* Out) const
 {
@@ -1065,185 +849,61 @@ void FGetBoundingBoxesDataflowNode::Evaluate(Dataflow::FContext& Context, const
 {
 	if (Out->IsA<TArray<FBox>>(&BoundingBoxes))
 	{
-		FManagedArrayCollection InCollection = GetValue<FManagedArrayCollection>(Context, &Collection);
-		FDataflowTransformSelection InTransformSelection = GetValue<FDataflowTransformSelection>(Context, &TransformSelection);
+		const FManagedArrayCollection& InCollection = GetValue<FManagedArrayCollection>(Context, &Collection);
+		const FDataflowTransformSelection& InTransformSelection = GetValue<FDataflowTransformSelection>(Context, &TransformSelection);
 
-		const TManagedArray<FBox>& BoundingBoxesArr = GeometryCollection::Facades::FBoundsFacade(InCollection).GetBoundingBoxes();
+		const TManagedArray<FBox>& InBoundingBoxes = GeometryCollection::Facades::FBoundsFacade(InCollection).GetBoundingBoxes();
 
-		SetValue<TArray<FBox>>(Context, BoundingBoxesArr.GetConstArray(), &BoundingBoxes);
-	}
-}
-
-void FGetCentroidsDataflowNode::Evaluate(Dataflow::FContext& Context, const FDataflowOutput* Out) const
-{
-	if (Out->IsA<TArray<FVector>>(&Centroids))
-	{
-		FManagedArrayCollection InCollection = GetValue<FManagedArrayCollection>(Context, &Collection);
-		FDataflowTransformSelection InTransformSelection = GetValue<FDataflowTransformSelection>(Context, &TransformSelection);
-
-		const TManagedArray<FBox>& BoundingBoxes = GeometryCollection::Facades::FBoundsFacade(InCollection).GetBoundingBoxes();
-
-		TArray<FVector> CentroidsArr;
-		for (int32 Idx = 0; Idx < BoundingBoxes.Num(); ++Idx)
+		TArray<FBox> BoundingBoxesArr;
+		for (int32 Idx = 0; Idx < InBoundingBoxes.Num(); ++Idx)
 		{
-			const FBox& BoundingBox = BoundingBoxes[Idx];
-			if (BoundingBox.IsValid)
+			if (IsConnected<FDataflowTransformSelection>(&TransformSelection))
 			{
-				if (InTransformSelection.Num() > 0)
-				{
-					if (InTransformSelection.IsSelected(Idx))
-					{
-						CentroidsArr.Add(BoundingBox.GetCenter());
-					}
-				}
-				else
+				if (InTransformSelection.IsSelected(Idx))
 				{
-					CentroidsArr.Add(BoundingBox.GetCenter());
+					BoundingBoxesArr.Add(InBoundingBoxes[Idx]);
 				}
 			}
-		}
-		SetValue<TArray<FVector>>(Context, CentroidsArr, &Centroids);
-	}
-}
-
-void FPointsToMeshDataflowNode::Evaluate(Dataflow::FContext& Context, const FDataflowOutput* Out) const
-{
-	if (Out->IsA<TObjectPtr<UDynamicMesh>>(&Mesh) || Out->IsA<int32>(&TriangleCount))
-	{
-		TArray<FVector> PointsArr = GetValue<TArray<FVector>>(Context, &Points);
-		if (PointsArr.Num() > 0)
-		{
-			TObjectPtr<UDynamicMesh> DynamicMesh = NewObject<UDynamicMesh>();
-			DynamicMesh->Reset();
-
-			UE::Geometry::FDynamicMesh3& DynMesh = DynamicMesh->GetMeshRef();
-
-			for (auto& Point : PointsArr)
+			else
 			{
-				DynMesh.AppendVertex(Point);
+				BoundingBoxesArr.Add(InBoundingBoxes[Idx]);
 			}
 
-			SetValue<TObjectPtr<UDynamicMesh>>(Context, DynamicMesh, &Mesh);
-			SetValue<int32>(Context, DynamicMesh->GetTriangleCount(), &TriangleCount);
 		}
-		else
-		{
-			SetValue<TObjectPtr<UDynamicMesh>>(Context, NewObject<UDynamicMesh>(), &Mesh);
-			SetValue<int32>(Context, 0, &TriangleCount);
-		}
-	}
-}
-
-void FBoxToMeshDataflowNode::Evaluate(Dataflow::FContext& Context, const FDataflowOutput* Out) const
-{
-	if (Out->IsA<TObjectPtr<UDynamicMesh>>(&Mesh) || Out->IsA<int32>(&TriangleCount))
-	{
-		TObjectPtr<UDynamicMesh> DynamicMesh = NewObject<UDynamicMesh>();
-		DynamicMesh->Reset();
-
-		UE::Geometry::FDynamicMesh3& DynMesh = DynamicMesh->GetMeshRef();
-
-		FBox InBox = GetValue<FBox>(Context, &Box);
-		FVector Min = InBox.Min;
-		FVector Max = InBox.Max;
-
-		// Add vertices
-		int Vertex0 = DynMesh.AppendVertex(Min);
-		int Vertex1 = DynMesh.AppendVertex(FVector(Min.X, Max.Y, Min.Z));
-		int Vertex2 = DynMesh.AppendVertex(FVector(Max.X, Max.Y, Min.Z));
-		int Vertex3 = DynMesh.AppendVertex(FVector(Max.X, Min.Y, Min.Z));
-		int Vertex4 = DynMesh.AppendVertex(FVector(Min.X, Min.Y, Max.Z));
-		int Vertex5 = DynMesh.AppendVertex(FVector(Min.X, Max.Y, Max.Z));
-		int Vertex6 = DynMesh.AppendVertex(Max);
-		int Vertex7 = DynMesh.AppendVertex(FVector(Max.X, Min.Y, Max.Z));
-
-		// Add triangles
-		int GroupID = 0;
-		DynMesh.AppendTriangle(Vertex0, Vertex1, Vertex3, GroupID);
-		DynMesh.AppendTriangle(Vertex1, Vertex2, Vertex3, GroupID);
-		DynMesh.AppendTriangle(Vertex3, Vertex6, Vertex7, GroupID);
-		DynMesh.AppendTriangle(Vertex3, Vertex2, Vertex6, GroupID);
-		DynMesh.AppendTriangle(Vertex7, Vertex4, Vertex0, GroupID);
-		DynMesh.AppendTriangle(Vertex0, Vertex3, Vertex7, GroupID);
-		DynMesh.AppendTriangle(Vertex0, Vertex4, Vertex5, GroupID);
-		DynMesh.AppendTriangle(Vertex0, Vertex5, Vertex1, GroupID);
-		DynMesh.AppendTriangle(Vertex1, Vertex5, Vertex6, GroupID);
-		DynMesh.AppendTriangle(Vertex6, Vertex2, Vertex1, GroupID);
-		DynMesh.AppendTriangle(Vertex4, Vertex6, Vertex5, GroupID);
-		DynMesh.AppendTriangle(Vertex4, Vertex7, Vertex6, GroupID);
-
-		SetValue<TObjectPtr<UDynamicMesh>>(Context, DynamicMesh, &Mesh);
-		SetValue<int32>(Context, DynamicMesh->GetTriangleCount(), &TriangleCount);
-	}
-}
-
-void FMeshInfoDataflowNode::Evaluate(Dataflow::FContext& Context, const FDataflowOutput* Out) const
-{
-	if (Out->IsA<FString>(&InfoString))
-	{
-		UE::Geometry::FDynamicMesh3& DynMesh = (GetValue<TObjectPtr<UDynamicMesh>>(Context, &Mesh))->GetMeshRef();
-
-		SetValue<FString>(Context, DynMesh.MeshInfoString(), &InfoString);
+		SetValue<TArray<FBox>>(Context, BoundingBoxesArr, &BoundingBoxes);
 	}
 }
 
-void FMeshToCollectionDataflowNode::Evaluate(Dataflow::FContext& Context, const FDataflowOutput* Out) const
+void FGetCentroidsDataflowNode::Evaluate(Dataflow::FContext& Context, const FDataflowOutput* Out) const
 {
-	if (Out->IsA<FManagedArrayCollection>(&Collection))
+	if (Out->IsA<TArray<FVector>>(&Centroids))
 	{
-		UE::Geometry::FDynamicMesh3& DynMesh = (GetValue<TObjectPtr<UDynamicMesh>>(Context, &Mesh))->GetMeshRef();
-		if (DynMesh.VertexCount() > 0)
-		{
-			FMeshDescription MeshDescription;
-			FStaticMeshAttributes Attributes(MeshDescription);
-			Attributes.Register();
-
-			FDynamicMeshToMeshDescription Converter;
-			Converter.Convert(&DynMesh, MeshDescription, true);
+		const FManagedArrayCollection& InCollection = GetValue<FManagedArrayCollection>(Context, &Collection);
+		const FDataflowTransformSelection& InTransformSelection = GetValue<FDataflowTransformSelection>(Context, &TransformSelection);
 
-			FGeometryCollection NewGeometryCollection = FGeometryCollection();
-			FGeometryCollectionEngineConversion::AppendMeshDescription(&MeshDescription, FString(TEXT("TEST")), 0, FTransform().Identity, &NewGeometryCollection);
+		const TManagedArray<FBox>& InBoundingBoxes = GeometryCollection::Facades::FBoundsFacade(InCollection).GetBoundingBoxes();
 
-			FManagedArrayCollection Col = FManagedArrayCollection();
-			NewGeometryCollection.CopyTo(&Col);
-
-			SetValue<FManagedArrayCollection>(Context, Col, &Collection);
-		}
-		else
-		{
-			SetValue<FManagedArrayCollection>(Context, FManagedArrayCollection(), &Collection);
-		}
-	}
-}
-
-void FStaticMeshToMeshDataflowNode::Evaluate(Dataflow::FContext& Context, const FDataflowOutput* Out) const
-{
-#if WITH_EDITORONLY_DATA
-	if (Out->IsA<TObjectPtr<UDynamicMesh>>(&Mesh))
-	{
-#if WITH_EDITORONLY_DATA
-		if (FMeshDescription* MeshDescription = UseHiRes ? StaticMesh->GetHiResMeshDescription() : StaticMesh->GetMeshDescription(LODLevel))
+		TArray<FVector> CentroidsArr;
+		for (int32 Idx = 0; Idx < InBoundingBoxes.Num(); ++Idx)
 		{
-			TObjectPtr<UDynamicMesh> DynamicMesh = NewObject<UDynamicMesh>();
-			DynamicMesh->Reset();
-
-			UE::Geometry::FDynamicMesh3& DynMesh = DynamicMesh->GetMeshRef();
+			const FBox& BoundingBox = InBoundingBoxes[Idx];
+			if (IsConnected<FDataflowTransformSelection>(&TransformSelection))
 			{
-				FMeshDescriptionToDynamicMesh ConverterToDynamicMesh;
-				ConverterToDynamicMesh.Convert(MeshDescription, DynMesh);
+				if (InTransformSelection.IsSelected(Idx))
+				{
+					CentroidsArr.Add(BoundingBox.GetCenter());
+				}
+			}
+			else
+			{
+				CentroidsArr.Add(BoundingBox.GetCenter());
 			}
-
-			SetValue<TObjectPtr<UDynamicMesh>>(Context, DynamicMesh, &Mesh);
-		}
-		else
-#endif
-		{
-			SetValue<TObjectPtr<UDynamicMesh>>(Context, NewObject<UDynamicMesh>(), &Mesh);
 		}
+		SetValue<TArray<FVector>>(Context, CentroidsArr, &Centroids);
 	}
-#endif
 }
 
+
 void FTransformDataflowNode::Evaluate(Dataflow::FContext& Context, const FDataflowOutput* Out) const
 {
 	if (Out->IsA<TObjectPtr<UDynamicMesh>>(&Mesh))
@@ -1252,126 +912,6 @@ void FTransformDataflowNode::Evaluate(Dataflow::FContext& Context, const FDatafl
 	}
 }
 
-void FMeshAppendDataflowNode::Evaluate(Dataflow::FContext& Context, const FDataflowOutput* Out) const
-{
-	if (Out->IsA<TObjectPtr<UDynamicMesh>>(&Mesh))
-	{
-		UE::Geometry::FDynamicMesh3& DynMesh1 = (GetValue<TObjectPtr<UDynamicMesh>>(Context, &Mesh1))->GetMeshRef();
-		UE::Geometry::FDynamicMesh3& DynMesh2 = (GetValue<TObjectPtr<UDynamicMesh>>(Context, &Mesh2))->GetMeshRef();
-
-		if (DynMesh1.VertexCount() > 0 || DynMesh2.VertexCount() > 0)
-		{
-			TObjectPtr<UDynamicMesh> DynamicMesh = NewObject<UDynamicMesh>();
-			DynamicMesh->Reset();
-
-			UE::Geometry::FDynamicMesh3& ResultDynMesh = DynamicMesh->GetMeshRef();
-
-			UE::Geometry::FDynamicMeshEditor MeshEditor(&ResultDynMesh);
-
-			UE::Geometry::FMeshIndexMappings IndexMaps1;
-			MeshEditor.AppendMesh(&DynMesh1, IndexMaps1);
-
-			UE::Geometry::FMeshIndexMappings IndexMaps2;
-			MeshEditor.AppendMesh(&DynMesh2, IndexMaps2);
-
-			SetValue<TObjectPtr<UDynamicMesh>>(Context, DynamicMesh, &Mesh);
-		}
-		else
-		{
-			SetValue<TObjectPtr<UDynamicMesh>>(Context, NewObject<UDynamicMesh>(), &Mesh);
-		}
-	}
-}
-
-void FMeshBooleanDataflowNode::Evaluate(Dataflow::FContext& Context, const FDataflowOutput* Out) const
-{
-	if (Out->IsA<TObjectPtr<UDynamicMesh>>(&Mesh))
-	{
-		UE::Geometry::FDynamicMesh3& DynMesh1 = (GetValue<TObjectPtr<UDynamicMesh>>(Context, &Mesh1))->GetMeshRef();
-		UE::Geometry::FDynamicMesh3& DynMesh2 = (GetValue<TObjectPtr<UDynamicMesh>>(Context, &Mesh2))->GetMeshRef();
-
-		if (DynMesh1.VertexCount() > 0 && DynMesh2.VertexCount() > 0)
-		{
-			// Get output
-			TObjectPtr<UDynamicMesh> DynamicMesh = NewObject<UDynamicMesh>();
-			DynamicMesh->Reset();
-			UE::Geometry::FDynamicMesh3& ResultDynMesh = DynamicMesh->GetMeshRef();
-
-			UE::Geometry::FMeshBoolean::EBooleanOp BoolOp = UE::Geometry::FMeshBoolean::EBooleanOp::Intersect;
-			if (Operation == EMeshBooleanOperationEnum::Dataflow_MeshBoolean_Intersect)
-			{
-				BoolOp = UE::Geometry::FMeshBoolean::EBooleanOp::Intersect;
-			}
-			else if (Operation == EMeshBooleanOperationEnum::Dataflow_MeshBoolean_Union)
-			{
-				BoolOp = UE::Geometry::FMeshBoolean::EBooleanOp::Union;
-			}
-			else if (Operation == EMeshBooleanOperationEnum::Dataflow_MeshBoolean_Difference)
-			{
-				BoolOp = UE::Geometry::FMeshBoolean::EBooleanOp::Difference;
-			}
-
-			UE::Geometry::FMeshBoolean Boolean(&DynMesh1, &DynMesh2, &ResultDynMesh, BoolOp);
-			Boolean.bSimplifyAlongNewEdges = true;
-			Boolean.PreserveUVsOnlyForMesh = 0; // slight warping of the autogenerated cell UVs generally doesn't matter
-			Boolean.bWeldSharedEdges = false;
-			Boolean.bTrackAllNewEdges = true;
-			if (!Boolean.Compute())
-			{
-				SetValue<TObjectPtr<UDynamicMesh>>(Context, NewObject<UDynamicMesh>(), &Mesh);
-				return;
-			}
-
-			SetValue<TObjectPtr<UDynamicMesh>>(Context, DynamicMesh, &Mesh);
-		}
-		else
-		{
-			SetValue<TObjectPtr<UDynamicMesh>>(Context, NewObject<UDynamicMesh>(), &Mesh);
-		}
-	}
-}
-
-void FMeshCopyToPointsDataflowNode::Evaluate(Dataflow::FContext& Context, const FDataflowOutput* Out) const
-{
-	if (Out->IsA<TObjectPtr<UDynamicMesh>>(&Mesh))
-	{
-		TArray<FVector> InPoints = GetValue<TArray<FVector>>(Context, &Points);
-		UE::Geometry::FDynamicMesh3& InDynMeshToCopy = (GetValue<TObjectPtr<UDynamicMesh>>(Context, &MeshToCopy))->GetMeshRef();
-
-		if (InPoints.Num() > 0 && InDynMeshToCopy.VertexCount() > 0)
-		{
-			TObjectPtr<UDynamicMesh> DynamicMesh = NewObject<UDynamicMesh>();
-			DynamicMesh->Reset();
-
-			UE::Geometry::FDynamicMesh3& ResultDynMesh = DynamicMesh->GetMeshRef();
-			UE::Geometry::FDynamicMeshEditor MeshEditor(&ResultDynMesh);
-
-			for (auto& Point : InPoints)
-			{
-				UE::Geometry::FDynamicMesh3 DynMeshTemp(InDynMeshToCopy);
-				UE::Geometry::FRefCountVector VertexRefCounts = DynMeshTemp.GetVerticesRefCounts();
-
-				UE::Geometry::FRefCountVector::IndexIterator ItVertexID = VertexRefCounts.BeginIndices();
-				const UE::Geometry::FRefCountVector::IndexIterator ItEndVertexID = VertexRefCounts.EndIndices();
-
-				while (ItVertexID != ItEndVertexID)
-				{
-					DynMeshTemp.SetVertex(*ItVertexID, Scale * DynMeshTemp.GetVertex(*ItVertexID) + Point);
-					++ItVertexID;
-				}
-
-				UE::Geometry::FMeshIndexMappings IndexMaps;
-				MeshEditor.AppendMesh(&DynMeshTemp, IndexMaps);
-			}
-
-			SetValue<TObjectPtr<UDynamicMesh>>(Context, DynamicMesh, &Mesh);
-		}
-		else
-		{
-			SetValue<TObjectPtr<UDynamicMesh>>(Context, NewObject<UDynamicMesh>(), &Mesh);
-		}
-	}
-}
 
 void FCompareIntDataflowNode::Evaluate(Dataflow::FContext& Context, const FDataflowOutput* Out) const
 {
@@ -1421,20 +961,6 @@ void FBranchDataflowNode::Evaluate(Dataflow::FContext& Context, const FDataflowO
 	}
 }
 
-void FGetMeshDataDataflowNode::Evaluate(Dataflow::FContext& Context, const FDataflowOutput* Out) const
-{
-	if (Out->IsA<int32>(&VertexCount) ||
-		Out->IsA<int32>(&EdgeCount) ||
-		Out->IsA<int32>(&TriangleCount))
-	{
-		TObjectPtr<UDynamicMesh> DynamicMesh = GetValue<TObjectPtr<UDynamicMesh>>(Context, &Mesh);
-		UE::Geometry::FDynamicMesh3& DynMesh = DynamicMesh->GetMeshRef();
-
-		SetValue<int32>(Context, DynMesh.VertexCount(), &VertexCount);
-		SetValue<int32>(Context, DynMesh.EdgeCount(), &EdgeCount);
-		SetValue<int32>(Context, DynMesh.TriangleCount(), &TriangleCount);
-	}
-}
 
 namespace {
 	inline FName GetArrayTypeString(FManagedArrayCollection::EArrayType ArrayType)
@@ -1454,7 +980,7 @@ void FGetSchemaDataflowNode::Evaluate(Dataflow::FContext& Context, const FDatafl
 {
 	if (Out->IsA<FString>(&String))
 	{
-		FManagedArrayCollection InCollection = GetValue<FManagedArrayCollection>(Context, &Collection);
+		const FManagedArrayCollection& InCollection = GetValue<FManagedArrayCollection>(Context, &Collection);
 
 		FString OutputStr;
 		OutputStr.Appendf(TEXT("\n----------------------------------------\n"));
@@ -1485,88 +1011,26 @@ void FGetSchemaDataflowNode::Evaluate(Dataflow::FContext& Context, const FDatafl
 	}
 }
 
-void FAutoClusterDataflowNode::Evaluate(Dataflow::FContext& Context, const FDataflowOutput* Out) const
+
+void FRemoveOnBreakDataflowNode::Evaluate(Dataflow::FContext& Context, const FDataflowOutput* Out) const
 {
 	if (Out->IsA<FManagedArrayCollection>(&Collection))
 	{
 		FManagedArrayCollection InCollection = GetValue<FManagedArrayCollection>(Context, &Collection);
-		FDataflowTransformSelection InTransformSelection = GetValue<FDataflowTransformSelection>(Context, &TransformSelection);
-
-		if (FGeometryCollection* GeomCollection = InCollection.NewCopy<FGeometryCollection>())
-		{
-			EClusterSizeMethodEnum InClusterSizeMethod = ClusterSizeMethod;
-			int32 InClusterSites = GetValue<int32>(Context, &ClusterSites);
-			float InClusterFraction = GetValue<float>(Context, &ClusterFraction);
-			float InSiteSize = GetValue<float>(Context, &SiteSize);
-			bool InAutoCluster = AutoCluster;
-			bool InAvoidIsolated = AvoidIsolated;
-
-			TArray<int32> SelectedBones;
-			InTransformSelection.AsArray(SelectedBones);
-
-			FFractureEngineClustering::AutoCluster(*GeomCollection, 
-				SelectedBones, 
-				(EFractureEngineClusterSizeMethod)InClusterSizeMethod, 
-				InClusterSites,
-				InClusterFraction,
-				InSiteSize,
-				InAutoCluster,
-				InAvoidIsolated);
-
-			SetValue<FManagedArrayCollection>(Context, (const FManagedArrayCollection&)(*GeomCollection), &Collection);
-		}
-	}
-}
 
-
-void FClusterFlattenDataflowNode::Evaluate(Dataflow::FContext& Context, const FDataflowOutput* Out) const
-{
-	if (Out->IsA<FManagedArrayCollection>(&Collection))
-	{
-		FManagedArrayCollection CollectionPtr = GetValue<FManagedArrayCollection>(Context, &Collection);
-		if (FGeometryCollection* GeomCollection = CollectionPtr.NewCopy<FGeometryCollection>())
+		if (!InCollection.HasAttribute("RemoveOnBreak", FGeometryCollection::TransformGroup))
 		{
-			FGeometryCollectionClusteringUtility::UpdateHierarchyLevelOfChildren(GeomCollection, -1);
-
-			const TManagedArray<int32>& Levels = GeomCollection->GetAttribute<int32>("Level", FGeometryCollection::TransformGroup);
-
-			// Populate Selected Bones in an Array
-			// @todo(harsha) Implement with Selection
-			// For every bone in selected array: [ClusterIndex]
-			int32 ClusterIndex = 0;
-			TArray<int32> LeafBones;
-			FGeometryCollectionClusteringUtility::GetLeafBones(GeomCollection, ClusterIndex, true, LeafBones);
-			FGeometryCollectionClusteringUtility::ClusterBonesUnderExistingNode(GeomCollection, ClusterIndex, LeafBones);
-			FGeometryCollectionClusteringUtility::RemoveDanglingClusters(GeomCollection);
-			// End for
-
-			FGeometryCollectionClusteringUtility::UpdateHierarchyLevelOfChildren(GeomCollection, -1);
-			SetValue<FManagedArrayCollection>(Context, (const FManagedArrayCollection&)(*GeomCollection), &Collection);
+			TManagedArray<FVector4f>& NewRemoveOnBreak = InCollection.AddAttribute<FVector4f>("RemoveOnBreak", FGeometryCollection::TransformGroup);
+			NewRemoveOnBreak.Fill(FRemoveOnBreakData::DisabledPackedData);
 		}
-	}
-}
-
-
-void FRemoveOnBreakDataflowNode::Evaluate(Dataflow::FContext& Context, const FDataflowOutput* Out) const
-{
-	if (Out->IsA<FManagedArrayCollection>(&Collection))
-	{
-		FManagedArrayCollection CollectionPtr = GetValue<FManagedArrayCollection>(Context, &Collection);
-		if (FGeometryCollection* GeometryCollection = CollectionPtr.NewCopy<FGeometryCollection>())
-		{
-			if (!GeometryCollection->HasAttribute("RemoveOnBreak", FGeometryCollection::TransformGroup))
-			{
-				TManagedArray<FVector4f>& NewRemoveOnBreak = GeometryCollection->AddAttribute<FVector4f>("RemoveOnBreak", FGeometryCollection::TransformGroup);
-				NewRemoveOnBreak.Fill(FRemoveOnBreakData::DisabledPackedData);
-			}
 
-			TManagedArray<FVector4f>& RemoveOnBreak = GeometryCollection->ModifyAttribute<FVector4f>("RemoveOnBreak", FGeometryCollection::TransformGroup);
+		TManagedArray<FVector4f>& RemoveOnBreak = InCollection.ModifyAttribute<FVector4f>("RemoveOnBreak", FGeometryCollection::TransformGroup);
 
-			const FVector2f PostBreakTimerData = GetValue<FVector2f>(Context, &PostBreakTimer);
-			const FVector2f RemovalTimerData = GetValue<FVector2f>(Context, &RemovalTimer);
-			const bool ClusterCrumblingData = GetValue<bool>(Context, &ClusterCrumbling);
+		const FVector2f PostBreakTimerData = GetValue<FVector2f>(Context, &PostBreakTimer);
+		const FVector2f RemovalTimerData = GetValue<FVector2f>(Context, &RemovalTimer);
+		const bool ClusterCrumblingData = GetValue<bool>(Context, &ClusterCrumbling);
 
-			RemoveOnBreak.Fill(FVector4f{ PostBreakTimerData.X, PostBreakTimerData.Y, RemovalTimerData.X, RemovalTimerData.Y });
+		RemoveOnBreak.Fill(FVector4f{ PostBreakTimerData.X, PostBreakTimerData.Y, RemovalTimerData.X, RemovalTimerData.Y });
 
 			// @todo(harsha) Implement with Selection
 			// const FRemoveOnBreakData RemoveOnBreakData(true, PostBreakTimerData, ClusterCrumblingData, RemovalTimerData);
@@ -1583,8 +1047,7 @@ void FRemoveOnBreakDataflowNode::Evaluate(Dataflow::FContext& Context, const FDa
 				// }
 			// }
 
-			SetValue<FManagedArrayCollection>(Context, (const FManagedArrayCollection&)(*GeometryCollection), &Collection);
-		}
+		SetValue<FManagedArrayCollection>(Context, InCollection, &Collection);
 	}
 }
 
@@ -1593,10 +1056,10 @@ void FSetAnchorStateDataflowNode::Evaluate(Dataflow::FContext& Context, const FD
 {
 	if (Out->IsA<FManagedArrayCollection>(&Collection))
 	{
-		FManagedArrayCollection InCollection = GetValue<FManagedArrayCollection>(Context, &Collection);
+		const FManagedArrayCollection& InCollection = GetValue<FManagedArrayCollection>(Context, &Collection);
 		FDataflowTransformSelection InTransformSelection = GetValue<FDataflowTransformSelection>(Context, &TransformSelection);
 
-		if (FGeometryCollection* GeomCollection = InCollection.NewCopy<FGeometryCollection>())
+		if (TUniquePtr<FGeometryCollection> GeomCollection = TUniquePtr<FGeometryCollection>(InCollection.NewCopy<FGeometryCollection>()))
 		{
 			Chaos::Facades::FCollectionAnchoringFacade AnchoringFacade(*GeomCollection);
 			if (!AnchoringFacade.HasAnchoredAttribute())
@@ -1626,9 +1089,9 @@ void FProximityDataflowNode::Evaluate(Dataflow::FContext& Context, const FDatafl
 {
 	if (Out->IsA<FManagedArrayCollection>(&Collection))
 	{
-		FManagedArrayCollection InCollection = GetValue<FManagedArrayCollection>(Context, &Collection);
+		const FManagedArrayCollection& InCollection = GetValue<FManagedArrayCollection>(Context, &Collection);
 
-		if (FGeometryCollection* GeomCollection = InCollection.NewCopy<FGeometryCollection>())
+		if (TUniquePtr<FGeometryCollection> GeomCollection = TUniquePtr<FGeometryCollection>(InCollection.NewCopy<FGeometryCollection>()))
 		{
 			FGeometryCollectionProximityPropertiesInterface::FProximityProperties Properties = GeomCollection->GetProximityProperties();
 
@@ -1639,7 +1102,7 @@ void FProximityDataflowNode::Evaluate(Dataflow::FContext& Context, const FDatafl
 			GeomCollection->SetProximityProperties(Properties);
 
 			// Invalidate proximity
-			FGeometryCollectionProximityUtility ProximityUtility(GeomCollection);
+			FGeometryCollectionProximityUtility ProximityUtility(GeomCollection.Get());
 			ProximityUtility.InvalidateProximity();
 
 			SetValue<FManagedArrayCollection>(Context, (const FManagedArrayCollection&)(*GeomCollection), &Collection);
diff --git a/Engine/Plugins/Experimental/GeometryCollectionPlugin/Source/GeometryCollectionNodes/Private/Dataflow/GeometryCollectionNodesPlugin.cpp b/Engine/Plugins/Experimental/GeometryCollectionPlugin/Source/GeometryCollectionNodes/Private/Dataflow/GeometryCollectionNodesPlugin.cpp
--- a/Engine/Plugins/Experimental/GeometryCollectionPlugin/Source/GeometryCollectionNodes/Private/Dataflow/GeometryCollectionNodesPlugin.cpp
+++ b/Engine/Plugins/Experimental/GeometryCollectionPlugin/Source/GeometryCollectionNodes/Private/Dataflow/GeometryCollectionNodesPlugin.cpp
@@ -10,6 +10,9 @@
 #include "Dataflow/GeometryCollectionProcessingNodes.h"
 #include "Dataflow/GeometryCollectionSkeletalMeshNodes.h"
 #include "Dataflow/GeometryCollectionSelectionNodes.h"
+#include "Dataflow/GeometryCollectionMeshNodes.h"
+#include "Dataflow/GeometryCollectionClusteringNodes.h"
+#include "Dataflow/GeometryCollectionFracturingNodes.h"
 
 #define LOCTEXT_NAMESPACE "DataflowNodes"
 
@@ -21,6 +24,9 @@ void IGeometryCollectionNodesPlugin::StartupModule()
 	Dataflow::GeometryCollectionProcessingNodes();
 	Dataflow::GeometryCollectionSkeletalMeshNodes();
 	Dataflow::GeometryCollectionSelectionNodes();
+	Dataflow::GeometryCollectionMeshNodes();
+	Dataflow::GeometryCollectionClusteringNodes();
+	Dataflow::GeometryCollectionFracturingNodes();
 }
 
 void IGeometryCollectionNodesPlugin::ShutdownModule()
diff --git a/Engine/Plugins/Experimental/GeometryCollectionPlugin/Source/GeometryCollectionNodes/Private/Dataflow/GeometryCollectionSelectionNodes.cpp b/Engine/Plugins/Experimental/GeometryCollectionPlugin/Source/GeometryCollectionNodes/Private/Dataflow/GeometryCollectionSelectionNodes.cpp
--- a/Engine/Plugins/Experimental/GeometryCollectionPlugin/Source/GeometryCollectionNodes/Private/Dataflow/GeometryCollectionSelectionNodes.cpp
+++ b/Engine/Plugins/Experimental/GeometryCollectionPlugin/Source/GeometryCollectionNodes/Private/Dataflow/GeometryCollectionSelectionNodes.cpp
@@ -64,19 +64,22 @@ namespace Dataflow
 
 void FCollectionTransformSelectionAllDataflowNode::Evaluate(Dataflow::FContext& Context, const FDataflowOutput* Out) const
 {
-	if (Out->IsA<FManagedArrayCollection>(&Collection) ||
-		Out->IsA<FDataflowTransformSelection>(&TransformSelection))
+	if (Out->IsA<FDataflowTransformSelection>(&TransformSelection))
 	{
-		FManagedArrayCollection InCollection = GetValue<FManagedArrayCollection>(Context, &Collection);
+		const FManagedArrayCollection& InCollection = GetValue<FManagedArrayCollection>(Context, &Collection);
 
-		int32 NumElems = InCollection.NumElements(FGeometryCollection::TransformGroup);
+		int32 NumTransforms = InCollection.NumElements(FGeometryCollection::TransformGroup);
 
 		FDataflowTransformSelection NewTransformSelection;
-		NewTransformSelection.Initialize(NumElems, true);
+		NewTransformSelection.Initialize(NumTransforms, true);
 
-		SetValue<FManagedArrayCollection>(Context, InCollection, &Collection);
 		SetValue<FDataflowTransformSelection>(Context, NewTransformSelection, &TransformSelection);
 	}
+	else if (Out->IsA<FManagedArrayCollection>(&Collection))
+	{
+		const FManagedArrayCollection& InCollection = GetValue<FManagedArrayCollection>(Context, &Collection);
+		SetValue<FManagedArrayCollection>(Context, InCollection, &Collection);
+	}
 }
 
 
@@ -84,8 +87,8 @@ void FCollectionTransformSelectionSetOperationDataflowNode::Evaluate(Dataflow::F
 {
 	if (Out->IsA<FDataflowTransformSelection>(&TransformSelection))
 	{
-		FDataflowTransformSelection InTransformSelectionA = GetValue<FDataflowTransformSelection>(Context, &TransformSelectionA);
-		FDataflowTransformSelection InTransformSelectionB = GetValue<FDataflowTransformSelection>(Context, &TransformSelectionB);
+		const FDataflowTransformSelection& InTransformSelectionA = GetValue<FDataflowTransformSelection>(Context, &TransformSelectionA);
+		const FDataflowTransformSelection& InTransformSelectionB = GetValue<FDataflowTransformSelection>(Context, &TransformSelectionB);
 
 		FDataflowTransformSelection NewTransformSelection;
 
@@ -137,11 +140,10 @@ static void ExpandRecursive(const int32 BoneIndex, int32 Level, const TManagedAr
 	}
 }
 
-
 static void BuildHierarchicalOutput(const TManagedArray<int32>& Parents, 
 	const TManagedArray<TSet<int32>>& Children, 
 	const TManagedArray<FString>& BoneNames,
-	FDataflowTransformSelection& TransformSelection, 
+	const FDataflowTransformSelection& TransformSelection, 
 	FString& OutputStr)
 {
 	TArray<BoneInfo> BoneHierarchy;
@@ -210,8 +212,8 @@ void FCollectionTransformSelectionInfoDataflowNode::Evaluate(Dataflow::FContext&
 {
 	if (Out->IsA<FString>(&String))
 	{
-		FDataflowTransformSelection InTransformSelection = GetValue<FDataflowTransformSelection>(Context, &TransformSelection);
-		FManagedArrayCollection InCollection = GetValue<FManagedArrayCollection>(Context, &Collection);
+		const FDataflowTransformSelection& InTransformSelection = GetValue<FDataflowTransformSelection>(Context, &TransformSelection);
+		const FManagedArrayCollection& InCollection = GetValue<FManagedArrayCollection>(Context, &Collection);
 
 		FString OutputStr;
 
@@ -257,19 +259,22 @@ void FCollectionTransformSelectionInfoDataflowNode::Evaluate(Dataflow::FContext&
 
 void FCollectionTransformSelectionNoneDataflowNode::Evaluate(Dataflow::FContext& Context, const FDataflowOutput* Out) const
 {
-	if (Out->IsA<FManagedArrayCollection>(&Collection) ||
-		Out->IsA<FDataflowTransformSelection>(&TransformSelection))
+	if (Out->IsA<FDataflowTransformSelection>(&TransformSelection))
 	{
-		FManagedArrayCollection InCollection = GetValue<FManagedArrayCollection>(Context, &Collection);
+		const FManagedArrayCollection& InCollection = GetValue<FManagedArrayCollection>(Context, &Collection);
 
-		const int32 NumElems = InCollection.NumElements(FGeometryCollection::TransformGroup);
+		int32 NumTransforms = InCollection.NumElements(FGeometryCollection::TransformGroup);
 
 		FDataflowTransformSelection NewTransformSelection;
-		NewTransformSelection.Initialize(NumElems, false);
+		NewTransformSelection.Initialize(NumTransforms, false);
 
-		SetValue<FManagedArrayCollection>(Context, InCollection, &Collection);
 		SetValue<FDataflowTransformSelection>(Context, NewTransformSelection, &TransformSelection);
 	}
+	else if (Out->IsA<FManagedArrayCollection>(&Collection))
+	{
+		const FManagedArrayCollection& InCollection = GetValue<FManagedArrayCollection>(Context, &Collection);
+		SetValue<FManagedArrayCollection>(Context, InCollection, &Collection);
+	}
 }
 
 
@@ -288,22 +293,21 @@ void FCollectionTransformSelectionInvertDataflowNode::Evaluate(Dataflow::FContex
 
 void FCollectionTransformSelectionRandomDataflowNode::Evaluate(Dataflow::FContext& Context, const FDataflowOutput* Out) const
 {
-	if (Out->IsA<FManagedArrayCollection>(&Collection) ||
-		Out->IsA<FDataflowTransformSelection>(&TransformSelection))
+	if (Out->IsA<FDataflowTransformSelection>(&TransformSelection))
 	{
-		FManagedArrayCollection InCollection = GetValue<FManagedArrayCollection>(Context, &Collection);
+		const FManagedArrayCollection& InCollection = GetValue<FManagedArrayCollection>(Context, &Collection);
 
-		int32 NumElems = InCollection.NumElements(FGeometryCollection::TransformGroup);
+		int32 NumTransforms = InCollection.NumElements(FGeometryCollection::TransformGroup);
 
 		FDataflowTransformSelection NewTransformSelection;
-		NewTransformSelection.Initialize(NumElems, false);
+		NewTransformSelection.Initialize(NumTransforms, false);
 
 		float RandomSeedVal = GetValue<float>(Context, &RandomSeed);
 		float RandomThresholdVal = GetValue<float>(Context, &RandomThreshold);
 
 		FRandomStream Stream(RandomSeedVal);
 
-		for (int32 Idx = 0; Idx < NumElems; ++Idx)
+		for (int32 Idx = 0; Idx < NumTransforms; ++Idx)
 		{
 			if (Deterministic)
 			{
@@ -325,36 +329,33 @@ void FCollectionTransformSelectionRandomDataflowNode::Evaluate(Dataflow::FContex
 			}
 		}
 
-		SetValue<FManagedArrayCollection>(Context, InCollection, &Collection);
 		SetValue<FDataflowTransformSelection>(Context, NewTransformSelection, &TransformSelection);
 	}
+	else if (Out->IsA<FManagedArrayCollection>(&Collection))
+	{
+		const FManagedArrayCollection& InCollection = GetValue<FManagedArrayCollection>(Context, &Collection);
+		SetValue<FManagedArrayCollection>(Context, InCollection, &Collection);
+	}
 }
 
 
 void FCollectionTransformSelectionRootDataflowNode::Evaluate(Dataflow::FContext& Context, const FDataflowOutput* Out) const
 {
-	if (Out->IsA<FManagedArrayCollection>(&Collection) ||
-		Out->IsA<FDataflowTransformSelection>(&TransformSelection))
+	if (Out->IsA<FDataflowTransformSelection>(&TransformSelection))
 	{
-		FManagedArrayCollection InCollection = GetValue<FManagedArrayCollection>(Context, &Collection);
+		const FManagedArrayCollection& InCollection = GetValue<FManagedArrayCollection>(Context, &Collection);
 
 		if (InCollection.HasGroup(FGeometryCollection::TransformGroup) &&
 			InCollection.HasAttribute("Parent", FGeometryCollection::TransformGroup))
 		{
-			const TManagedArray<int32>& Parent = InCollection.GetAttribute<int32>("Parent", FGeometryCollection::TransformGroup);
+			TArray<int32> RootBones;
+			FFractureEngineSelection::GetRootBones(InCollection, RootBones);
 
-			int32 NumElems = InCollection.NumElements(FGeometryCollection::TransformGroup);
+			int32 NumTransforms = InCollection.NumElements(FGeometryCollection::TransformGroup);
 
 			FDataflowTransformSelection NewTransformSelection;
-			NewTransformSelection.Initialize(NumElems, false);
-
-			for (int32 Idx = 0; Idx < NumElems; ++Idx)
-			{
-				if (Parent[Idx] == FGeometryCollection::Invalid)
-				{
-					NewTransformSelection.SetSelected(Idx);
-				}
-			}
+			NewTransformSelection.Initialize(NumTransforms, false);
+			NewTransformSelection.SetFromArray(RootBones);
 
 			SetValue<FDataflowTransformSelection>(Context, NewTransformSelection, &TransformSelection);
 		}
@@ -362,7 +363,10 @@ void FCollectionTransformSelectionRootDataflowNode::Evaluate(Dataflow::FContext&
 		{
 			SetValue<FDataflowTransformSelection>(Context, FDataflowTransformSelection(), &TransformSelection);
 		}
-
+	}
+	else if (Out->IsA<FManagedArrayCollection>(&Collection))
+	{
+		const FManagedArrayCollection& InCollection = GetValue<FManagedArrayCollection>(Context, &Collection);
 		SetValue<FManagedArrayCollection>(Context, InCollection, &Collection);
 	}
 }
@@ -370,17 +374,16 @@ void FCollectionTransformSelectionRootDataflowNode::Evaluate(Dataflow::FContext&
 
 void FCollectionTransformSelectionCustomDataflowNode::Evaluate(Dataflow::FContext& Context, const FDataflowOutput* Out) const
 {
-	if (Out->IsA<FManagedArrayCollection>(&Collection) ||
-		Out->IsA<FDataflowTransformSelection>(&TransformSelection))
+	if (Out->IsA<FDataflowTransformSelection>(&TransformSelection))
 	{
-		FManagedArrayCollection InCollection = GetValue<FManagedArrayCollection>(Context, &Collection);
+		const FManagedArrayCollection& InCollection = GetValue<FManagedArrayCollection>(Context, &Collection);
 
 		if (InCollection.HasGroup(FGeometryCollection::TransformGroup))
 		{
-			const int32 NumElems = InCollection.NumElements(FGeometryCollection::TransformGroup);
+			const int32 NumTransforms = InCollection.NumElements(FGeometryCollection::TransformGroup);
 
 			FDataflowTransformSelection NewTransformSelection;
-			NewTransformSelection.Initialize(NumElems, false);
+			NewTransformSelection.Initialize(NumTransforms, false);
 
 			const FString InBoneIndices = GetValue<FString>(Context, &BoneIndicies);
 
@@ -392,7 +395,7 @@ void FCollectionTransformSelectionCustomDataflowNode::Evaluate(Dataflow::FContex
 				if (IndexStr.IsNumeric())
 				{
 					int32 Index = FCString::Atoi(*IndexStr);
-					if (Index > 0 && Index < NumElems)
+					if (Index >= 0 && Index < NumTransforms)
 					{
 						NewTransformSelection.SetSelected(Index);
 					}
@@ -411,7 +414,10 @@ void FCollectionTransformSelectionCustomDataflowNode::Evaluate(Dataflow::FContex
 		{
 			SetValue<FDataflowTransformSelection>(Context, FDataflowTransformSelection(), &TransformSelection);
 		}
-
+	}
+	else if (Out->IsA<FManagedArrayCollection>(&Collection))
+	{
+		const FManagedArrayCollection& InCollection = GetValue<FManagedArrayCollection>(Context, &Collection);
 		SetValue<FManagedArrayCollection>(Context, InCollection, &Collection);
 	}
 }
@@ -419,18 +425,18 @@ void FCollectionTransformSelectionCustomDataflowNode::Evaluate(Dataflow::FContex
 
 void FCollectionTransformSelectionParentDataflowNode::Evaluate(Dataflow::FContext& Context, const FDataflowOutput* Out) const
 {
-	if (Out->IsA<FDataflowTransformSelection>(&TransformSelection) ||
-		Out->IsA<FManagedArrayCollection>(&Collection))
+	if (Out->IsA<FDataflowTransformSelection>(&TransformSelection))
 	{
-		FManagedArrayCollection InCollection = GetValue<FManagedArrayCollection>(Context, &Collection);
+		const FManagedArrayCollection& InCollection = GetValue<FManagedArrayCollection>(Context, &Collection);
 		FDataflowTransformSelection InTransformSelection = GetValue<FDataflowTransformSelection>(Context, &TransformSelection);
 
-		if (FGeometryCollection* GeomCollection = InCollection.NewCopy<FGeometryCollection>())
-		{
-			FFractureEngineSelection::SelectParent(*GeomCollection, InTransformSelection);
-		}
+		FFractureEngineSelection::SelectParent(InCollection, InTransformSelection);
 		
 		SetValue<FDataflowTransformSelection>(Context, InTransformSelection, &TransformSelection);
+	}
+	else if (Out->IsA<FManagedArrayCollection>(&Collection))
+	{
+		const FManagedArrayCollection& InCollection = GetValue<FManagedArrayCollection>(Context, &Collection);
 		SetValue<FManagedArrayCollection>(Context, InCollection, &Collection);
 	}
 }
@@ -438,40 +444,34 @@ void FCollectionTransformSelectionParentDataflowNode::Evaluate(Dataflow::FContex
 
 void FCollectionTransformSelectionByPercentageDataflowNode::Evaluate(Dataflow::FContext& Context, const FDataflowOutput* Out) const
 {
-	if (Out->IsA<FDataflowTransformSelection>(&TransformSelection) ||
-		Out->IsA<FManagedArrayCollection>(&Collection))
+	if (Out->IsA<FDataflowTransformSelection>(&TransformSelection))
 	{
-		FManagedArrayCollection InCollection = GetValue<FManagedArrayCollection>(Context, &Collection);
 		FDataflowTransformSelection InTransformSelection = GetValue<FDataflowTransformSelection>(Context, &TransformSelection);
 
 		int32 InPercentage = GetValue<int32>(Context, &Percentage);
 		float InRandomSeed = GetValue<float>(Context, &RandomSeed);
 
-		if (FGeometryCollection* GeomCollection = InCollection.NewCopy<FGeometryCollection>())
-		{
-			FFractureEngineSelection::SelectByPercentage(*GeomCollection, InTransformSelection, InPercentage, Deterministic, InRandomSeed);
-		}
+		FFractureEngineSelection::SelectByPercentage(InTransformSelection, InPercentage, Deterministic, InRandomSeed);
 
 		SetValue<FDataflowTransformSelection>(Context, InTransformSelection, &TransformSelection);
-		SetValue<FManagedArrayCollection>(Context, InCollection, &Collection);
 	}
 }
 
 
 void FCollectionTransformSelectionChildrenDataflowNode::Evaluate(Dataflow::FContext& Context, const FDataflowOutput* Out) const
 {
-	if (Out->IsA<FDataflowTransformSelection>(&TransformSelection) ||
-		Out->IsA<FManagedArrayCollection>(&Collection))
+	if (Out->IsA<FDataflowTransformSelection>(&TransformSelection))
 	{
-		FManagedArrayCollection InCollection = GetValue<FManagedArrayCollection>(Context, &Collection);
+		const FManagedArrayCollection& InCollection = GetValue<FManagedArrayCollection>(Context, &Collection);
 		FDataflowTransformSelection InTransformSelection = GetValue<FDataflowTransformSelection>(Context, &TransformSelection);
 
-		if (FGeometryCollection* GeomCollection = InCollection.NewCopy<FGeometryCollection>())
-		{
-			FFractureEngineSelection::SelectChildren(*GeomCollection, InTransformSelection);
-		}
+		FFractureEngineSelection::SelectChildren(InCollection, InTransformSelection);
 
 		SetValue<FDataflowTransformSelection>(Context, InTransformSelection, &TransformSelection);
+	}
+	else if (Out->IsA<FManagedArrayCollection>(&Collection))
+	{
+		const FManagedArrayCollection& InCollection = GetValue<FManagedArrayCollection>(Context, &Collection);
 		SetValue<FManagedArrayCollection>(Context, InCollection, &Collection);
 	}
 }
@@ -479,18 +479,18 @@ void FCollectionTransformSelectionChildrenDataflowNode::Evaluate(Dataflow::FCont
 
 void FCollectionTransformSelectionSiblingsDataflowNode::Evaluate(Dataflow::FContext& Context, const FDataflowOutput* Out) const
 {
-	if (Out->IsA<FDataflowTransformSelection>(&TransformSelection) ||
-		Out->IsA<FManagedArrayCollection>(&Collection))
+	if (Out->IsA<FDataflowTransformSelection>(&TransformSelection))
 	{
-		FManagedArrayCollection InCollection = GetValue<FManagedArrayCollection>(Context, &Collection);
+		const FManagedArrayCollection& InCollection = GetValue<FManagedArrayCollection>(Context, &Collection);
 		FDataflowTransformSelection InTransformSelection = GetValue<FDataflowTransformSelection>(Context, &TransformSelection);
 
-		if (FGeometryCollection* GeomCollection = InCollection.NewCopy<FGeometryCollection>())
-		{
-			FFractureEngineSelection::SelectSiblings(*GeomCollection, InTransformSelection);
-		}
+		FFractureEngineSelection::SelectSiblings(InCollection, InTransformSelection);
 
 		SetValue<FDataflowTransformSelection>(Context, InTransformSelection, &TransformSelection);
+	}
+	else if (Out->IsA<FManagedArrayCollection>(&Collection))
+	{
+		const FManagedArrayCollection& InCollection = GetValue<FManagedArrayCollection>(Context, &Collection);
 		SetValue<FManagedArrayCollection>(Context, InCollection, &Collection);
 	}
 }
@@ -498,18 +498,18 @@ void FCollectionTransformSelectionSiblingsDataflowNode::Evaluate(Dataflow::FCont
 
 void FCollectionTransformSelectionLevelDataflowNode::Evaluate(Dataflow::FContext& Context, const FDataflowOutput* Out) const
 {
-	if (Out->IsA<FDataflowTransformSelection>(&TransformSelection) ||
-		Out->IsA<FManagedArrayCollection>(&Collection))
+	if (Out->IsA<FDataflowTransformSelection>(&TransformSelection))
 	{
-		FManagedArrayCollection InCollection = GetValue<FManagedArrayCollection>(Context, &Collection);
+		const FManagedArrayCollection& InCollection = GetValue<FManagedArrayCollection>(Context, &Collection);
 		FDataflowTransformSelection InTransformSelection = GetValue<FDataflowTransformSelection>(Context, &TransformSelection);
 
-		if (FGeometryCollection* GeomCollection = InCollection.NewCopy<FGeometryCollection>())
-		{
-			FFractureEngineSelection::SelectLevel(*GeomCollection, InTransformSelection);
-		}
+		FFractureEngineSelection::SelectLevel(InCollection, InTransformSelection);
 
 		SetValue<FDataflowTransformSelection>(Context, InTransformSelection, &TransformSelection);
+	}
+	else if (Out->IsA<FManagedArrayCollection>(&Collection))
+	{
+		const FManagedArrayCollection& InCollection = GetValue<FManagedArrayCollection>(Context, &Collection);
 		SetValue<FManagedArrayCollection>(Context, InCollection, &Collection);
 	}
 }
@@ -517,13 +517,12 @@ void FCollectionTransformSelectionLevelDataflowNode::Evaluate(Dataflow::FContext
 
 void FCollectionTransformSelectionContactDataflowNode::Evaluate(Dataflow::FContext& Context, const FDataflowOutput* Out) const
 {
-	if (Out->IsA<FDataflowTransformSelection>(&TransformSelection) ||
-		Out->IsA<FManagedArrayCollection>(&Collection))
+	if (Out->IsA<FDataflowTransformSelection>(&TransformSelection))
 	{
-		FManagedArrayCollection InCollection = GetValue<FManagedArrayCollection>(Context, &Collection);
+		const FManagedArrayCollection& InCollection = GetValue<FManagedArrayCollection>(Context, &Collection);
 		FDataflowTransformSelection InTransformSelection = GetValue<FDataflowTransformSelection>(Context, &TransformSelection);
 
-		if (FGeometryCollection* GeomCollection = InCollection.NewCopy<FGeometryCollection>())
+		if (TUniquePtr<FGeometryCollection> GeomCollection = TUniquePtr<FGeometryCollection>(InCollection.NewCopy<FGeometryCollection>()))
 		{
 			FFractureEngineSelection::SelectContact(*GeomCollection, InTransformSelection);
 		}
@@ -531,24 +530,32 @@ void FCollectionTransformSelectionContactDataflowNode::Evaluate(Dataflow::FConte
 		SetValue<FDataflowTransformSelection>(Context, InTransformSelection, &TransformSelection);
 		SetValue<FManagedArrayCollection>(Context, InCollection, &Collection);
 	}
+	else if (Out->IsA<FManagedArrayCollection>(&Collection))
+	{
+		const FManagedArrayCollection& InCollection = GetValue<FManagedArrayCollection>(Context, &Collection);
+		SetValue<FManagedArrayCollection>(Context, InCollection, &Collection);
+	}
 }
 
 
 void FCollectionTransformSelectionLeafDataflowNode::Evaluate(Dataflow::FContext& Context, const FDataflowOutput* Out) const
 {
-	if (Out->IsA<FDataflowTransformSelection>(&TransformSelection) ||
-		Out->IsA<FManagedArrayCollection>(&Collection))
+	if (Out->IsA<FDataflowTransformSelection>(&TransformSelection))
 	{
-		FManagedArrayCollection InCollection = GetValue<FManagedArrayCollection>(Context, &Collection);
+		const FManagedArrayCollection& InCollection = GetValue<FManagedArrayCollection>(Context, &Collection);
 
 		FDataflowTransformSelection NewTransformSelection;
 
-		if (FGeometryCollection* GeomCollection = InCollection.NewCopy<FGeometryCollection>())
+		if (TUniquePtr<FGeometryCollection> GeomCollection = TUniquePtr<FGeometryCollection>(InCollection.NewCopy<FGeometryCollection>()))
 		{
 			FFractureEngineSelection::SelectLeaf(*GeomCollection, NewTransformSelection);
 		}
 
 		SetValue<FDataflowTransformSelection>(Context, NewTransformSelection, &TransformSelection);
+	}
+	else if (Out->IsA<FManagedArrayCollection>(&Collection))
+	{
+		const FManagedArrayCollection& InCollection = GetValue<FManagedArrayCollection>(Context, &Collection);
 		SetValue<FManagedArrayCollection>(Context, InCollection, &Collection);
 	}
 }
@@ -556,19 +563,26 @@ void FCollectionTransformSelectionLeafDataflowNode::Evaluate(Dataflow::FContext&
 
 void FCollectionTransformSelectionClusterDataflowNode::Evaluate(Dataflow::FContext& Context, const FDataflowOutput* Out) const
 {
-	if (Out->IsA<FDataflowTransformSelection>(&TransformSelection) ||
-		Out->IsA<FManagedArrayCollection>(&Collection))
+	if (Out->IsA<FDataflowTransformSelection>(&TransformSelection))
 	{
-		FManagedArrayCollection InCollection = GetValue<FManagedArrayCollection>(Context, &Collection);
+		const FManagedArrayCollection& InCollection = GetValue<FManagedArrayCollection>(Context, &Collection);
 
+		const int32 NumTransforms = InCollection.NumElements(FGeometryCollection::TransformGroup);
+
+		// TODO: Convert FractureEngine API to use only FManagedArrayCollection
 		FDataflowTransformSelection NewTransformSelection;
+		NewTransformSelection.Initialize(NumTransforms, false);
 
-		if (FGeometryCollection* GeomCollection = InCollection.NewCopy<FGeometryCollection>())
+		if (TUniquePtr<FGeometryCollection> GeomCollection = TUniquePtr<FGeometryCollection>(InCollection.NewCopy<FGeometryCollection>()))
 		{
-			FFractureEngineSelection::SelectCluster(*GeomCollection, NewTransformSelection);
+			FFractureEngineSelection::SelectCluster(*GeomCollection, NewTransformSelection);			
 		}
 
 		SetValue<FDataflowTransformSelection>(Context, NewTransformSelection, &TransformSelection);
+	}
+	else if (Out->IsA<FManagedArrayCollection>(&Collection))
+	{
+		const FManagedArrayCollection& InCollection = GetValue<FManagedArrayCollection>(Context, &Collection);
 		SetValue<FManagedArrayCollection>(Context, InCollection, &Collection);
 	}
 }
@@ -576,21 +590,24 @@ void FCollectionTransformSelectionClusterDataflowNode::Evaluate(Dataflow::FConte
 
 void FCollectionTransformSelectionBySizeDataflowNode::Evaluate(Dataflow::FContext& Context, const FDataflowOutput* Out) const
 {
-	if (Out->IsA<FDataflowTransformSelection>(&TransformSelection) ||
-		Out->IsA<FManagedArrayCollection>(&Collection))
+	if (Out->IsA<FDataflowTransformSelection>(&TransformSelection))
 	{
-		FManagedArrayCollection InCollection = GetValue<FManagedArrayCollection>(Context, &Collection);
+		const FManagedArrayCollection& InCollection = GetValue<FManagedArrayCollection>(Context, &Collection);
 		float InSizeMin = GetValue<float>(Context, &SizeMin);
 		float InSizeMax = GetValue<float>(Context, &SizeMax);
 
 		FDataflowTransformSelection NewTransformSelection;
 
-		if (FGeometryCollection* GeomCollection = InCollection.NewCopy<FGeometryCollection>())
+		if (TUniquePtr<FGeometryCollection> GeomCollection = TUniquePtr<FGeometryCollection>(InCollection.NewCopy<FGeometryCollection>()))
 		{
 			FFractureEngineSelection::SelectBySize(*GeomCollection, NewTransformSelection, InSizeMin, InSizeMax);
 		}
 
 		SetValue<FDataflowTransformSelection>(Context, NewTransformSelection, &TransformSelection);
+	}
+	else if (Out->IsA<FManagedArrayCollection>(&Collection))
+	{
+		const FManagedArrayCollection& InCollection = GetValue<FManagedArrayCollection>(Context, &Collection);
 		SetValue<FManagedArrayCollection>(Context, InCollection, &Collection);
 	}
 }
@@ -598,21 +615,24 @@ void FCollectionTransformSelectionBySizeDataflowNode::Evaluate(Dataflow::FContex
 
 void FCollectionTransformSelectionByVolumeDataflowNode::Evaluate(Dataflow::FContext& Context, const FDataflowOutput* Out) const
 {
-	if (Out->IsA<FDataflowTransformSelection>(&TransformSelection) ||
-		Out->IsA<FManagedArrayCollection>(&Collection))
+	if (Out->IsA<FDataflowTransformSelection>(&TransformSelection))
 	{
-		FManagedArrayCollection InCollection = GetValue<FManagedArrayCollection>(Context, &Collection);
+		const FManagedArrayCollection& InCollection = GetValue<FManagedArrayCollection>(Context, &Collection);
 		float InVolumeMin = GetValue<float>(Context, &VolumeMin);
 		float InVolumeMax = GetValue<float>(Context, &VolumeMax);
 
 		FDataflowTransformSelection NewTransformSelection;
 
-		if (FGeometryCollection* GeomCollection = InCollection.NewCopy<FGeometryCollection>())
+		if (TUniquePtr<FGeometryCollection> GeomCollection = TUniquePtr<FGeometryCollection>(InCollection.NewCopy<FGeometryCollection>()))
 		{
 			FFractureEngineSelection::SelectByVolume(*GeomCollection, NewTransformSelection, InVolumeMin, InVolumeMax);
 		}
 
 		SetValue<FDataflowTransformSelection>(Context, NewTransformSelection, &TransformSelection);
+	}
+	else if (Out->IsA<FManagedArrayCollection>(&Collection))
+	{
+		const FManagedArrayCollection& InCollection = GetValue<FManagedArrayCollection>(Context, &Collection);
 		SetValue<FManagedArrayCollection>(Context, InCollection, &Collection);
 	}
 }
diff --git a/Engine/Plugins/Experimental/GeometryCollectionPlugin/Source/GeometryCollectionNodes/Private/Dataflow/Nodes/GeometryCollectionAssetNodes.cpp b/Engine/Plugins/Experimental/GeometryCollectionPlugin/Source/GeometryCollectionNodes/Private/Dataflow/Nodes/GeometryCollectionAssetNodes.cpp
--- a/Engine/Plugins/Experimental/GeometryCollectionPlugin/Source/GeometryCollectionNodes/Private/Dataflow/Nodes/GeometryCollectionAssetNodes.cpp
+++ b/Engine/Plugins/Experimental/GeometryCollectionPlugin/Source/GeometryCollectionNodes/Private/Dataflow/Nodes/GeometryCollectionAssetNodes.cpp
@@ -13,10 +13,15 @@ namespace Dataflow
 {
 	void GeometryCollectionEngineAssetNodes()
 	{
+		static const FLinearColor CDefaultNodeBodyTintColor = FLinearColor(0.f, 0.f, 0.f, 0.5f);
+
 		DATAFLOW_NODE_REGISTER_CREATION_FACTORY(FGeometryCollectionTerminalDataflowNode);
 		DATAFLOW_NODE_REGISTER_CREATION_FACTORY(FGetGeometryCollectionAssetDataflowNode);
 		DATAFLOW_NODE_REGISTER_CREATION_FACTORY(FGetGeometryCollectionSourcesDataflowNode);
 		DATAFLOW_NODE_REGISTER_CREATION_FACTORY(FCreateGeometryCollectionFromSourcesDataflowNode);
+
+		// Terminal
+		DATAFLOW_NODE_REGISTER_CREATION_FACTORY_NODE_COLORS_BY_CATEGORY("Terminal", FLinearColor(0.f, 0.f, 0.f), CDefaultNodeBodyTintColor);
 	}
 }
 
diff --git a/Engine/Plugins/Experimental/GeometryCollectionPlugin/Source/GeometryCollectionNodes/Public/Dataflow/GeometryCollectionClusteringNodes.h b/Engine/Plugins/Experimental/GeometryCollectionPlugin/Source/GeometryCollectionNodes/Public/Dataflow/GeometryCollectionClusteringNodes.h
new file mode 100644
--- /dev/null
+++ b/Engine/Plugins/Experimental/GeometryCollectionPlugin/Source/GeometryCollectionNodes/Public/Dataflow/GeometryCollectionClusteringNodes.h
@@ -0,0 +1,120 @@
+// Copyright Epic Games, Inc. All Rights Reserved.
+
+#pragma once 
+
+#include "CoreMinimal.h"
+#include "Dataflow/DataflowEngine.h"
+#include "GeometryCollection/ManagedArrayCollection.h"
+#include "Dataflow/DataflowSelection.h"
+
+#include "GeometryCollectionClusteringNodes.generated.h"
+
+class FGeometryCollection;
+
+
+UENUM(BlueprintType)
+enum class EClusterSizeMethodEnum : uint8
+{
+	Dataflow_ClusterSizeMethod_ByNumber UMETA(DisplayName = "By Number"),
+	Dataflow_ClusterSizeMethod_ByFractionOfInput UMETA(DisplayName = "By Fraction Of Input"),
+	Dataflow_ClusterSizeMethod_BySize UMETA(DisplayName = "By Size"),
+	//~~~
+	//256th entry
+	Dataflow_Max                UMETA(Hidden)
+};
+
+/**
+ *
+ * Automatically group pieces of a fractured Collection into a specified number of clusters
+ *
+ */
+USTRUCT(meta = (DataflowGeometryCollection))
+struct FAutoClusterDataflowNode : public FDataflowNode
+{
+	GENERATED_USTRUCT_BODY()
+	DATAFLOW_NODE_DEFINE_INTERNAL(FAutoClusterDataflowNode, "AutoCluster", "GeometryCollection|Cluster", "")
+
+public:
+	/** How to choose the size of the clusters to create */
+	UPROPERTY(EditAnywhere, Category = "Cluster Size");
+	EClusterSizeMethodEnum ClusterSizeMethod = EClusterSizeMethodEnum::Dataflow_ClusterSizeMethod_ByNumber;
+
+	/** Use a Voronoi diagram with this many Voronoi sites as a guide for deciding cluster boundaries */
+	UPROPERTY(EditAnywhere, Category = "Cluster Size", meta = (DataflowInput, UIMin = 2, UIMax = 5000, EditCondition = "ClusterSizeMethod == EClusterSizeMethodEnum::Dataflow_ClusterSizeMethod_ByNumber"))
+	int32 ClusterSites = 10;
+
+	/** Choose the number of Voronoi sites used for clustering as a fraction of the number of child bones to process */
+	UPROPERTY(EditAnywhere, Category = "Cluster Size", meta = (DataflowInput, UIMin = 0.f, UIMax = 0.5f, EditCondition = "ClusterSizeMethod == EClusterSizeMethodEnum::Dataflow_ClusterSizeMethod_ByFractionOfInput"))
+	float ClusterFraction = 0.25;
+
+	/** Choose the Edge-Size of the cube used to groups bones under a cluster (in cm). */
+	UPROPERTY(EditAnywhere, Category = "ClusterSize", meta = (DataflowInput, DisplayName = "Cluster Size", UIMin = ".01", UIMax = "100", ClampMin = ".0001", ClampMax = "10000", EditCondition = "ClusterSizeMethod == EClusterSizeMethodEnum::Dataflow_ClusterSizeMethod_BySize"))
+	float SiteSize = 1;
+
+	/** If true, bones will only be added to the same cluster if they are physically connected (either directly, or via other bones in the same cluster) */
+	UPROPERTY(EditAnywhere, Category = "Auto Cluster")
+	bool AutoCluster = true;
+
+	/** If true, prevent the creation of clusters with only a single child. Either by merging into a neighboring cluster, or not creating the cluster. */
+	UPROPERTY(EditAnywhere, Category = "Auto Cluster")
+	bool AvoidIsolated = true;
+
+	/** Fractured GeometryCollection to cluster */
+	UPROPERTY(meta = (DataflowInput, DataflowOutput, DataflowPassthrough = "Collection", DataflowIntrinsic))
+	FManagedArrayCollection Collection;
+
+	/** Bone selection for the clustering */
+	UPROPERTY(meta = (DataflowInput, DisplayName = "TransformSelection"))
+	FDataflowTransformSelection TransformSelection;
+
+	FAutoClusterDataflowNode(const Dataflow::FNodeParameters& InParam, FGuid InGuid = FGuid::NewGuid())
+		: FDataflowNode(InParam, InGuid)
+	{
+		RegisterInputConnection(&Collection);
+		RegisterInputConnection(&TransformSelection);
+		RegisterInputConnection(&ClusterSites);
+		RegisterInputConnection(&ClusterFraction);
+		RegisterInputConnection(&SiteSize);
+		RegisterOutputConnection(&Collection, &Collection);
+	}
+
+	virtual void Evaluate(Dataflow::FContext& Context, const FDataflowOutput* Out) const override;
+
+};
+
+
+/**
+ *
+ * Flattens all bones to level 1
+ *
+ */
+USTRUCT(meta = (DataflowGeometryCollection))
+struct FClusterFlattenDataflowNode : public FDataflowNode
+{
+	GENERATED_USTRUCT_BODY()
+	DATAFLOW_NODE_DEFINE_INTERNAL(FClusterFlattenDataflowNode, "Flatten", "GeometryCollection|Cluster", "")
+
+public:
+	// @todo(harsha) Support Selections
+
+	/** Fractured GeometryCollection to flatten */
+	UPROPERTY(meta = (DataflowInput, DataflowOutput, DataflowPassthrough = "Collection", DataflowIntrinsic))
+	FManagedArrayCollection Collection;
+
+	FClusterFlattenDataflowNode(const Dataflow::FNodeParameters& InParam, FGuid InGuid = FGuid::NewGuid())
+		: FDataflowNode(InParam, InGuid)
+	{
+		RegisterInputConnection(&Collection);
+		RegisterOutputConnection(&Collection, &Collection);
+	}
+
+	virtual void Evaluate(Dataflow::FContext& Context, const FDataflowOutput* Out) const override;
+
+};
+
+
+namespace Dataflow
+{
+	void GeometryCollectionClusteringNodes();
+}
+
diff --git a/Engine/Plugins/Experimental/GeometryCollectionPlugin/Source/GeometryCollectionNodes/Public/Dataflow/GeometryCollectionFracturingNodes.h b/Engine/Plugins/Experimental/GeometryCollectionPlugin/Source/GeometryCollectionNodes/Public/Dataflow/GeometryCollectionFracturingNodes.h
new file mode 100644
--- /dev/null
+++ b/Engine/Plugins/Experimental/GeometryCollectionPlugin/Source/GeometryCollectionNodes/Public/Dataflow/GeometryCollectionFracturingNodes.h
@@ -0,0 +1,265 @@
+// Copyright Epic Games, Inc. All Rights Reserved.
+
+#pragma once 
+
+#include "CoreMinimal.h"
+#include "Dataflow/DataflowEngine.h"
+#include "GeometryCollection/ManagedArrayCollection.h"
+#include "Dataflow/DataflowSelection.h"
+
+#include "GeometryCollectionFracturingNodes.generated.h"
+
+class FGeometryCollection;
+
+
+USTRUCT(meta = (DataflowGeometryCollection))
+struct FUniformScatterPointsDataflowNode : public FDataflowNode
+{
+	GENERATED_USTRUCT_BODY()
+	DATAFLOW_NODE_DEFINE_INTERNAL(FUniformScatterPointsDataflowNode, "UniformScatterPoints", "Generators|Point", "")
+
+public:
+	UPROPERTY(EditAnywhere, Category = "Scatter", meta = (DataflowInput));
+	int32 MinNumberOfPoints = 20;
+
+	UPROPERTY(EditAnywhere, Category = "Scatter", meta = (DataflowInput));
+	int32 MaxNumberOfPoints = 20;
+
+	UPROPERTY(EditAnywhere, Category = "Scatter", meta = (DataflowInput));
+	float RandomSeed = -1.f;
+
+	UPROPERTY(meta = (DataflowInput))
+	FBox BoundingBox = FBox(ForceInit);
+
+	UPROPERTY(meta = (DataflowOutput))
+	TArray<FVector> Points;
+
+	FUniformScatterPointsDataflowNode(const Dataflow::FNodeParameters& InParam, FGuid InGuid = FGuid::NewGuid())
+		: FDataflowNode(InParam, InGuid)
+	{
+		RegisterInputConnection(&BoundingBox);
+		RegisterInputConnection(&MinNumberOfPoints);
+		RegisterInputConnection(&MaxNumberOfPoints);
+		RegisterInputConnection(&RandomSeed);
+		RegisterOutputConnection(&Points);
+	}
+
+	virtual void Evaluate(Dataflow::FContext& Context, const FDataflowOutput* Out) const override;
+
+};
+
+
+USTRUCT(meta = (DataflowGeometryCollection))
+struct FRadialScatterPointsDataflowNode : public FDataflowNode
+{
+	GENERATED_USTRUCT_BODY()
+	DATAFLOW_NODE_DEFINE_INTERNAL(FRadialScatterPointsDataflowNode, "RadialScatterPoints", "Generators|Point", "")
+
+public:
+	UPROPERTY(EditAnywhere, Category = "Scatter", meta = (DataflowInput));
+	FVector Center = FVector(0.0);
+
+	UPROPERTY(EditAnywhere, Category = "Scatter", meta = (DataflowInput));
+	FVector Normal = FVector(0.0, 0.0, 1.0);
+
+	UPROPERTY(EditAnywhere, Category = "Scatter", meta = (DataflowInput, UIMin = 0.01f));
+	float Radius = 50.f;
+
+	UPROPERTY(EditAnywhere, Category = "Scatter", meta = (DataflowInput, UIMin = 1, UIMax = 50));
+	int32 AngularSteps = 5;
+
+	UPROPERTY(EditAnywhere, Category = "Scatter", meta = (DataflowInput, UIMin = 1, UIMax = 50));
+	int32 RadialSteps = 5;
+
+	UPROPERTY(EditAnywhere, Category = "Scatter", meta = (DataflowInput));
+	float AngleOffset = 0.f;
+
+	UPROPERTY(EditAnywhere, Category = "Scatter", meta = (DataflowInput, UIMin = 0.f));
+	float Variability = 0.f;
+
+	UPROPERTY(EditAnywhere, Category = "Scatter", meta = (DataflowInput));
+	float RandomSeed = -1.f;
+
+	UPROPERTY(meta = (DataflowOutput))
+	TArray<FVector> Points;
+
+	FRadialScatterPointsDataflowNode(const Dataflow::FNodeParameters& InParam, FGuid InGuid = FGuid::NewGuid())
+		: FDataflowNode(InParam, InGuid)
+	{
+		RegisterInputConnection(&Center);
+		RegisterInputConnection(&Normal);
+		RegisterInputConnection(&Radius);
+		RegisterInputConnection(&AngularSteps);
+		RegisterInputConnection(&RadialSteps);
+		RegisterInputConnection(&AngleOffset);
+		RegisterInputConnection(&Variability);
+		RegisterInputConnection(&RandomSeed);
+		RegisterOutputConnection(&Points);
+	}
+
+	virtual void Evaluate(Dataflow::FContext& Context, const FDataflowOutput* Out) const override;
+
+};
+
+
+//
+// GridScatterPoints
+//
+
+
+/**
+ *
+ * Generates a Voronoi fracture
+ *
+ */
+USTRUCT(meta = (DataflowGeometryCollection))
+struct FVoronoiFractureDataflowNode : public FDataflowNode
+{
+	GENERATED_USTRUCT_BODY()
+	DATAFLOW_NODE_DEFINE_INTERNAL(FVoronoiFractureDataflowNode, "VoronoiFracture", "GeometryCollection|Fracture", "")
+
+public:
+	UPROPERTY(meta = (DataflowInput, DataflowOutput, DataflowPassthrough = "Collection", DataflowIntrinsic))
+	FManagedArrayCollection Collection;
+
+	UPROPERTY(meta = (DataflowInput, DataflowIntrinsic))
+	TArray<FVector> Points;
+
+	UPROPERTY(EditAnywhere, Category = "Fracture", meta = (DataflowInput));
+	float RandomSeed = -1.f;
+
+	UPROPERTY(EditAnywhere, Category = "Fracture", meta = (DataflowInput, UIMin = 0.f, UIMax = 1.f));
+	float ChanceToFracture = 1.f;
+
+	UPROPERTY(EditAnywhere, Category = "Fracture");
+	bool GroupFracture = true;
+
+	UPROPERTY(EditAnywhere, Category = "Fracture", meta = (DataflowInput, UIMin = 0.f));
+	float Grout = 0.f;
+
+	UPROPERTY(EditAnywhere, Category = "Noise", meta = (DataflowInput, UIMin = 0.f));
+	float Amplitude = 0.f;
+
+	UPROPERTY(EditAnywhere, Category = "Noise", meta = (DataflowInput, UIMin = 0.00001f));
+	float Frequency = 0.1f;
+
+	UPROPERTY(EditAnywhere, Category = "Noise", meta = (DataflowInput, UIMin = 0.f));
+	float Persistence = 0.5f;
+
+	UPROPERTY(EditAnywhere, Category = "Noise", meta = (DataflowInput, UIMin = 0.f));
+	float Lacunarity = 2.f;
+
+	UPROPERTY(EditAnywhere, Category = "Noise", meta = (DataflowInput, UIMin = 0.f));
+	int32 OctaveNumber = 4;
+
+	UPROPERTY(EditAnywhere, Category = "Noise", meta = (DataflowInput, UIMin = 0.f));
+	float PointSpacing = 10.f;
+
+	UPROPERTY(EditAnywhere, Category = "Collision");
+	bool AddSamplesForCollision = false;
+
+	UPROPERTY(EditAnywhere, Category = "Collision", meta = (DataflowInput, UIMin = 0.f));
+	float CollisionSampleSpacing = 50.f;
+
+	FVoronoiFractureDataflowNode(const Dataflow::FNodeParameters& InParam, FGuid InGuid = FGuid::NewGuid())
+		: FDataflowNode(InParam, InGuid)
+	{
+		RegisterInputConnection(&Collection);
+		RegisterInputConnection(&Points);
+		RegisterInputConnection(&RandomSeed);
+		RegisterInputConnection(&ChanceToFracture);
+		RegisterInputConnection(&Grout);
+		RegisterInputConnection(&Amplitude);
+		RegisterInputConnection(&Frequency);
+		RegisterInputConnection(&Persistence);
+		RegisterInputConnection(&Lacunarity);
+		RegisterInputConnection(&OctaveNumber);
+		RegisterInputConnection(&PointSpacing);
+		RegisterInputConnection(&CollisionSampleSpacing);
+		RegisterOutputConnection(&Collection, &Collection);
+	}
+
+	virtual void Evaluate(Dataflow::FContext& Context, const FDataflowOutput* Out) const override;
+
+};
+
+
+/**
+ *
+ * Cuts geometry using a set of noised up planes
+ *
+ */
+USTRUCT(meta = (DataflowGeometryCollection))
+struct FPlaneCutterDataflowNode : public FDataflowNode
+{
+	GENERATED_USTRUCT_BODY()
+		DATAFLOW_NODE_DEFINE_INTERNAL(FPlaneCutterDataflowNode, "PlaneCutter", "GeometryCollection|Fracture", "")
+
+public:
+	UPROPERTY(meta = (DataflowInput, DataflowOutput, DataflowPassthrough = "Collection", DataflowIntrinsic))
+	FManagedArrayCollection Collection;
+
+	UPROPERTY(meta = (DataflowInput, DataflowIntrinsic))
+	FBox BoundingBox = FBox(ForceInit);
+
+	UPROPERTY(EditAnywhere, Category = "Fracture", meta = (DataflowInput, UIMin = 1))
+	int32 NumPlanes = 1;
+
+	UPROPERTY(EditAnywhere, Category = "Fracture", meta = (DataflowInput))
+	float RandomSeed = -1.f;
+
+	UPROPERTY(EditAnywhere, Category = "Fracture", meta = (DataflowInput, UIMin = 0.f))
+	float Grout = 0.f;
+
+	UPROPERTY(EditAnywhere, Category = "Noise", meta = (DataflowInput, UIMin = 0.f));
+	float Amplitude = 0.f;
+
+	UPROPERTY(EditAnywhere, Category = "Noise", meta = (DataflowInput, UIMin = 0.00001f));
+	float Frequency = 0.1f;
+
+	UPROPERTY(EditAnywhere, Category = "Noise", meta = (DataflowInput, UIMin = 0.f));
+	float Persistence = 0.5f;
+
+	UPROPERTY(EditAnywhere, Category = "Noise", meta = (DataflowInput, UIMin = 0.f));
+	float Lacunarity = 2.f;
+
+	UPROPERTY(EditAnywhere, Category = "Noise", meta = (DataflowInput, UIMin = 0.f));
+	int32 OctaveNumber = 4;
+
+	UPROPERTY(EditAnywhere, Category = "Noise", meta = (DataflowInput, UIMin = 0.f));
+	float PointSpacing = 10.f;
+
+	UPROPERTY(EditAnywhere, Category = "Collision");
+	bool AddSamplesForCollision = false;
+
+	UPROPERTY(EditAnywhere, Category = "Collision", meta = (DataflowInput, UIMin = 0.f));
+	float CollisionSampleSpacing = 50.f;
+
+	FPlaneCutterDataflowNode(const Dataflow::FNodeParameters& InParam, FGuid InGuid = FGuid::NewGuid())
+		: FDataflowNode(InParam, InGuid)
+	{
+		RegisterInputConnection(&Collection);
+		RegisterInputConnection(&BoundingBox);
+		RegisterInputConnection(&NumPlanes);
+		RegisterInputConnection(&RandomSeed);
+		RegisterInputConnection(&Grout);
+		RegisterInputConnection(&Amplitude);
+		RegisterInputConnection(&Frequency);
+		RegisterInputConnection(&Persistence);
+		RegisterInputConnection(&Lacunarity);
+		RegisterInputConnection(&OctaveNumber);
+		RegisterInputConnection(&PointSpacing);
+		RegisterInputConnection(&CollisionSampleSpacing);
+		RegisterOutputConnection(&Collection, &Collection);
+	}
+
+	virtual void Evaluate(Dataflow::FContext& Context, const FDataflowOutput* Out) const override;
+
+};
+
+
+namespace Dataflow
+{
+	void GeometryCollectionFracturingNodes();
+}
+
diff --git a/Engine/Plugins/Experimental/GeometryCollectionPlugin/Source/GeometryCollectionNodes/Public/Dataflow/GeometryCollectionMeshNodes.h b/Engine/Plugins/Experimental/GeometryCollectionPlugin/Source/GeometryCollectionNodes/Public/Dataflow/GeometryCollectionMeshNodes.h
new file mode 100644
--- /dev/null
+++ b/Engine/Plugins/Experimental/GeometryCollectionPlugin/Source/GeometryCollectionNodes/Public/Dataflow/GeometryCollectionMeshNodes.h
@@ -0,0 +1,371 @@
+// Copyright Epic Games, Inc. All Rights Reserved.
+
+#pragma once 
+
+#include "CoreMinimal.h"
+#include "Dataflow/DataflowEngine.h"
+#include "GeometryCollection/ManagedArrayCollection.h"
+#include "Dataflow/DataflowSelection.h"
+
+#include "GeometryCollectionMeshNodes.generated.h"
+
+class FGeometryCollection;
+class UStaticMesh;
+class UDynamicMesh;
+
+
+/**
+ *
+ * Converts points into a DynamicMesh
+ *
+ */
+USTRUCT(meta = (DataflowGeometryCollection))
+struct FPointsToMeshDataflowNode : public FDataflowNode
+{
+	GENERATED_USTRUCT_BODY()
+	DATAFLOW_NODE_DEFINE_INTERNAL(FPointsToMeshDataflowNode, "PointsToMesh", "Mesh|Utilities", "")
+
+public:
+	/** Points input */
+	UPROPERTY(meta = (DataflowInput, DataflowIntrinsic))
+	TArray<FVector> Points;
+
+	/** Mesh output */
+	UPROPERTY(meta = (DataflowOutput))
+	TObjectPtr<UDynamicMesh> Mesh;
+
+	/** Mesh triangle count */
+	UPROPERTY(meta = (DataflowOutput))
+	int32 TriangleCount = 0;
+
+	FPointsToMeshDataflowNode(const Dataflow::FNodeParameters& InParam, FGuid InGuid = FGuid::NewGuid())
+		: FDataflowNode(InParam, InGuid)
+	{
+		RegisterInputConnection(&Points);
+		RegisterOutputConnection(&Mesh);
+		RegisterOutputConnection(&TriangleCount);
+	}
+
+	virtual void Evaluate(Dataflow::FContext& Context, const FDataflowOutput* Out) const override;
+
+};
+
+
+/**
+ *
+ * Converts a BoundingBox into a DynamicMesh
+ *
+ */
+USTRUCT(meta = (DataflowGeometryCollection))
+struct FBoxToMeshDataflowNode : public FDataflowNode
+{
+	GENERATED_USTRUCT_BODY()
+	DATAFLOW_NODE_DEFINE_INTERNAL(FBoxToMeshDataflowNode, "BoxToMesh", "Mesh|Utilities", "")
+
+public:
+	/** BoundingBox input */
+	UPROPERTY(meta = (DataflowInput, DataflowIntrinsic))
+	FBox Box = FBox(ForceInit);
+
+	/** Mesh output */
+	UPROPERTY(meta = (DataflowOutput))
+	TObjectPtr<UDynamicMesh> Mesh;
+
+	/** Mesh triangle count */
+	UPROPERTY(meta = (DataflowOutput))
+	int32 TriangleCount = 0;
+
+	FBoxToMeshDataflowNode(const Dataflow::FNodeParameters& InParam, FGuid InGuid = FGuid::NewGuid())
+		: FDataflowNode(InParam, InGuid)
+	{
+		RegisterInputConnection(&Box);
+		RegisterOutputConnection(&Mesh);
+		RegisterOutputConnection(&TriangleCount);
+	}
+
+	virtual void Evaluate(Dataflow::FContext& Context, const FDataflowOutput* Out) const override;
+
+};
+
+
+/**
+ *
+ * Collects information from the DynamicMesh and outputs it into a formatted string
+ *
+ */
+USTRUCT(meta = (DataflowGeometryCollection))
+struct FMeshInfoDataflowNode : public FDataflowNode
+{
+	GENERATED_USTRUCT_BODY()
+	DATAFLOW_NODE_DEFINE_INTERNAL(FMeshInfoDataflowNode, "MeshInfo", "Mesh|Utilities", "")
+
+public:
+	/** DynamicMesh for the information */
+	UPROPERTY(meta = (DataflowInput, DataflowIntrinsic))
+	TObjectPtr<UDynamicMesh> Mesh;
+
+	/** Formatted output string */
+	UPROPERTY(meta = (DataflowOutput))
+	FString InfoString = FString("");
+
+	FMeshInfoDataflowNode(const Dataflow::FNodeParameters& InParam, FGuid InGuid = FGuid::NewGuid())
+		: FDataflowNode(InParam, InGuid)
+	{
+		RegisterInputConnection(&Mesh);
+		RegisterOutputConnection(&InfoString);
+	}
+
+	virtual void Evaluate(Dataflow::FContext& Context, const FDataflowOutput* Out) const override;
+
+};
+
+
+/**
+ *
+ * Converts a DynamicMesh to a Collection
+ *
+ */
+USTRUCT(meta = (DataflowGeometryCollection))
+struct FMeshToCollectionDataflowNode : public FDataflowNode
+{
+	GENERATED_USTRUCT_BODY()
+	DATAFLOW_NODE_DEFINE_INTERNAL(FMeshToCollectionDataflowNode, "MeshToCollection", "Mesh|Utilities", "")
+	DATAFLOW_NODE_RENDER_TYPE(FGeometryCollection::StaticType(), "Collection")
+
+public:
+	/** DynamicMesh to convert */
+	UPROPERTY(meta = (DataflowInput, DataflowIntrinsic))
+	TObjectPtr<UDynamicMesh> Mesh;
+
+	/** Output Collection */
+	UPROPERTY(meta = (DataflowOutput))
+	FManagedArrayCollection Collection;
+
+	FMeshToCollectionDataflowNode(const Dataflow::FNodeParameters& InParam, FGuid InGuid = FGuid::NewGuid())
+		: FDataflowNode(InParam, InGuid)
+	{
+		RegisterInputConnection(&Mesh);
+		RegisterOutputConnection(&Collection);
+	}
+
+	virtual void Evaluate(Dataflow::FContext& Context, const FDataflowOutput* Out) const override;
+
+};
+
+
+/**
+ *
+ * Converts a StaticMesh into a DynamicMesh
+ *
+ */
+USTRUCT(meta = (DataflowGeometryCollection))
+struct FStaticMeshToMeshDataflowNode : public FDataflowNode
+{
+	GENERATED_USTRUCT_BODY()
+	DATAFLOW_NODE_DEFINE_INTERNAL(FStaticMeshToMeshDataflowNode, "StaticMeshToMesh", "Mesh|Utilities", "")
+
+public:
+	/** StaticMesh to convert */
+	UPROPERTY(EditAnywhere, Category = "StaticMesh");
+	TObjectPtr<UStaticMesh> StaticMesh;
+
+	/** Output the HiRes representation, if set to true and HiRes doesn't exist it will output empty mesh */
+	UPROPERTY(EditAnywhere, Category = "StaticMesh", meta = (DisplayName = "Use HiRes"));
+	bool UseHiRes = true;
+
+	/** Specifies the LOD level to use */
+	UPROPERTY(EditAnywhere, Category = "StaticMesh", meta = (DisplayName = "LOD Level"));
+	int32 LODLevel = 0;
+
+	/** Output mesh */
+	UPROPERTY(meta = (DataflowOutput))
+	TObjectPtr<UDynamicMesh> Mesh;
+
+	FStaticMeshToMeshDataflowNode(const Dataflow::FNodeParameters& InParam, FGuid InGuid = FGuid::NewGuid())
+		: FDataflowNode(InParam, InGuid)
+	{
+		RegisterOutputConnection(&Mesh);
+	}
+
+	virtual void Evaluate(Dataflow::FContext& Context, const FDataflowOutput* Out) const override;
+
+};
+
+
+/**
+ *
+ * Appends two meshes
+ *
+ */
+USTRUCT(meta = (DataflowGeometryCollection))
+struct FMeshAppendDataflowNode : public FDataflowNode
+{
+	GENERATED_USTRUCT_BODY()
+	DATAFLOW_NODE_DEFINE_INTERNAL(FMeshAppendDataflowNode, "MeshAppend", "Mesh|Utilities", "")
+
+public:
+	/** Mesh input */
+	UPROPERTY(meta = (DataflowInput, DataflowIntrinsic))
+	TObjectPtr<UDynamicMesh> Mesh1;
+
+	/** Mesh input */
+	UPROPERTY(meta = (DataflowInput, DataflowIntrinsic))
+	TObjectPtr<UDynamicMesh> Mesh2;
+
+	/** Output mesh */
+	UPROPERTY(meta = (DataflowOutput))
+	TObjectPtr<UDynamicMesh> Mesh;
+
+	FMeshAppendDataflowNode(const Dataflow::FNodeParameters& InParam, FGuid InGuid = FGuid::NewGuid())
+		: FDataflowNode(InParam, InGuid)
+	{
+		RegisterInputConnection(&Mesh1);
+		RegisterInputConnection(&Mesh2);
+		RegisterOutputConnection(&Mesh);
+	}
+
+	virtual void Evaluate(Dataflow::FContext& Context, const FDataflowOutput* Out) const override;
+
+};
+
+
+UENUM(BlueprintType)
+enum class EMeshBooleanOperationEnum : uint8
+{
+	Dataflow_MeshBoolean_Union UMETA(DisplayName = "Union"),
+	Dataflow_MeshBoolean_Intersect UMETA(DisplayName = "Intersect"),
+	Dataflow_MeshBoolean_Difference UMETA(DisplayName = "Difference"),
+	//~~~
+	//256th entry
+	Dataflow_Max                UMETA(Hidden)
+};
+
+/**
+ *
+ * Mesh boolean (Union, Intersect, Difference) between two meshes
+ *
+ */
+USTRUCT(meta = (DataflowGeometryCollection))
+struct FMeshBooleanDataflowNode : public FDataflowNode
+{
+	GENERATED_USTRUCT_BODY()
+	DATAFLOW_NODE_DEFINE_INTERNAL(FMeshBooleanDataflowNode, "MeshBoolean", "Mesh|Utilities", "")
+
+public:
+	/** Boolean operation */
+	UPROPERTY(EditAnywhere, Category = "Boolean");
+	EMeshBooleanOperationEnum Operation = EMeshBooleanOperationEnum::Dataflow_MeshBoolean_Intersect;
+
+	/** Mesh input */
+	UPROPERTY(meta = (DataflowInput, DataflowIntrinsic))
+	TObjectPtr<UDynamicMesh> Mesh1;
+
+	/** Mesh input */
+	UPROPERTY(meta = (DataflowInput, DataflowIntrinsic))
+	TObjectPtr<UDynamicMesh> Mesh2;
+
+	/** Output mesh */
+	UPROPERTY(meta = (DataflowOutput))
+	TObjectPtr<UDynamicMesh> Mesh;
+
+	FMeshBooleanDataflowNode(const Dataflow::FNodeParameters& InParam, FGuid InGuid = FGuid::NewGuid())
+		: FDataflowNode(InParam, InGuid)
+	{
+		RegisterInputConnection(&Mesh1);
+		RegisterInputConnection(&Mesh2);
+		RegisterOutputConnection(&Mesh);
+	}
+
+	virtual void Evaluate(Dataflow::FContext& Context, const FDataflowOutput* Out) const override;
+
+};
+
+
+/**
+ *
+ * Copies the same mesh with scale onto points
+ *
+ */
+USTRUCT(meta = (DataflowGeometryCollection))
+struct FMeshCopyToPointsDataflowNode : public FDataflowNode
+{
+	GENERATED_USTRUCT_BODY()
+		DATAFLOW_NODE_DEFINE_INTERNAL(FMeshCopyToPointsDataflowNode, "MeshCopyToPoints", "Mesh|Utilities", "")
+
+public:
+	/** Points to copy meshes onto */
+	UPROPERTY(meta = (DataflowInput, DataflowIntrinsic))
+	TArray<FVector> Points;
+
+	/** Mesh to copy onto points */
+	UPROPERTY(meta = (DataflowInput, DataflowIntrinsic))
+	TObjectPtr<UDynamicMesh> MeshToCopy;
+
+	/** Scale appied to the mesh */
+	UPROPERTY(EditAnywhere, Category = "Copy");
+	float Scale = 1.f;
+
+	/** Copied meshes */
+	UPROPERTY(meta = (DataflowOutput))
+	TObjectPtr<UDynamicMesh> Mesh;
+
+	FMeshCopyToPointsDataflowNode(const Dataflow::FNodeParameters& InParam, FGuid InGuid = FGuid::NewGuid())
+		: FDataflowNode(InParam, InGuid)
+	{
+		RegisterInputConnection(&Points);
+		RegisterInputConnection(&MeshToCopy);
+		RegisterOutputConnection(&Mesh);
+	}
+
+	virtual void Evaluate(Dataflow::FContext& Context, const FDataflowOutput* Out) const override;
+
+};
+
+
+/**
+ *
+ * Outputs Mesh data
+ *
+ */
+USTRUCT(meta = (DataflowGeometryCollection))
+struct FGetMeshDataDataflowNode : public FDataflowNode
+{
+	GENERATED_USTRUCT_BODY()
+	DATAFLOW_NODE_DEFINE_INTERNAL(FGetMeshDataDataflowNode, "GetMeshData", "Mesh|Utilities", "")
+
+public:
+	/** Mesh for the data */
+	UPROPERTY(meta = (DataflowInput, DataflowIntrinsic))
+	TObjectPtr<UDynamicMesh> Mesh;
+
+	/** Number of vertices */
+	UPROPERTY(meta = (DataflowOutput))
+	int32 VertexCount = 0;
+
+	/** Number of edges */
+	UPROPERTY(meta = (DataflowOutput))
+	int32 EdgeCount = 0;
+
+	/** Number of triangles */
+	UPROPERTY(meta = (DataflowOutput))
+	int32 TriangleCount = 0;
+
+	FGetMeshDataDataflowNode(const Dataflow::FNodeParameters& InParam, FGuid InGuid = FGuid::NewGuid())
+		: FDataflowNode(InParam, InGuid)
+	{
+		RegisterInputConnection(&Mesh);
+		RegisterOutputConnection(&VertexCount);
+		RegisterOutputConnection(&EdgeCount);
+		RegisterOutputConnection(&TriangleCount);
+	}
+
+	virtual void Evaluate(Dataflow::FContext& Context, const FDataflowOutput* Out) const override;
+
+};
+
+
+namespace Dataflow
+{
+	void GeometryCollectionMeshNodes();
+}
+
diff --git a/Engine/Plugins/Experimental/GeometryCollectionPlugin/Source/GeometryCollectionNodes/Public/Dataflow/GeometryCollectionNodes.h b/Engine/Plugins/Experimental/GeometryCollectionPlugin/Source/GeometryCollectionNodes/Public/Dataflow/GeometryCollectionNodes.h
--- a/Engine/Plugins/Experimental/GeometryCollectionPlugin/Source/GeometryCollectionNodes/Public/Dataflow/GeometryCollectionNodes.h
+++ b/Engine/Plugins/Experimental/GeometryCollectionPlugin/Source/GeometryCollectionNodes/Public/Dataflow/GeometryCollectionNodes.h
@@ -15,7 +15,12 @@
 class FGeometryCollection;
 class UStaticMesh;
 
-USTRUCT()
+/**
+ *
+ * Description for this node
+ *
+ */
+USTRUCT(meta = (DataflowGeometryCollection))
 struct FGetCollectionAssetDataflowNode : public FDataflowNode
 {
 	GENERATED_USTRUCT_BODY()
@@ -37,10 +42,13 @@ public:
 
 };
 
+
 /**
-	Description for this node
-*/
-USTRUCT()
+ *
+ * Description for this node
+ *
+ */
+USTRUCT(meta = (DataflowGeometryCollection))
 struct FExampleCollectionEditDataflowNode : public FDataflowNode
 {
 	GENERATED_USTRUCT_BODY()
@@ -67,11 +75,17 @@ public:
 
 };
 
-USTRUCT()
+
+/**
+ *
+ * Description for this node
+ *
+ */
+USTRUCT(meta = (DataflowGeometryCollection, DataflowTerminal))
 struct FSetCollectionAssetDataflowNode : public FDataflowTerminalNode
 {
 	GENERATED_USTRUCT_BODY()
-	DATAFLOW_NODE_DEFINE_INTERNAL(FSetCollectionAssetDataflowNode, "SetCollectionAsset", "GeometryCollection", "")
+	DATAFLOW_NODE_DEFINE_INTERNAL(FSetCollectionAssetDataflowNode, "SetCollectionAsset", "Terminal", "")
 
 public:
 	typedef FManagedArrayCollection DataType;
@@ -90,7 +104,12 @@ public:
 };
 
 
-USTRUCT()
+/**
+ *
+ * Description for this node
+ *
+ */
+USTRUCT(meta = (DataflowGeometryCollection))
 struct FAppendCollectionAssetsDataflowNode : public FDataflowNode
 {
 	GENERATED_USTRUCT_BODY()
@@ -118,7 +137,12 @@ public:
 
 };
 
-USTRUCT()
+/**
+ *
+ * Description for this node
+ *
+ */
+USTRUCT(meta = (DataflowGeometryCollection))
 struct FResetGeometryCollectionDataflowNode : public FDataflowNode
 {
 	GENERATED_USTRUCT_BODY()
@@ -141,6 +165,11 @@ public:
 };
 
 
+/**
+ *
+ * Description for this node
+ *
+ */
 USTRUCT()
 struct FPrintStringDataflowNode : public FDataflowNode
 {
@@ -174,6 +203,11 @@ public:
 };
 
 
+/**
+ *
+ * Description for this node
+ *
+ */
 USTRUCT()
 struct FLogStringDataflowNode : public FDataflowNode
 {
@@ -198,6 +232,11 @@ public:
 };
 
 
+/**
+ *
+ * Description for this node
+ *
+ */
 USTRUCT()
 struct FMakeLiteralStringDataflowNode : public FDataflowNode
 {
@@ -221,15 +260,19 @@ public:
 
 };
 
-USTRUCT()
+
+/**
+ *
+ * Description for this node
+ *
+ */
+USTRUCT(meta = (DataflowGeometryCollection))
 struct FBoundingBoxDataflowNode : public FDataflowNode
 {
 	GENERATED_USTRUCT_BODY()
 	DATAFLOW_NODE_DEFINE_INTERNAL(FBoundingBoxDataflowNode, "BoundingBox", "Utilities|Box", "")
 
 public:
-	typedef FManagedArrayCollection DataType;
-
 	UPROPERTY(meta = (DataflowInput))
 	FManagedArrayCollection Collection;
 
@@ -247,6 +290,12 @@ public:
 
 };
 
+
+/**
+ *
+ * Description for this node
+ *
+ */
 USTRUCT()
 struct FExpandBoundingBoxDataflowNode : public FDataflowNode
 {
@@ -287,6 +336,12 @@ public:
 
 };
 
+
+/**
+ *
+ * Description for this node
+ *
+ */
 USTRUCT()
 struct FVectorToStringDataflowNode : public FDataflowNode
 {
@@ -311,6 +366,12 @@ public:
 
 };
 
+
+/**
+ *
+ * Description for this node
+ *
+ */
 USTRUCT()
 struct FFloatToStringDataflowNode : public FDataflowNode
 {
@@ -335,6 +396,12 @@ public:
 
 };
 
+
+/**
+ *
+ * Description for this node
+ *
+ */
 USTRUCT()
 struct FMakePointsDataflowNode : public FDataflowNode
 {
@@ -358,6 +425,7 @@ public:
 
 };
 
+
 UENUM(BlueprintType)
 enum class EMakeBoxDataTypeEnum : uint8
 {
@@ -368,6 +436,11 @@ enum class EMakeBoxDataTypeEnum : uint8
 	Dataflow_Max                UMETA(Hidden)
 };
 
+/**
+ *
+ * Description for this node
+ *
+ */
 USTRUCT()
 struct FMakeBoxDataflowNode : public FDataflowNode
 {
@@ -407,101 +480,12 @@ public:
 
 };
 
-USTRUCT()
-struct FUniformScatterPointsDataflowNode : public FDataflowNode
-{
-	GENERATED_USTRUCT_BODY()
-		DATAFLOW_NODE_DEFINE_INTERNAL(FUniformScatterPointsDataflowNode, "UniformScatterPoints", "Generators|Point", "")
-
-public:
-	UPROPERTY(EditAnywhere, Category = "Scatter", meta = (DataflowInput));
-	int32 MinNumberOfPoints = 20;
-
-	UPROPERTY(EditAnywhere, Category = "Scatter", meta = (DataflowInput));
-	int32 MaxNumberOfPoints = 20;
-
-	UPROPERTY(EditAnywhere, Category = "Scatter", meta = (DataflowInput));
-	float RandomSeed = -1.f;
-
-	UPROPERTY(meta = (DataflowInput))
-	FBox BoundingBox = FBox(ForceInit);
-
-	UPROPERTY(meta = (DataflowOutput))
-	TArray<FVector> Points;
-
-	FUniformScatterPointsDataflowNode(const Dataflow::FNodeParameters& InParam, FGuid InGuid = FGuid::NewGuid())
-		: FDataflowNode(InParam, InGuid)
-	{
-		RegisterInputConnection(&BoundingBox);
-		RegisterInputConnection(&MinNumberOfPoints);
-		RegisterInputConnection(&MaxNumberOfPoints);
-		RegisterInputConnection(&RandomSeed);
-		RegisterOutputConnection(&Points);
-	}
-
-	virtual void Evaluate(Dataflow::FContext& Context, const FDataflowOutput* Out) const override;
-
-};
-
-USTRUCT()
-struct FRadialScatterPointsDataflowNode : public FDataflowNode
-{
-	GENERATED_USTRUCT_BODY()
-	DATAFLOW_NODE_DEFINE_INTERNAL(FRadialScatterPointsDataflowNode, "RadialScatterPoints", "Generators|Point", "")
-
-public:
-	UPROPERTY(EditAnywhere, Category = "Scatter", meta = (DataflowInput));
-	FVector Center = FVector(0.0);
-
-	UPROPERTY(EditAnywhere, Category = "Scatter", meta = (DataflowInput));
-	FVector Normal = FVector(0.0, 0.0, 1.0);
-
-	UPROPERTY(EditAnywhere, Category = "Scatter", meta = (DataflowInput, UIMin = 0.01f));
-	float Radius = 50.f;
-
-	UPROPERTY(EditAnywhere, Category = "Scatter", meta = (DataflowInput, UIMin = 1, UIMax = 50));
-	int32 AngularSteps = 5;
-
-	UPROPERTY(EditAnywhere, Category = "Scatter", meta = (DataflowInput, UIMin = 1, UIMax = 50));
-	int32 RadialSteps = 5;
-
-	UPROPERTY(EditAnywhere, Category = "Scatter", meta = (DataflowInput));
-	float AngleOffset = 0.f;
-
-	UPROPERTY(EditAnywhere, Category = "Scatter", meta = (DataflowInput, UIMin = 0.f));
-	float Variability = 0.f;
-
-	UPROPERTY(EditAnywhere, Category = "Scatter", meta = (DataflowInput));
-	float RandomSeed = -1.f;
-
-	UPROPERTY(meta = (DataflowOutput))
-	TArray<FVector> Points;
-
-	FRadialScatterPointsDataflowNode(const Dataflow::FNodeParameters& InParam, FGuid InGuid = FGuid::NewGuid())
-		: FDataflowNode(InParam, InGuid)
-	{
-		RegisterInputConnection(&Center);
-		RegisterInputConnection(&Normal);
-		RegisterInputConnection(&Radius);
-		RegisterInputConnection(&AngularSteps);
-		RegisterInputConnection(&RadialSteps);
-		RegisterInputConnection(&AngleOffset);
-		RegisterInputConnection(&Variability);
-		RegisterInputConnection(&RandomSeed);
-		RegisterOutputConnection(&Points);
-	}
-
-	virtual void Evaluate(Dataflow::FContext& Context, const FDataflowOutput* Out) const override;
-
-};
-
-
-//
-// GridScatterPoints
-//
-
-
 
+/**
+ *
+ * Description for this node
+ *
+ */
 USTRUCT()
 struct FMakeLiteralFloatDataflowNode : public FDataflowNode
 {
@@ -525,6 +509,12 @@ public:
 
 };
 
+
+/**
+ *
+ * Description for this node
+ *
+ */
 USTRUCT()
 struct FMakeLiteralIntDataflowNode : public FDataflowNode
 {
@@ -548,6 +538,11 @@ public:
 
 };
 
+/**
+ *
+ * Description for this node
+ *
+ */
 USTRUCT()
 struct FMakeLiteralBoolDataflowNode : public FDataflowNode
 {
@@ -581,7 +576,7 @@ USTRUCT()
 struct FMakeLiteralVectorDataflowNode : public FDataflowNode
 {
 	GENERATED_USTRUCT_BODY()
-		DATAFLOW_NODE_DEFINE_INTERNAL(FMakeLiteralVectorDataflowNode, "MakeLiteralVector", "Math|Vector", "")
+	DATAFLOW_NODE_DEFINE_INTERNAL(FMakeLiteralVectorDataflowNode, "MakeLiteralVector", "Math|Vector", "")
 
 public:
 	UPROPERTY(EditAnywhere, Category = "Vector");
@@ -610,7 +605,7 @@ USTRUCT()
 struct FIntToStringDataflowNode : public FDataflowNode
 {
 	GENERATED_USTRUCT_BODY()
-		DATAFLOW_NODE_DEFINE_INTERNAL(FIntToStringDataflowNode, "IntToString", "Utilities|String", "")
+	DATAFLOW_NODE_DEFINE_INTERNAL(FIntToStringDataflowNode, "IntToString", "Utilities|String", "")
 
 public:
 	UPROPERTY(EditAnywhere, Category = "Int", meta = (DataflowInput))
@@ -730,83 +725,6 @@ public:
 };
 
 
-/**
- *
- * Generates a Voronoi fracture
- *
- */
-USTRUCT()
-struct FVoronoiFractureDataflowNode : public FDataflowNode
-{
-	GENERATED_USTRUCT_BODY()
-	DATAFLOW_NODE_DEFINE_INTERNAL(FVoronoiFractureDataflowNode, "VoronoiFracture", "Fracture", "")
-
-public:
-	UPROPERTY(meta = (DataflowInput, DataflowOutput, DisplayName = "Collection"))
-	FManagedArrayCollection Collection;
-
-	UPROPERTY(meta = (DataflowInput))
-	TArray<FVector> Points;
-
-	UPROPERTY(EditAnywhere, Category = "Fracture", meta = (DataflowInput));
-	float RandomSeed = -1.f;
-
-	UPROPERTY(EditAnywhere, Category = "Fracture", meta = (DataflowInput, UIMin = 0.f, UIMax = 1.f));
-	float ChanceToFracture = 1.f;
-
-	UPROPERTY(EditAnywhere, Category = "Fracture");
-	bool GroupFracture = true;
-
-	UPROPERTY(EditAnywhere, Category = "Fracture", meta = (DataflowInput, UIMin = 0.f));
-	float Grout = 0.f;
-
-	UPROPERTY(EditAnywhere, Category = "Noise", meta = (DataflowInput, UIMin = 0.f));
-	float Amplitude = 0.f;
-
-	UPROPERTY(EditAnywhere, Category = "Noise", meta = (DataflowInput, UIMin = 0.00001f));
-	float Frequency = 0.1f;
-
-	UPROPERTY(EditAnywhere, Category = "Noise", meta = (DataflowInput, UIMin = 0.f));
-	float Persistence = 0.5f;
-
-	UPROPERTY(EditAnywhere, Category = "Noise", meta = (DataflowInput, UIMin = 0.f));
-	float Lacunarity = 2.f;
-
-	UPROPERTY(EditAnywhere, Category = "Noise", meta = (DataflowInput, UIMin = 0.f));
-	int32 OctaveNumber = 4;
-
-	UPROPERTY(EditAnywhere, Category = "Noise", meta = (DataflowInput, UIMin = 0.f));
-	float PointSpacing = 10.f;
-
-	UPROPERTY(EditAnywhere, Category = "Collision");
-	bool AddSamplesForCollision = false;
-
-	UPROPERTY(EditAnywhere, Category = "Collision", meta = (DataflowInput, UIMin = 0.f));
-	float CollisionSampleSpacing = 50.f;
-
-	FVoronoiFractureDataflowNode(const Dataflow::FNodeParameters& InParam, FGuid InGuid = FGuid::NewGuid())
-		: FDataflowNode(InParam, InGuid)
-	{
-		RegisterInputConnection(&Collection);
-		RegisterInputConnection(&Points);
-		RegisterInputConnection(&RandomSeed);
-		RegisterInputConnection(&ChanceToFracture);
-		RegisterInputConnection(&Grout);
-		RegisterInputConnection(&Amplitude);
-		RegisterInputConnection(&Frequency);
-		RegisterInputConnection(&Persistence);
-		RegisterInputConnection(&Lacunarity);
-		RegisterInputConnection(&OctaveNumber);
-		RegisterInputConnection(&PointSpacing);
-		RegisterInputConnection(&CollisionSampleSpacing);
-		RegisterOutputConnection(&Collection);
-	}
-
-	virtual void Evaluate(Dataflow::FContext& Context, const FDataflowOutput* Out) const override;
-
-};
-
-
 /**
  *
  * Concatenates two strings together to make a new string
@@ -850,7 +768,7 @@ USTRUCT()
 struct FRandomFloatDataflowNode : public FDataflowNode
 {
 	GENERATED_USTRUCT_BODY()
-		DATAFLOW_NODE_DEFINE_INTERNAL(FRandomFloatDataflowNode, "RandomFloat", "Math|Random", "")
+	DATAFLOW_NODE_DEFINE_INTERNAL(FRandomFloatDataflowNode, "RandomFloat", "Math|Random", "")
 
 public:
 	UPROPERTY(EditAnywhere, Category = "Seed")
@@ -883,7 +801,7 @@ USTRUCT()
 struct FRandomFloatInRangeDataflowNode : public FDataflowNode
 {
 	GENERATED_USTRUCT_BODY()
-		DATAFLOW_NODE_DEFINE_INTERNAL(FRandomFloatInRangeDataflowNode, "RandomFloatInRange", "Math|Random", "")
+	DATAFLOW_NODE_DEFINE_INTERNAL(FRandomFloatInRangeDataflowNode, "RandomFloatInRange", "Math|Random", "")
 
 public:
 	UPROPERTY(EditAnywhere, Category = "Seed")
@@ -1059,7 +977,7 @@ public:
  * "Explodes" the pieces from the Collection for better visualization
  *
  */
-USTRUCT()
+USTRUCT(meta = (DataflowGeometryCollection))
 struct FExplodedViewDataflowNode : public FDataflowNode
 {
 	GENERATED_USTRUCT_BODY()
@@ -1098,11 +1016,11 @@ private:
  * Generates convex hull representation for the bones for simulation
  *
  */
-USTRUCT()
+USTRUCT(meta = (DataflowGeometryCollection))
 struct FCreateNonOverlappingConvexHullsDataflowNode : public FDataflowNode
 {
 	GENERATED_USTRUCT_BODY()
-	DATAFLOW_NODE_DEFINE_INTERNAL(FCreateNonOverlappingConvexHullsDataflowNode, "CreateNonOverlappingConvexHulls", "Fracture|Utilities", "")
+	DATAFLOW_NODE_DEFINE_INTERNAL(FCreateNonOverlappingConvexHullsDataflowNode, "CreateNonOverlappingConvexHulls", "GeometryCollection|Utilities", "")
 
 public:
 	UPROPERTY(meta = (DataflowInput, DataflowOutput))
@@ -1132,80 +1050,6 @@ public:
 };
 
 
-/**
- *
- * Cuts geometry using a set of noised up planes
- *
- */
-USTRUCT()
-struct FPlaneCutterDataflowNode : public FDataflowNode
-{
-	GENERATED_USTRUCT_BODY()
-	DATAFLOW_NODE_DEFINE_INTERNAL(FPlaneCutterDataflowNode, "PlaneCutter", "Fracture", "")
-
-public:
-	UPROPERTY(meta = (DataflowInput, DataflowOutput))
-	FManagedArrayCollection Collection;
-
-	UPROPERTY(meta = (DataflowInput))
-	FBox BoundingBox = FBox(ForceInit);
-
-	UPROPERTY(EditAnywhere, Category = "Fracture", meta = (DataflowInput, UIMin = 1))
-	int32 NumPlanes = 1;
-
-	UPROPERTY(EditAnywhere, Category = "Fracture", meta = (DataflowInput))
-	float RandomSeed = -1.f;
-
-	UPROPERTY(EditAnywhere, Category = "Fracture", meta = (DataflowInput, UIMin = 0.f))
-	float Grout = 0.f;
-
-	UPROPERTY(EditAnywhere, Category = "Noise", meta = (DataflowInput, UIMin = 0.f));
-	float Amplitude = 0.f;
-
-	UPROPERTY(EditAnywhere, Category = "Noise", meta = (DataflowInput, UIMin = 0.00001f));
-	float Frequency = 0.1f;
-
-	UPROPERTY(EditAnywhere, Category = "Noise", meta = (DataflowInput, UIMin = 0.f));
-	float Persistence = 0.5f;
-
-	UPROPERTY(EditAnywhere, Category = "Noise", meta = (DataflowInput, UIMin = 0.f));
-	float Lacunarity = 2.f;
-
-	UPROPERTY(EditAnywhere, Category = "Noise", meta = (DataflowInput, UIMin = 0.f));
-	int32 OctaveNumber = 4;
-
-	UPROPERTY(EditAnywhere, Category = "Noise", meta = (DataflowInput, UIMin = 0.f));
-	float PointSpacing = 10.f;
-
-	UPROPERTY(EditAnywhere, Category = "Collision");
-	bool AddSamplesForCollision = false;
-
-	UPROPERTY(EditAnywhere, Category = "Collision", meta = (DataflowInput, UIMin = 0.f));
-	float CollisionSampleSpacing = 50.f;
-
-	FPlaneCutterDataflowNode(const Dataflow::FNodeParameters& InParam, FGuid InGuid = FGuid::NewGuid())
-		: FDataflowNode(InParam, InGuid)
-	{
-		RegisterInputConnection(&Collection);
-		RegisterInputConnection(&BoundingBox);
-		RegisterInputConnection(&NumPlanes);
-		RegisterInputConnection(&RandomSeed);
-		RegisterInputConnection(&Grout);
-		RegisterInputConnection(&Amplitude);
-		RegisterInputConnection(&Frequency);
-		RegisterInputConnection(&Persistence);
-		RegisterInputConnection(&Lacunarity);
-		RegisterInputConnection(&OctaveNumber);
-		RegisterInputConnection(&PointSpacing);
-		RegisterInputConnection(&CollisionSampleSpacing);
-		RegisterOutputConnection(&Collection);
-	}
-
-	virtual void Evaluate(Dataflow::FContext& Context, const FDataflowOutput* Out) const override;
-
-};
-
-
 /**
  *
  * Generates a hash value from a string
@@ -1215,7 +1059,7 @@ USTRUCT()
 struct FHashStringDataflowNode : public FDataflowNode
 {
 	GENERATED_USTRUCT_BODY()
-		DATAFLOW_NODE_DEFINE_INTERNAL(FHashStringDataflowNode, "HashString", "Utilities|String", "")
+	DATAFLOW_NODE_DEFINE_INTERNAL(FHashStringDataflowNode, "HashString", "Utilities|String", "")
 
 public:
 	/** String to hash */
@@ -1446,7 +1290,7 @@ public:
  * Gets BoundingBoxes of pieces from a Collection
  *
  */
-USTRUCT()
+USTRUCT(meta = (DataflowGeometryCollection))
 struct FGetBoundingBoxesDataflowNode : public FDataflowNode
 {
 	GENERATED_USTRUCT_BODY()
@@ -1483,7 +1327,7 @@ public:
  * Gets centroids of pieces from a Collection
  *
  */
-USTRUCT()
+USTRUCT(meta = (DataflowGeometryCollection))
 struct FGetCentroidsDataflowNode : public FDataflowNode
 {
 	GENERATED_USTRUCT_BODY()
@@ -1517,216 +1361,39 @@ public:
 
 /**
  *
- * Converts points into a DynamicMesh
+ * Transforms a mesh
  *
  */
 USTRUCT()
-struct FPointsToMeshDataflowNode : public FDataflowNode
+struct FTransformDataflowNode : public FDataflowNode
 {
 	GENERATED_USTRUCT_BODY()
-	DATAFLOW_NODE_DEFINE_INTERNAL(FPointsToMeshDataflowNode, "PointsToMesh", "Mesh|Utilities", "")
+	DATAFLOW_NODE_DEFINE_INTERNAL(FTransformDataflowNode, "Transform", "Math", "")
 
 public:
-	/** Points input */
-	UPROPERTY(meta = (DataflowInput))
-	TArray<FVector> Points;
-
-	/** Mesh output */
-	UPROPERTY(meta = (DataflowOutput))
-	TObjectPtr<UDynamicMesh> Mesh;
+	/** Translation */
+	UPROPERTY(EditAnywhere, Category = "Transform");
+	FVector Translate = FVector(0.0);
 
-	/** Mesh triangle count */
-	UPROPERTY(meta = (DataflowOutput))
-	int32 TriangleCount = 0;
+	/** Rotation */
+	UPROPERTY(EditAnywhere, Category = "Transform");
+	FVector Rotate = FVector(0.0);
 
-	FPointsToMeshDataflowNode(const Dataflow::FNodeParameters& InParam, FGuid InGuid = FGuid::NewGuid())
-		: FDataflowNode(InParam, InGuid)
-	{
-		RegisterInputConnection(&Points);
-		RegisterOutputConnection(&Mesh);
-		RegisterOutputConnection(&TriangleCount);
-	}
+	/** Scale */
+	UPROPERTY(EditAnywhere, Category = "Transform");
+	FVector Scale = FVector(1.0);
 
-	virtual void Evaluate(Dataflow::FContext& Context, const FDataflowOutput* Out) const override;
+	/** Shear */
+	UPROPERTY(EditAnywhere, Category = "Transform");
+	FVector Shear = FVector(0.0);
 
-};
+	/** Uniform scale */
+	UPROPERTY(EditAnywhere, Category = "Transform");
+	float UniformScale = 1.f;
 
-
-/**
- *
- * Converts a BoundingBox into a DynamicMesh
- *
- */
-USTRUCT()
-struct FBoxToMeshDataflowNode : public FDataflowNode
-{
-	GENERATED_USTRUCT_BODY()
-	DATAFLOW_NODE_DEFINE_INTERNAL(FBoxToMeshDataflowNode, "BoxToMesh", "Mesh|Utilities", "")
-
-public:
-	/** BoundingBox input */
-	UPROPERTY(meta = (DataflowInput))
-	FBox Box = FBox(ForceInit);
-
-	/** Mesh output */
-	UPROPERTY(meta = (DataflowOutput))
-	TObjectPtr<UDynamicMesh> Mesh;
-
-	/** Mesh triangle count */
-	UPROPERTY(meta = (DataflowOutput))
-	int32 TriangleCount = 0;
-
-	FBoxToMeshDataflowNode(const Dataflow::FNodeParameters& InParam, FGuid InGuid = FGuid::NewGuid())
-		: FDataflowNode(InParam, InGuid)
-	{
-		RegisterInputConnection(&Box);
-		RegisterOutputConnection(&Mesh);
-		RegisterOutputConnection(&TriangleCount);
-	}
-
-	virtual void Evaluate(Dataflow::FContext& Context, const FDataflowOutput* Out) const override;
-
-};
-
-
-/**
- *
- * Collects information from the DynamicMesh and outputs it into a formatted string
- *
- */
-USTRUCT()
-struct FMeshInfoDataflowNode : public FDataflowNode
-{
-	GENERATED_USTRUCT_BODY()
-	DATAFLOW_NODE_DEFINE_INTERNAL(FMeshInfoDataflowNode, "MeshInfo", "Mesh|Utilities", "")
-
-public:
-	/** DynamicMesh for the information */
-	UPROPERTY(meta = (DataflowInput))
-	TObjectPtr<UDynamicMesh> Mesh;
-
-	/** Formatted output string */
-	UPROPERTY(meta = (DataflowOutput))
-	FString InfoString = FString("");
-
-	FMeshInfoDataflowNode(const Dataflow::FNodeParameters& InParam, FGuid InGuid = FGuid::NewGuid())
-		: FDataflowNode(InParam, InGuid)
-	{
-		RegisterInputConnection(&Mesh);
-		RegisterOutputConnection(&InfoString);
-	}
-
-	virtual void Evaluate(Dataflow::FContext& Context, const FDataflowOutput* Out) const override;
-
-};
-
-
-/**
- *
- * Converts a DynamicMesh to a Collection
- *
- */
-USTRUCT()
-struct FMeshToCollectionDataflowNode : public FDataflowNode
-{
-	GENERATED_USTRUCT_BODY()
-	DATAFLOW_NODE_DEFINE_INTERNAL(FMeshToCollectionDataflowNode, "MeshToCollection", "Mesh|Utilities", "")
-
-public:
-	/** DynamicMesh to convert */
-	UPROPERTY(meta = (DataflowInput))
-	TObjectPtr<UDynamicMesh> Mesh;
-
-	/** Output Collection */
-	UPROPERTY(meta = (DataflowOutput))
-	FManagedArrayCollection Collection;
-
-	FMeshToCollectionDataflowNode(const Dataflow::FNodeParameters& InParam, FGuid InGuid = FGuid::NewGuid())
-		: FDataflowNode(InParam, InGuid)
-	{
-		RegisterInputConnection(&Mesh);
-		RegisterOutputConnection(&Collection);
-	}
-
-	virtual void Evaluate(Dataflow::FContext& Context, const FDataflowOutput* Out) const override;
-
-};
-
-
-/**
- *
- * Converts a StaticMesh into a DynamicMesh
- *
- */
-USTRUCT()
-struct FStaticMeshToMeshDataflowNode : public FDataflowNode
-{
-	GENERATED_USTRUCT_BODY()
-	DATAFLOW_NODE_DEFINE_INTERNAL(FStaticMeshToMeshDataflowNode, "StaticMeshToMesh", "Mesh|Utilities", "")
-
-public:
-	/** StaticMesh to convert */
-	UPROPERTY(EditAnywhere, Category = "StaticMesh");
-	TObjectPtr<UStaticMesh> StaticMesh;
-
-	/** Output the HiRes representation, if set to true and HiRes doesn't exist it will output empty mesh */
-	UPROPERTY(EditAnywhere, Category = "StaticMesh", meta = (DisplayName = "Use HiRes"));
-	bool UseHiRes = true;
-
-	/** Specifies the LOD level to use */
-	UPROPERTY(EditAnywhere, Category = "StaticMesh", meta = (DisplayName = "LOD Level"));
-	int32 LODLevel = 0;
-
-	/** Output mesh */
-	UPROPERTY(meta = (DataflowOutput))
-	TObjectPtr<UDynamicMesh> Mesh;
-
-	FStaticMeshToMeshDataflowNode(const Dataflow::FNodeParameters& InParam, FGuid InGuid = FGuid::NewGuid())
-		: FDataflowNode(InParam, InGuid)
-	{
-		RegisterOutputConnection(&Mesh);
-	}
-
-	virtual void Evaluate(Dataflow::FContext& Context, const FDataflowOutput* Out) const override;
-
-};
-
-
-/**
- *
- * Transforms a mesh
- *
- */
-USTRUCT()
-struct FTransformDataflowNode : public FDataflowNode
-{
-	GENERATED_USTRUCT_BODY()
-	DATAFLOW_NODE_DEFINE_INTERNAL(FTransformDataflowNode, "Transform", "Math", "")
-
-public:
-	/** Translation */
-	UPROPERTY(EditAnywhere, Category = "Transform");
-	FVector Translate = FVector(0.0);
-
-	/** Rotation */
-	UPROPERTY(EditAnywhere, Category = "Transform");
-	FVector Rotate = FVector(0.0);
-
-	/** Scale */
-	UPROPERTY(EditAnywhere, Category = "Transform");
-	FVector Scale = FVector(1.0);
-
-	/** Shear */
-	UPROPERTY(EditAnywhere, Category = "Transform");
-	FVector Shear = FVector(0.0);
-
-	/** Uniform scale */
-	UPROPERTY(EditAnywhere, Category = "Transform");
-	float UniformScale = 1.f;
-
-	/** Pivot for the rotation */
-	UPROPERTY(EditAnywhere, Category = "Pivot");
-	FVector RotatePivot = FVector(0.0);
+	/** Pivot for the rotation */
+	UPROPERTY(EditAnywhere, Category = "Pivot");
+	FVector RotatePivot = FVector(0.0);
 
 	/** Pivot for the scale */
 	UPROPERTY(EditAnywhere, Category = "Pivot");
@@ -1752,136 +1419,6 @@ public:
 };
 
 
-/**
- *
- * Appends two meshes
- *
- */
-USTRUCT()
-struct FMeshAppendDataflowNode : public FDataflowNode
-{
-	GENERATED_USTRUCT_BODY()
-	DATAFLOW_NODE_DEFINE_INTERNAL(FMeshAppendDataflowNode, "MeshAppend", "Mesh|Utilities", "")
-
-public:
-	/** Mesh input */
-	UPROPERTY(meta = (DataflowInput))
-	TObjectPtr<UDynamicMesh> Mesh1;
-
-	/** Mesh input */
-	UPROPERTY(meta = (DataflowInput))
-	TObjectPtr<UDynamicMesh> Mesh2;
-
-	/** Output mesh */
-	UPROPERTY(meta = (DataflowOutput))
-	TObjectPtr<UDynamicMesh> Mesh;
-
-	FMeshAppendDataflowNode(const Dataflow::FNodeParameters& InParam, FGuid InGuid = FGuid::NewGuid())
-		: FDataflowNode(InParam, InGuid)
-	{
-		RegisterInputConnection(&Mesh1);
-		RegisterInputConnection(&Mesh2);
-		RegisterOutputConnection(&Mesh);
-	}
-
-	virtual void Evaluate(Dataflow::FContext& Context, const FDataflowOutput* Out) const override;
-
-};
-
-
-UENUM(BlueprintType)
-enum class EMeshBooleanOperationEnum : uint8
-{
-	Dataflow_MeshBoolean_Union UMETA(DisplayName = "Union"),
-	Dataflow_MeshBoolean_Intersect UMETA(DisplayName = "Intersect"),
-	Dataflow_MeshBoolean_Difference UMETA(DisplayName = "Difference"),
-	//~~~
-	//256th entry
-	Dataflow_Max                UMETA(Hidden)
-};
-
-/**
- *
- * Mesh boolean (Union, Intersect, Difference) between two meshes
- *
- */
-USTRUCT()
-struct FMeshBooleanDataflowNode : public FDataflowNode
-{
-	GENERATED_USTRUCT_BODY()
-	DATAFLOW_NODE_DEFINE_INTERNAL(FMeshBooleanDataflowNode, "MeshBoolean", "Mesh|Utilities", "")
-
-public:
-	/** Boolean operation */
-	UPROPERTY(EditAnywhere, Category = "Boolean");
-	EMeshBooleanOperationEnum Operation = EMeshBooleanOperationEnum::Dataflow_MeshBoolean_Intersect;
-
-	/** Mesh input */
-	UPROPERTY(meta = (DataflowInput))
-	TObjectPtr<UDynamicMesh> Mesh1;
-
-	/** Mesh input */
-	UPROPERTY(meta = (DataflowInput))
-	TObjectPtr<UDynamicMesh> Mesh2;
-
-	/** Output mesh */
-	UPROPERTY(meta = (DataflowOutput))
-	TObjectPtr<UDynamicMesh> Mesh;
-
-	FMeshBooleanDataflowNode(const Dataflow::FNodeParameters& InParam, FGuid InGuid = FGuid::NewGuid())
-		: FDataflowNode(InParam, InGuid)
-	{
-		RegisterInputConnection(&Mesh1);
-		RegisterInputConnection(&Mesh2);
-		RegisterOutputConnection(&Mesh);
-	}
-
-	virtual void Evaluate(Dataflow::FContext& Context, const FDataflowOutput* Out) const override;
-
-};
-
-
-/**
- *
- * Copies the same mesh with scale onto points
- *
- */
-USTRUCT()
-struct FMeshCopyToPointsDataflowNode : public FDataflowNode
-{
-	GENERATED_USTRUCT_BODY()
-	DATAFLOW_NODE_DEFINE_INTERNAL(FMeshCopyToPointsDataflowNode, "MeshCopyToPoints", "Mesh|Utilities", "")
-
-public:
-	/** Points to copy meshes onto */
-	UPROPERTY(meta = (DataflowInput))
-	TArray<FVector> Points;
-
-	/** Mesh to copy onto points */
-	UPROPERTY(meta = (DataflowInput))
-	TObjectPtr<UDynamicMesh> MeshToCopy;
-
-	/** Scale appied to the mesh */
-	UPROPERTY(EditAnywhere, Category = "Copy");
-	float Scale = 1.f;
-
-	/** Copied meshes */
-	UPROPERTY(meta = (DataflowOutput))
-	TObjectPtr<UDynamicMesh> Mesh;
-
-	FMeshCopyToPointsDataflowNode(const Dataflow::FNodeParameters& InParam, FGuid InGuid = FGuid::NewGuid())
-		: FDataflowNode(InParam, InGuid)
-	{
-		RegisterInputConnection(&Points);
-		RegisterInputConnection(&MeshToCopy);
-		RegisterOutputConnection(&Mesh);
-	}
-
-	virtual void Evaluate(Dataflow::FContext& Context, const FDataflowOutput* Out) const override;
-
-};
-
-
 UENUM(BlueprintType)
 enum class ECompareOperationEnum : uint8
 {
@@ -1978,54 +1515,12 @@ public:
 };
 
 
-/**
- *
- * Outputs Mesh data
- *
- */
-USTRUCT()
-struct FGetMeshDataDataflowNode : public FDataflowNode
-{
-	GENERATED_USTRUCT_BODY()
-	DATAFLOW_NODE_DEFINE_INTERNAL(FGetMeshDataDataflowNode, "GetMeshData", "Mesh|Utilities", "")
-
-public:
-	/** Mesh for the data */
-	UPROPERTY(meta = (DataflowInput))
-	TObjectPtr<UDynamicMesh> Mesh;
-
-	/** Number of vertices */
-	UPROPERTY(meta = (DataflowOutput))
-	int32 VertexCount = 0;
-
-	/** Number of edges */
-	UPROPERTY(meta = (DataflowOutput))
-	int32 EdgeCount = 0;
-
-	/** Number of triangles */
-	UPROPERTY(meta = (DataflowOutput))
-	int32 TriangleCount = 0;
-
-	FGetMeshDataDataflowNode(const Dataflow::FNodeParameters& InParam, FGuid InGuid = FGuid::NewGuid())
-		: FDataflowNode(InParam, InGuid)
-	{
-		RegisterInputConnection(&Mesh);
-		RegisterOutputConnection(&VertexCount);
-		RegisterOutputConnection(&EdgeCount);
-		RegisterOutputConnection(&TriangleCount);
-	}
-
-	virtual void Evaluate(Dataflow::FContext& Context, const FDataflowOutput* Out) const override;
-
-};
-
-
 /**
  *
  * Collects grooup and attribute information from the Collection and outputs it into a formatted string
  *
  */
-USTRUCT()
+USTRUCT(meta = (DataflowGeometryCollection))
 struct FGetSchemaDataflowNode : public FDataflowNode
 {
 	GENERATED_USTRUCT_BODY()
@@ -2052,106 +1547,11 @@ public:
 };
 
 
-UENUM(BlueprintType)
-enum class EClusterSizeMethodEnum : uint8
-{
-	Dataflow_ClusterSizeMethod_ByNumber UMETA(DisplayName = "By Number"),
-	Dataflow_ClusterSizeMethod_ByFractionOfInput UMETA(DisplayName = "By Fraction Of Input"),
-	Dataflow_ClusterSizeMethod_BySize UMETA(DisplayName = "By Size"),
-	//~~~
-	//256th entry
-	Dataflow_Max                UMETA(Hidden)
-};
-
-/**
- *
- * Automatically group pieces of a fractured Collection into a specified number of clusters
- *
- */
-USTRUCT()
-struct FAutoClusterDataflowNode : public FDataflowNode
-{
-	GENERATED_USTRUCT_BODY()
-	DATAFLOW_NODE_DEFINE_INTERNAL(FAutoClusterDataflowNode, "AutoCluster", "GeometryCollection|Cluster", "")
-
-public:
-	/** How to choose the size of the clusters to create */
-	UPROPERTY(EditAnywhere, Category = "Cluster Size");
-	EClusterSizeMethodEnum ClusterSizeMethod = EClusterSizeMethodEnum::Dataflow_ClusterSizeMethod_ByNumber;
-
-	/** Use a Voronoi diagram with this many Voronoi sites as a guide for deciding cluster boundaries */
-	UPROPERTY(EditAnywhere, Category = "Cluster Size", meta = (DataflowInput, UIMin = 2, UIMax = 5000, EditCondition = "ClusterSizeMethod == EClusterSizeMethodEnum::Dataflow_ClusterSizeMethod_ByNumber"))
-	int32 ClusterSites = 10;
-
-	/** Choose the number of Voronoi sites used for clustering as a fraction of the number of child bones to process */
-	UPROPERTY(EditAnywhere, Category = "Cluster Size", meta = (DataflowInput, UIMin = 0.f, UIMax = 0.5f, EditCondition = "ClusterSizeMethod == EClusterSizeMethodEnum::Dataflow_ClusterSizeMethod_ByFractionOfInput"))
-	float ClusterFraction = 0.25;
-
-	/** Choose the Edge-Size of the cube used to groups bones under a cluster (in cm). */
-	UPROPERTY(EditAnywhere, Category = "ClusterSize", meta = (DataflowInput, DisplayName = "Cluster Size", UIMin = ".01", UIMax = "100", ClampMin = ".0001", ClampMax = "10000", EditCondition = "ClusterSizeMethod == EClusterSizeMethodEnum::Dataflow_ClusterSizeMethod_BySize"))
-	float SiteSize = 1;
-
-	/** If true, bones will only be added to the same cluster if they are physically connected (either directly, or via other bones in the same cluster) */
-	UPROPERTY(EditAnywhere, Category = "Auto Cluster")
-	bool AutoCluster = true;
-
-	/** If true, prevent the creation of clusters with only a single child. Either by merging into a neighboring cluster, or not creating the cluster. */
-	UPROPERTY(EditAnywhere, Category = "Auto Cluster")
-	bool AvoidIsolated = true;
-
-	/** Fractured GeometryCollection to cluster */
-	UPROPERTY(meta = (DataflowInput, DataflowOutput, DataflowPassthrough = "Collection", DataflowIntrinsic))
-	FManagedArrayCollection Collection;
-
-	/** Bone selection for the clustering */
-	UPROPERTY(meta = (DataflowInput, DisplayName = "TransformSelection"))
-	FDataflowTransformSelection TransformSelection;
-
-	FAutoClusterDataflowNode(const Dataflow::FNodeParameters& InParam, FGuid InGuid = FGuid::NewGuid())
-		: FDataflowNode(InParam, InGuid)
-	{
-		RegisterInputConnection(&Collection);
-		RegisterInputConnection(&TransformSelection);
-		RegisterInputConnection(&ClusterSites);
-		RegisterInputConnection(&ClusterFraction);
-		RegisterInputConnection(&SiteSize);
-		RegisterOutputConnection(&Collection, &Collection);
-	}
-
-	virtual void Evaluate(Dataflow::FContext& Context, const FDataflowOutput* Out) const override;
-
-};
-
-
-USTRUCT()
-struct FClusterFlattenDataflowNode : public FDataflowNode
-{
-	GENERATED_USTRUCT_BODY()
-	DATAFLOW_NODE_DEFINE_INTERNAL(FClusterFlattenDataflowNode, "Flatten", "GeometryCollection|Cluster", "")
-
-public:
-	// @todo(harsha) Support Selections
-
-	UPROPERTY(meta = (DataflowInput, DataflowOutput, DataflowPassthrough = "Collection", DataflowIntrinsic))
-	FManagedArrayCollection Collection;
-
-	FClusterFlattenDataflowNode(const Dataflow::FNodeParameters& InParam, FGuid InGuid = FGuid::NewGuid())
-		: FDataflowNode(InParam, InGuid)
-	{
-		RegisterInputConnection(&Collection);
-		RegisterOutputConnection(&Collection, &Collection);
-	}
-
-	virtual void Evaluate(Dataflow::FContext& Context, const FDataflowOutput* Out) const override;
-
-};
-
-
-USTRUCT()
+USTRUCT(meta = (DataflowGeometryCollection))
 struct FRemoveOnBreakDataflowNode : public FDataflowNode
 {
 	GENERATED_USTRUCT_BODY()
-	DATAFLOW_NODE_DEFINE_INTERNAL(FRemoveOnBreakDataflowNode, "RemoveOnBreak", "Fracture|Utilities", "")
+	DATAFLOW_NODE_DEFINE_INTERNAL(FRemoveOnBreakDataflowNode, "RemoveOnBreak", "GeometryCollection", "")
 
 public:
 	// @todo(harsha) Support Selections
@@ -2198,7 +1598,7 @@ enum class EAnchorStateEnum : uint8
  * Sets the anchored state on the selected bones in a Collection
  *
  */
-USTRUCT()
+USTRUCT(meta = (DataflowGeometryCollection))
 struct FSetAnchorStateDataflowNode : public FDataflowNode
 {
 	GENERATED_USTRUCT_BODY()
@@ -2251,7 +1651,7 @@ enum class EProximityMethodEnum : uint8
  * Update the proximity (contact) graph for the bones in a Collection
  *
  */
-USTRUCT()
+USTRUCT(meta = (DataflowGeometryCollection))
 struct FProximityDataflowNode : public FDataflowNode
 {
 	GENERATED_USTRUCT_BODY()
@@ -2290,4 +1690,3 @@ namespace Dataflow
 {
 	void GeometryCollectionEngineNodes();
 }
-
diff --git a/Engine/Plugins/Experimental/GeometryCollectionPlugin/Source/GeometryCollectionNodes/Public/Dataflow/GeometryCollectionProcessingNodes.h b/Engine/Plugins/Experimental/GeometryCollectionPlugin/Source/GeometryCollectionNodes/Public/Dataflow/GeometryCollectionProcessingNodes.h
--- a/Engine/Plugins/Experimental/GeometryCollectionPlugin/Source/GeometryCollectionNodes/Public/Dataflow/GeometryCollectionProcessingNodes.h
+++ b/Engine/Plugins/Experimental/GeometryCollectionPlugin/Source/GeometryCollectionNodes/Public/Dataflow/GeometryCollectionProcessingNodes.h
@@ -9,7 +9,7 @@
 #include "GeometryCollectionProcessingNodes.generated.h"
 
 
-USTRUCT()
+USTRUCT(meta = (DataflowGeometryCollection))
 struct FCloseGeometryOnCollectionDataflowNode : public FDataflowNode
 {
 	GENERATED_USTRUCT_BODY()
diff --git a/Engine/Plugins/Experimental/GeometryCollectionPlugin/Source/GeometryCollectionNodes/Public/Dataflow/GeometryCollectionSelectionNodes.h b/Engine/Plugins/Experimental/GeometryCollectionPlugin/Source/GeometryCollectionNodes/Public/Dataflow/GeometryCollectionSelectionNodes.h
--- a/Engine/Plugins/Experimental/GeometryCollectionPlugin/Source/GeometryCollectionNodes/Public/Dataflow/GeometryCollectionSelectionNodes.h
+++ b/Engine/Plugins/Experimental/GeometryCollectionPlugin/Source/GeometryCollectionNodes/Public/Dataflow/GeometryCollectionSelectionNodes.h
@@ -18,7 +18,7 @@ class FGeometryCollection;
  * Selects all the bones for the Collection
  *
  */
-USTRUCT()
+USTRUCT(meta = (DataflowGeometryCollection))
 struct FCollectionTransformSelectionAllDataflowNode : public FDataflowNode
 {
 	GENERATED_USTRUCT_BODY()
@@ -62,7 +62,7 @@ enum class ESetOperationEnum : uint8
  * Runs boolean operation on incoming TransformSelections
  *
  */
-USTRUCT()
+USTRUCT(meta = (DataflowGeometryCollection))
 struct FCollectionTransformSelectionSetOperationDataflowNode : public FDataflowNode
 {
 	GENERATED_USTRUCT_BODY()
@@ -103,7 +103,7 @@ public:
  * Generates a formatted string of the bones and the selection
  *
  */
-USTRUCT()
+USTRUCT(meta = (DataflowGeometryCollection))
 struct FCollectionTransformSelectionInfoDataflowNode : public FDataflowNode
 {
 	GENERATED_USTRUCT_BODY()
@@ -140,7 +140,7 @@ public:
  * Generates an empty bone selection for the Collection
  *
  */
-USTRUCT()
+USTRUCT(meta = (DataflowGeometryCollection))
 struct FCollectionTransformSelectionNoneDataflowNode : public FDataflowNode
 {
 	GENERATED_USTRUCT_BODY()
@@ -173,7 +173,7 @@ public:
  * Inverts the incoming selection of bones
  *
  */
-USTRUCT()
+USTRUCT(meta = (DataflowGeometryCollection))
 struct FCollectionTransformSelectionInvertDataflowNode : public FDataflowNode
 {
 	GENERATED_USTRUCT_BODY()
@@ -201,7 +201,7 @@ public:
  * Selects bones randomly in the Collection
  *
  */
-USTRUCT()
+USTRUCT(meta = (DataflowGeometryCollection))
 struct FCollectionTransformSelectionRandomDataflowNode : public FDataflowNode
 {
 	GENERATED_USTRUCT_BODY()
@@ -249,7 +249,7 @@ public:
  * Selects the root bones in the Collection
  *
  */
-USTRUCT()
+USTRUCT(meta = (DataflowGeometryCollection))
 struct FCollectionTransformSelectionRootDataflowNode : public FDataflowNode
 {
 	GENERATED_USTRUCT_BODY()
@@ -283,7 +283,7 @@ public:
  * space separated list
  *
  */
-USTRUCT(meta = (DataflowDeprecated, DataflowTerminal))
+USTRUCT(meta = (DataflowGeometryCollection))
 struct FCollectionTransformSelectionCustomDataflowNode : public FDataflowNode
 {
 	GENERATED_USTRUCT_BODY()
@@ -321,7 +321,7 @@ public:
  * Selects the parents of the currently selected bones
  *
  */
-USTRUCT(meta = (DataflowDeprecated, DataflowTerminal))
+USTRUCT(meta = (DataflowGeometryCollection))
 struct FCollectionTransformSelectionParentDataflowNode : public FDataflowNode
 {
 	GENERATED_USTRUCT_BODY()
@@ -355,7 +355,7 @@ public:
  * Outputs the specified percentage of the incoming bone selection
  *
  */
-USTRUCT(meta = (DataflowDeprecated, DataflowTerminal))
+USTRUCT(meta = (DataflowGeometryCollection))
 struct FCollectionTransformSelectionByPercentageDataflowNode : public FDataflowNode
 {
 	GENERATED_USTRUCT_BODY()
@@ -366,10 +366,6 @@ public:
 	UPROPERTY(meta = (DataflowInput, DataflowOutput, DisplayName = "TransformSelection", DataflowPassthrough = "TransformSelection", DataflowIntrinsic))
 	FDataflowTransformSelection TransformSelection;
 
-	/** GeometryCollection for the selection */
-	UPROPERTY(meta = (DataflowInput, DataflowOutput, DataflowPassthrough = "Collection", DataflowIntrinsic))
-	FManagedArrayCollection Collection;
-
 	/** Percentage to keep from the original selection */
 	UPROPERTY(EditAnywhere, Category = "Selection", meta = (UIMin = 0, UIMax = 100))
 	int32 Percentage = 100;
@@ -387,11 +383,9 @@ public:
 	{
 		RandomSeed = FMath::RandRange(-100000, 100000);
 		RegisterInputConnection(&TransformSelection);
-		RegisterInputConnection(&Collection);
 		RegisterInputConnection(&Percentage);
 		RegisterInputConnection(&RandomSeed);
 		RegisterOutputConnection(&TransformSelection, &TransformSelection);
-		RegisterOutputConnection(&Collection, &Collection);
 	}
 
 	virtual void Evaluate(Dataflow::FContext& Context, const FDataflowOutput* Out) const override;
@@ -404,7 +398,7 @@ public:
  * Selects the children of the incoming bone selection
  *
  */
-USTRUCT(meta = (DataflowDeprecated, DataflowTerminal))
+USTRUCT(meta = (DataflowGeometryCollection))
 struct FCollectionTransformSelectionChildrenDataflowNode : public FDataflowNode
 {
 	GENERATED_USTRUCT_BODY()
@@ -438,7 +432,7 @@ public:
  * Selects the siblings of the incoming bone selection
  *
  */
-USTRUCT(meta = (DataflowDeprecated, DataflowTerminal))
+USTRUCT(meta = (DataflowGeometryCollection))
 struct FCollectionTransformSelectionSiblingsDataflowNode : public FDataflowNode
 {
 	GENERATED_USTRUCT_BODY()
@@ -472,7 +466,7 @@ public:
  * Selects the level of the incoming bone selection
  *
  */
-USTRUCT(meta = (DataflowDeprecated, DataflowTerminal))
+USTRUCT(meta = (DataflowGeometryCollection))
 struct FCollectionTransformSelectionLevelDataflowNode : public FDataflowNode
 {
 	GENERATED_USTRUCT_BODY()
@@ -506,7 +500,7 @@ public:
  * Selects the contact(s) of the incoming bone selection
  *
  */
-USTRUCT(meta = (DataflowDeprecated, DataflowTerminal))
+USTRUCT(meta = (DataflowGeometryCollection))
 struct FCollectionTransformSelectionContactDataflowNode : public FDataflowNode
 {
 	GENERATED_USTRUCT_BODY()
@@ -540,7 +534,7 @@ public:
  * Selects the leaves in the GeometryCollection
  *
  */
-USTRUCT(meta = (DataflowDeprecated, DataflowTerminal))
+USTRUCT(meta = (DataflowGeometryCollection))
 struct FCollectionTransformSelectionLeafDataflowNode : public FDataflowNode
 {
 	GENERATED_USTRUCT_BODY()
@@ -573,7 +567,7 @@ public:
  * Selects the clusters in the GeometryCollection
  *
  */
-USTRUCT(meta = (DataflowDeprecated, DataflowTerminal))
+USTRUCT(meta = (DataflowGeometryCollection))
 struct FCollectionTransformSelectionClusterDataflowNode : public FDataflowNode
 {
 	GENERATED_USTRUCT_BODY()
@@ -606,7 +600,7 @@ public:
  * Selects pieces based on their size
  *
  */
-USTRUCT(meta = (DataflowDeprecated, DataflowTerminal))
+USTRUCT(meta = (DataflowGeometryCollection))
 struct FCollectionTransformSelectionBySizeDataflowNode : public FDataflowNode
 {
 	GENERATED_USTRUCT_BODY()
@@ -649,7 +643,7 @@ public:
  * Selects pieces based on their volume
  *
  */
-USTRUCT(meta = (DataflowDeprecated, DataflowTerminal))
+USTRUCT(meta = (DataflowGeometryCollection))
 struct FCollectionTransformSelectionByVolumeDataflowNode : public FDataflowNode
 {
 	GENERATED_USTRUCT_BODY()
diff --git a/Engine/Plugins/Experimental/GeometryCollectionPlugin/Source/GeometryCollectionNodes/Public/Dataflow/GeometryCollectionSkeletalMeshNodes.h b/Engine/Plugins/Experimental/GeometryCollectionPlugin/Source/GeometryCollectionNodes/Public/Dataflow/GeometryCollectionSkeletalMeshNodes.h
--- a/Engine/Plugins/Experimental/GeometryCollectionPlugin/Source/GeometryCollectionNodes/Public/Dataflow/GeometryCollectionSkeletalMeshNodes.h
+++ b/Engine/Plugins/Experimental/GeometryCollectionPlugin/Source/GeometryCollectionNodes/Public/Dataflow/GeometryCollectionSkeletalMeshNodes.h
@@ -10,7 +10,7 @@
 
 class USkeletalMesh;
 
-USTRUCT()
+USTRUCT(meta = (DataflowGeometryCollection))
 struct FSkeletonToCollectionDataflowNode : public FDataflowNode
 {
 	GENERATED_USTRUCT_BODY()
diff --git a/Engine/Plugins/Experimental/GeometryCollectionPlugin/Source/GeometryCollectionNodes/Public/Dataflow/Nodes/GeometryCollectionAssetNodes.h b/Engine/Plugins/Experimental/GeometryCollectionPlugin/Source/GeometryCollectionNodes/Public/Dataflow/Nodes/GeometryCollectionAssetNodes.h
--- a/Engine/Plugins/Experimental/GeometryCollectionPlugin/Source/GeometryCollectionNodes/Public/Dataflow/Nodes/GeometryCollectionAssetNodes.h
+++ b/Engine/Plugins/Experimental/GeometryCollectionPlugin/Source/GeometryCollectionNodes/Public/Dataflow/Nodes/GeometryCollectionAssetNodes.h
@@ -11,11 +11,11 @@
 
 class UGeometryCollection;
 
-USTRUCT()
+USTRUCT(meta = (DataflowGeometryCollection, DataflowTerminal))
 struct FGeometryCollectionTerminalDataflowNode : public FDataflowTerminalNode
 {
 	GENERATED_USTRUCT_BODY()
-	DATAFLOW_NODE_DEFINE_INTERNAL(FGeometryCollectionTerminalDataflowNode, "GeometryCollectionTerminal", "GeometryCollection", "")
+	DATAFLOW_NODE_DEFINE_INTERNAL(FGeometryCollectionTerminalDataflowNode, "GeometryCollectionTerminal", "Terminal", "")
 	DATAFLOW_NODE_RENDER_TYPE(FGeometryCollection::StaticType(), "Collection")
 
 public:
@@ -46,7 +46,7 @@ public:
  * Get Current geometry collection asset 
  * Note : Use with caution as this may get replaced in a near future for a more generic getAsset node
  */
-USTRUCT()
+USTRUCT(meta = (DataflowGeometryCollection))
 struct FGetGeometryCollectionAssetDataflowNode : public FDataflowNode
 {
 	GENERATED_USTRUCT_BODY()
@@ -65,7 +65,7 @@ public:
  * Get the list of the original mesh information used to create a specific geometryc collection asset
  * each entry contains a mesh, a transform and a list of override materials
  */
-USTRUCT()
+USTRUCT(meta = (DataflowGeometryCollection))
 struct FGetGeometryCollectionSourcesDataflowNode : public FDataflowNode
 {
 	GENERATED_USTRUCT_BODY()
@@ -87,7 +87,7 @@ public:
 /**
  * create a geometry collection from a set of geometry sources    
  */
-USTRUCT()
+USTRUCT(meta = (DataflowGeometryCollection))
 struct FCreateGeometryCollectionFromSourcesDataflowNode : public FDataflowNode
 {
 	GENERATED_USTRUCT_BODY()
diff --git a/Engine/Source/Runtime/Experimental/Dataflow/Core/Private/Dataflow/DataflowNode.cpp b/Engine/Source/Runtime/Experimental/Dataflow/Core/Private/Dataflow/DataflowNode.cpp
--- a/Engine/Source/Runtime/Experimental/Dataflow/Core/Private/Dataflow/DataflowNode.cpp
+++ b/Engine/Source/Runtime/Experimental/Dataflow/Core/Private/Dataflow/DataflowNode.cpp
@@ -13,6 +13,9 @@ const FName FDataflowNode::DataflowOutput = TEXT("DataflowOutput");
 const FName FDataflowNode::DataflowPassthrough = TEXT("DataflowPassthrough");
 const FName FDataflowNode::DataflowIntrinsic = TEXT("DataflowIntrinsic");
 
+const FLinearColor FDataflowNode::DefaultNodeTitleColor = FLinearColor(1.f, 1.f, 0.8f);
+const FLinearColor FDataflowNode::DefaultNodeBodyTintColor = FLinearColor(0.f, 0.f, 0.f, 0.5f);
+
 //
 // Inputs
 //
diff --git a/Engine/Source/Runtime/Experimental/Dataflow/Core/Private/Dataflow/DataflowSelection.cpp b/Engine/Source/Runtime/Experimental/Dataflow/Core/Private/Dataflow/DataflowSelection.cpp
--- a/Engine/Source/Runtime/Experimental/Dataflow/Core/Private/Dataflow/DataflowSelection.cpp
+++ b/Engine/Source/Runtime/Experimental/Dataflow/Core/Private/Dataflow/DataflowSelection.cpp
@@ -33,7 +33,7 @@ void FDataflowSelection::AsArray(TArray<int32>& SelectionArr) const
 
 void FDataflowSelection::SetFromArray(const TArray<int32>& SelectionArr)
 {
-	SelectionArray.Init(false, SelectionArr.Num());
+	SelectionArray.Init(false, SelectionArray.Num());
 
 	for (int32 Elem : SelectionArr)
 	{
@@ -58,35 +58,36 @@ void FDataflowSelection::XOR(const FDataflowSelection& Other, FDataflowSelection
 
 int32 FDataflowSelection::NumSelected() const
 {
-	// todo : use SelectionArray.CountSetBits instead 
-	int32 NumSelectedBits = 0;
+	return SelectionArray.CountSetBits(0, SelectionArray.Num());
+}
 
+bool FDataflowSelection::AnySelected() const
+{
 	TBitArray<>::FConstIterator It(SelectionArray);
 	while (It)
-	{		 
+	{
 		if (It.GetValue())
 		{
-			NumSelectedBits++;
+			return true;
 		}
 
 		++It;
 	}
 
-	return NumSelectedBits;
+	return false;
 }
 
-bool FDataflowSelection::AnySelected() const
+void FDataflowSelection::SetWithMask(const bool Value, const FDataflowSelection& Mask)
 {
-	TBitArray<>::FConstIterator It(SelectionArray);
-	while (It)
+	if (SelectionArray.Num() == Mask.Num())
 	{
-		if (It.GetValue())
+		for (int32 Idx = 0; Idx < SelectionArray.Num(); ++Idx)
 		{
-			return true;
+			if (Mask.IsSelected(Idx))
+			{
+				SelectionArray[Idx] = Value;
+			}
 		}
-
-		++It;
 	}
-
-	return false;
 }
+
diff --git a/Engine/Source/Runtime/Experimental/Dataflow/Core/Public/Dataflow/DataflowNode.h b/Engine/Source/Runtime/Experimental/Dataflow/Core/Public/Dataflow/DataflowNode.h
--- a/Engine/Source/Runtime/Experimental/Dataflow/Core/Public/Dataflow/DataflowNode.h
+++ b/Engine/Source/Runtime/Experimental/Dataflow/Core/Public/Dataflow/DataflowNode.h
@@ -111,6 +111,9 @@ struct DATAFLOWCORE_API FDataflowNode
 	static const FName DataflowPassthrough;
 	static const FName DataflowIntrinsic;
 
+	static const FLinearColor DefaultNodeTitleColor;
+	static const FLinearColor DefaultNodeBodyTintColor;
+
 	//
 	//  Struct Support
 	//
diff --git a/Engine/Source/Runtime/Experimental/Dataflow/Core/Public/Dataflow/DataflowSelection.h b/Engine/Source/Runtime/Experimental/Dataflow/Core/Public/Dataflow/DataflowSelection.h
--- a/Engine/Source/Runtime/Experimental/Dataflow/Core/Public/Dataflow/DataflowSelection.h
+++ b/Engine/Source/Runtime/Experimental/Dataflow/Core/Public/Dataflow/DataflowSelection.h
@@ -23,7 +23,7 @@ struct DATAFLOWCORE_API FDataflowSelection
 	void OR(const FDataflowSelection& Other, FDataflowSelection& Result) const;
 	void XOR(const FDataflowSelection& Other, FDataflowSelection& Result) const;
 	void Invert() { SelectionArray.BitwiseNOT(); }
-
+	void SetWithMask(const bool Value, const FDataflowSelection& Mask);
 private:
 	TBitArray<> SelectionArray;
 };
diff --git a/Engine/Source/Runtime/Experimental/Dataflow/Engine/Private/Dataflow/DataflowEdNode.cpp b/Engine/Source/Runtime/Experimental/Dataflow/Engine/Private/Dataflow/DataflowEdNode.cpp
--- a/Engine/Source/Runtime/Experimental/Dataflow/Engine/Private/Dataflow/DataflowEdNode.cpp
+++ b/Engine/Source/Runtime/Experimental/Dataflow/Engine/Private/Dataflow/DataflowEdNode.cpp
@@ -181,7 +181,7 @@ FLinearColor UDataflowEdNode::GetNodeTitleColor() const
 			}
 		}
 	}
-	return FLinearColor(0.f, 0.f, 0.f);
+	return FDataflowNode::DefaultNodeTitleColor;
 }
 
 FLinearColor UDataflowEdNode::GetNodeBodyTintColor() const
@@ -196,7 +196,7 @@ FLinearColor UDataflowEdNode::GetNodeBodyTintColor() const
 			}
 		}
 	}
-	return FLinearColor(0.f, 0.f, 0.f);
+	return FDataflowNode::DefaultNodeBodyTintColor;
 }
 
 FText UDataflowEdNode::GetTooltipText() const
