Dataflow:
- Added GetBoundingBox() to CollectionBoundsFacade to compute BoundingBox for the whole collection and updated BoundingBox node to use it

diff --git a/Engine/Plugins/Experimental/GeometryCollectionPlugin/Source/GeometryCollectionNodes/Private/Dataflow/GeometryCollectionNodes.cpp b/Engine/Plugins/Experimental/GeometryCollectionPlugin/Source/GeometryCollectionNodes/Private/Dataflow/GeometryCollectionNodes.cpp
--- a/Engine/Plugins/Experimental/GeometryCollectionPlugin/Source/GeometryCollectionNodes/Private/Dataflow/GeometryCollectionNodes.cpp
+++ b/Engine/Plugins/Experimental/GeometryCollectionPlugin/Source/GeometryCollectionNodes/Private/Dataflow/GeometryCollectionNodes.cpp
@@ -264,31 +264,6 @@ void FMakeLiteralStringDataflowNode::Evaluate(Dataflow::FContext& Context, const
 	}
 }
 
-void ComputeBoundingBox(const FManagedArrayCollection& Collection, FBox& BoundingBox)
-{
-	if (Collection.HasAttribute("Transform", FGeometryCollection::TransformGroup) &&
-		Collection.HasAttribute("Parent", FGeometryCollection::TransformGroup) &&
-		Collection.HasAttribute("TransformIndex", FGeometryCollection::GeometryGroup) &&
-		Collection.HasAttribute("BoundingBox", FGeometryCollection::GeometryGroup))
-	{
-		const TManagedArray<FTransform>& Transforms = Collection.GetAttribute<FTransform>("Transform", FGeometryCollection::TransformGroup);
-		const TManagedArray<int32>& ParentIndices = Collection.GetAttribute<int32>("Parent", FGeometryCollection::TransformGroup);
-		const TManagedArray<int32>& TransformIndices = Collection.GetAttribute<int32>("TransformIndex", FGeometryCollection::GeometryGroup);
-		const TManagedArray<FBox>& BoundingBoxes = Collection.GetAttribute<FBox>("BoundingBox", FGeometryCollection::GeometryGroup);
-
-		TArray<FMatrix> TmpGlobalMatrices;
-		GeometryCollectionAlgo::GlobalMatrices(Transforms, ParentIndices, TmpGlobalMatrices);
-
-		if (TmpGlobalMatrices.Num() > 0)
-		{
-			for (int32 BoxIdx = 0; BoxIdx < BoundingBoxes.Num(); ++BoxIdx)
-			{
-				const int32 TransformIndex = TransformIndices[BoxIdx];
-				BoundingBox += BoundingBoxes[BoxIdx].TransformBy(TmpGlobalMatrices[TransformIndex]);
-			}
-		}
-	}
-}
 
 void FBoundingBoxDataflowNode::Evaluate(Dataflow::FContext& Context, const FDataflowOutput* Out) const
 {
@@ -296,14 +271,14 @@ void FBoundingBoxDataflowNode::Evaluate(Dataflow::FContext& Context, const FData
 	{
 		const FManagedArrayCollection& InCollection = GetValue<FManagedArrayCollection>(Context, &Collection);
 
-		FBox BBox(ForceInit);
-
-		ComputeBoundingBox(InCollection, BBox);
+		GeometryCollection::Facades::FBoundsFacade BoundsFacade(InCollection);
+		const FBox& BBox = BoundsFacade.GetBoundingBox();
 
 		SetValue<FBox>(Context, BBox, &BoundingBox);
 	}
 }
 
+
 void FExpandBoundingBoxDataflowNode::Evaluate(Dataflow::FContext& Context, const FDataflowOutput* Out) const
 {
 	FBox BBox = GetValue<FBox>(Context, &BoundingBox);
@@ -852,7 +827,8 @@ void FGetBoundingBoxesDataflowNode::Evaluate(Dataflow::FContext& Context, const
 		const FManagedArrayCollection& InCollection = GetValue<FManagedArrayCollection>(Context, &Collection);
 		const FDataflowTransformSelection& InTransformSelection = GetValue<FDataflowTransformSelection>(Context, &TransformSelection);
 
-		const TManagedArray<FBox>& InBoundingBoxes = GeometryCollection::Facades::FBoundsFacade(InCollection).GetBoundingBoxes();
+		GeometryCollection::Facades::FBoundsFacade BoundsFacade(InCollection);
+		const TManagedArray<FBox>& InBoundingBoxes = BoundsFacade.GetBoundingBoxes();
 
 		TArray<FBox> BoundingBoxesArr;
 		for (int32 Idx = 0; Idx < InBoundingBoxes.Num(); ++Idx)
@@ -881,7 +857,8 @@ void FGetCentroidsDataflowNode::Evaluate(Dataflow::FContext& Context, const FDat
 		const FManagedArrayCollection& InCollection = GetValue<FManagedArrayCollection>(Context, &Collection);
 		const FDataflowTransformSelection& InTransformSelection = GetValue<FDataflowTransformSelection>(Context, &TransformSelection);
 
-		const TManagedArray<FBox>& InBoundingBoxes = GeometryCollection::Facades::FBoundsFacade(InCollection).GetBoundingBoxes();
+		GeometryCollection::Facades::FBoundsFacade BoundsFacade(InCollection);
+		const TManagedArray<FBox>& InBoundingBoxes = BoundsFacade.GetBoundingBoxes();
 
 		TArray<FVector> CentroidsArr;
 		for (int32 Idx = 0; Idx < InBoundingBoxes.Num(); ++Idx)
diff --git a/Engine/Source/Runtime/Experimental/Chaos/Private/GeometryCollection/Facades/CollectionBoundsFacade.cpp b/Engine/Source/Runtime/Experimental/Chaos/Private/GeometryCollection/Facades/CollectionBoundsFacade.cpp
--- a/Engine/Source/Runtime/Experimental/Chaos/Private/GeometryCollection/Facades/CollectionBoundsFacade.cpp
+++ b/Engine/Source/Runtime/Experimental/Chaos/Private/GeometryCollection/Facades/CollectionBoundsFacade.cpp
@@ -6,6 +6,7 @@
 
 #include "GeometryCollection/Facades/CollectionBoundsFacade.h"
 #include "GeometryCollection/GeometryCollection.h"
+#include "GeometryCollection/Facades/CollectionHierarchyFacade.h"
 
 namespace GeometryCollection::Facades
 {
@@ -15,6 +16,7 @@ namespace GeometryCollection::Facades
 		, VertexAttribute(InCollection,"Vertex", FGeometryCollection::VerticesGroup, FGeometryCollection::VerticesGroup)
 		, BoneMapAttribute(InCollection,"BoneMap", FGeometryCollection::VerticesGroup)
 		, TransformToGeometryIndexAttribute(InCollection,"TransformToGeometryIndex", FTransformCollection::TransformGroup)
+		, ParentAttribute(InCollection, "Parent", FTransformCollection::TransformGroup)
 	{}
 
 	FBoundsFacade::FBoundsFacade(const FManagedArrayCollection& InCollection)
@@ -22,6 +24,7 @@ namespace GeometryCollection::Facades
 		, VertexAttribute(InCollection, "Vertex", FGeometryCollection::VerticesGroup, FGeometryCollection::VerticesGroup)
 		, BoneMapAttribute(InCollection, "BoneMap", FGeometryCollection::VerticesGroup)
 		, TransformToGeometryIndexAttribute(InCollection, "TransformToGeometryIndex", FTransformCollection::TransformGroup)
+		, ParentAttribute(InCollection, "Parent", FTransformCollection::TransformGroup)
 	{}
 
 	//
@@ -35,12 +38,14 @@ namespace GeometryCollection::Facades
 		VertexAttribute.Add();
 		BoneMapAttribute.Add();
 		TransformToGeometryIndexAttribute.Add();
+		ParentAttribute.Add();
 	}
 
 	bool FBoundsFacade::IsValid() const
 	{
 		return BoundingBoxAttribute.IsValid() && VertexAttribute.IsValid() 
-			&& BoneMapAttribute.IsValid() && TransformToGeometryIndexAttribute.IsValid();
+			&& BoneMapAttribute.IsValid() && TransformToGeometryIndexAttribute.IsValid()
+			&& ParentAttribute.IsValid();
 	}
 
 
@@ -74,6 +79,23 @@ namespace GeometryCollection::Facades
 			}
 		}
 	}
+
+	FBox FBoundsFacade::GetBoundingBox()
+	{
+		const TSet<int32>& RootIndices = Chaos::Facades::FCollectionHierarchyFacade::GetRootIndices(ParentAttribute);
+
+		const TManagedArray<FBox>& BoundingBoxArr = BoundingBoxAttribute.Get();
+
+		FBox BoundingBox;
+		BoundingBox.Init();
+
+		for (auto& Idx : RootIndices)
+		{
+			BoundingBox += BoundingBoxArr[Idx];
+		}
+
+		return BoundingBox;
+	}
 }; // GeometryCollection::Facades
 
 
diff --git a/Engine/Source/Runtime/Experimental/Chaos/Private/GeometryCollection/Facades/CollectionHierarchyFacade.cpp b/Engine/Source/Runtime/Experimental/Chaos/Private/GeometryCollection/Facades/CollectionHierarchyFacade.cpp
--- a/Engine/Source/Runtime/Experimental/Chaos/Private/GeometryCollection/Facades/CollectionHierarchyFacade.cpp
+++ b/Engine/Source/Runtime/Experimental/Chaos/Private/GeometryCollection/Facades/CollectionHierarchyFacade.cpp
@@ -102,4 +102,21 @@ namespace Chaos::Facades
 			}
 		}
 	}
+
+	TSet<int32> FCollectionHierarchyFacade::GetRootIndices(const TManagedArrayAccessor<int32>& ParentAttribute)
+	{
+		TSet<int32> Roots;
+		if (ParentAttribute.IsValid())
+		{
+			const TManagedArray<int32>& Parents = ParentAttribute.Get();
+			for (int32 Idx = 0; Idx < Parents.Num(); ++Idx)
+			{
+				if (Parents[Idx] == INDEX_NONE)
+				{
+					Roots.Add(Idx);
+				}
+			}
+		}
+		return Roots;
+	}
 }
\ No newline at end of file
diff --git a/Engine/Source/Runtime/Experimental/Chaos/Private/GeometryCollection/Facades/CollectionTransformFacade.cpp b/Engine/Source/Runtime/Experimental/Chaos/Private/GeometryCollection/Facades/CollectionTransformFacade.cpp
--- a/Engine/Source/Runtime/Experimental/Chaos/Private/GeometryCollection/Facades/CollectionTransformFacade.cpp
+++ b/Engine/Source/Runtime/Experimental/Chaos/Private/GeometryCollection/Facades/CollectionTransformFacade.cpp
@@ -2,6 +2,7 @@
 
 #include "GeometryCollection/Facades/CollectionTransformFacade.h"
 #include "GeometryCollection/TransformCollection.h"
+#include "GeometryCollection/Facades/CollectionHierarchyFacade.h"
 
 namespace GeometryCollection::Facades
 {
@@ -28,16 +29,7 @@ namespace GeometryCollection::Facades
 
 	TSet<int32> FTransformFacade::GetRootIndices() const
 	{
-		TSet<int32> Roots;
-		if (Parent.IsValid())
-		{
-			const TManagedArray<int32>& Parents = Parent.Get();
-			for (int i = 0; i < Parents.Num(); i++)
-			{
-				Roots.Add(i);
-			}
-		}
-		return Roots;
+		return Chaos::Facades::FCollectionHierarchyFacade::GetRootIndices(Parent);
 	}
 
 }
\ No newline at end of file
diff --git a/Engine/Source/Runtime/Experimental/Chaos/Public/GeometryCollection/Facades/CollectionBoundsFacade.h b/Engine/Source/Runtime/Experimental/Chaos/Public/GeometryCollection/Facades/CollectionBoundsFacade.h
--- a/Engine/Source/Runtime/Experimental/Chaos/Public/GeometryCollection/Facades/CollectionBoundsFacade.h
+++ b/Engine/Source/Runtime/Experimental/Chaos/Public/GeometryCollection/Facades/CollectionBoundsFacade.h
@@ -44,14 +44,18 @@ namespace GeometryCollection::Facades
 		/** BoundingBox access */
 		const TManagedArray< FBox >& GetBoundingBoxes() const { return BoundingBoxAttribute.Get(); }
 
+		/** BoundingBox for the whole collection */
+		FBox GetBoundingBox();
+
 		/** TransformToGeometryIndex Access */
 		const TManagedArray< int32 >& GetTransformToGeometryIndex() const { return TransformToGeometryIndexAttribute.Get(); }
 
 private:
-		TManagedArrayAccessor<FBox> BoundingBoxAttribute;
-		TManagedArrayAccessor<FVector3f> VertexAttribute;
-		TManagedArrayAccessor<int32> BoneMapAttribute;
-		TManagedArrayAccessor<int32> TransformToGeometryIndexAttribute;
+		TManagedArrayAccessor<FBox>			BoundingBoxAttribute;
+		TManagedArrayAccessor<FVector3f>	VertexAttribute;
+		TManagedArrayAccessor<int32>		BoneMapAttribute;
+		TManagedArrayAccessor<int32>		TransformToGeometryIndexAttribute;
+		TManagedArrayAccessor<int32>		ParentAttribute;
 	};
 
 }
\ No newline at end of file
diff --git a/Engine/Source/Runtime/Experimental/Chaos/Public/GeometryCollection/Facades/CollectionHierarchyFacade.h b/Engine/Source/Runtime/Experimental/Chaos/Public/GeometryCollection/Facades/CollectionHierarchyFacade.h
--- a/Engine/Source/Runtime/Experimental/Chaos/Public/GeometryCollection/Facades/CollectionHierarchyFacade.h
+++ b/Engine/Source/Runtime/Experimental/Chaos/Public/GeometryCollection/Facades/CollectionHierarchyFacade.h
@@ -47,6 +47,10 @@ namespace Chaos::Facades
 		*/ 
 		void GenerateLevelAttribute();
 
+	public:
+		/** Get the root indicies */
+		static TSet<int32> GetRootIndices(const TManagedArrayAccessor<int32>& ParentAttribute);
+
 	private:
 		TManagedArrayAccessor<int32>		ParentAttribute;
 		TManagedArrayAccessor<TSet<int32>>	ChildrenAttribute;
