Dataflow:
- Updated Selection nodes

diff --git a/Engine/Plugins/Experimental/GeometryCollectionPlugin/Source/GeometryCollectionNodes/Public/Dataflow/GeometryCollectionSelectionNodes.h b/Engine/Plugins/Experimental/GeometryCollectionPlugin/Source/GeometryCollectionNodes/Public/Dataflow/GeometryCollectionSelectionNodes.h
--- a/Engine/Plugins/Experimental/GeometryCollectionPlugin/Source/GeometryCollectionNodes/Public/Dataflow/GeometryCollectionSelectionNodes.h
+++ b/Engine/Plugins/Experimental/GeometryCollectionPlugin/Source/GeometryCollectionNodes/Public/Dataflow/GeometryCollectionSelectionNodes.h
@@ -26,7 +26,7 @@ struct FCollectionTransformSelectionAllDataflowNode : public FDataflowNode
 
 public:
 	/** GeometryCollection for the selection */
-	UPROPERTY(meta = (DataflowInput, DataflowOutput))
+	UPROPERTY(meta = (DataflowInput, DataflowOutput, DataflowPassthrough = "Collection", DataflowIntrinsic))
 	FManagedArrayCollection Collection;
 
 	/** Array of the selected bone indicies */
@@ -38,7 +38,7 @@ public:
 	{
 		RegisterInputConnection(&Collection);
 		RegisterOutputConnection(&TransformSelection);
-		RegisterOutputConnection(&Collection);
+		RegisterOutputConnection(&Collection, &Collection);
 	}
 
 	virtual void Evaluate(Dataflow::FContext& Context, const FDataflowOutput* Out) const override;
@@ -49,9 +49,12 @@ public:
 UENUM(BlueprintType)
 enum class ESetOperationEnum : uint8
 {
-	Dataflow_SetOperation_AND UMETA(DisplayName = "AND"),
-	Dataflow_SetOperation_OR UMETA(DisplayName = "OR"),
-	Dataflow_SetOperation_XOR UMETA(DisplayName = "XOR"),
+	/** Bitwise AND of incoming TransformSelections */
+	Dataflow_SetOperation_AND UMETA(DisplayName = "Intersect"),
+	/** Bitwise OR of incoming TransformSelections */
+	Dataflow_SetOperation_OR UMETA(DisplayName = "Union"),
+	/** Bitwise XOR of incoming TransformSelections */
+	Dataflow_SetOperation_XOR UMETA(DisplayName = "Symmetric Difference"),
 	//~~~
 	//256th entry
 	Dataflow_Max                UMETA(Hidden)
@@ -59,7 +62,7 @@ enum class ESetOperationEnum : uint8
 
 /**
  *
- * Runs boolean operation on incoming TransformSelections
+ * Runs boolean operation on TransformSelections
  *
  */
 USTRUCT(meta = (DataflowGeometryCollection))
@@ -74,15 +77,15 @@ public:
 	ESetOperationEnum Operation = ESetOperationEnum::Dataflow_SetOperation_AND;
 
 	/** Array of the selected bone indicies */
-	UPROPERTY(meta = (DataflowInput, DisplayName = "TransformSelectionA"))
+	UPROPERTY(meta = (DataflowInput, DisplayName = "TransformSelectionA", DataflowIntrinsic))
 	FDataflowTransformSelection TransformSelectionA;
 
 	/** Array of the selected bone indicies */
-	UPROPERTY(meta = (DataflowInput, DisplayName = "TransformSelectionB"))
+	UPROPERTY(meta = (DataflowInput, DisplayName = "TransformSelectionB", DataflowIntrinsic))
 	FDataflowTransformSelection TransformSelectionB;
 
 	/** Array of the selected bone indicies after operation*/
-	UPROPERTY(meta = (DataflowOutput, DisplayName = "TransformSelection"))
+	UPROPERTY(meta = (DataflowOutput, DisplayName = "TransformSelection", DataflowPassthrough = "TransformSelectionA"))
 	FDataflowTransformSelection TransformSelection;
 
 	FCollectionTransformSelectionSetOperationDataflowNode(const Dataflow::FNodeParameters& InParam, FGuid InGuid = FGuid::NewGuid())
@@ -90,7 +93,7 @@ public:
 	{
 		RegisterInputConnection(&TransformSelectionA);
 		RegisterInputConnection(&TransformSelectionB);
-		RegisterOutputConnection(&TransformSelection);
+		RegisterOutputConnection(&TransformSelection, &TransformSelectionA);
 	}
 
 	virtual void Evaluate(Dataflow::FContext& Context, const FDataflowOutput* Out) const override;
@@ -111,7 +114,7 @@ struct FCollectionTransformSelectionInfoDataflowNode : public FDataflowNode
 
 public:
 	/** Array of the selected bone indicies */
-	UPROPERTY(meta = (DataflowInput, DisplayName = "TransformSelection"))
+	UPROPERTY(meta = (DataflowInput, DisplayName = "TransformSelection", DataflowIntrinsic))
 	FDataflowTransformSelection TransformSelection;
 
 	/** GeometryCollection for the selection */
@@ -148,7 +151,7 @@ struct FCollectionTransformSelectionNoneDataflowNode : public FDataflowNode
 
 public:
 	/** GeometryCollection for the selection */
-	UPROPERTY(meta = (DataflowInput, DataflowOutput))
+	UPROPERTY(meta = (DataflowInput, DataflowOutput, DataflowPassthrough = "Collection", DataflowIntrinsic))
 	FManagedArrayCollection Collection;
 
 	/** Array of the selected bone indicies */
@@ -160,7 +163,7 @@ public:
 	{
 		RegisterInputConnection(&Collection);
 		RegisterOutputConnection(&TransformSelection);
-		RegisterOutputConnection(&Collection);
+		RegisterOutputConnection(&Collection, &Collection);
 	}
 
 	virtual void Evaluate(Dataflow::FContext& Context, const FDataflowOutput* Out) const override;
@@ -170,7 +173,7 @@ public:
 
 /**
  *
- * Inverts the incoming selection of bones
+ * Inverts selection of bones
  *
  */
 USTRUCT(meta = (DataflowGeometryCollection))
@@ -181,14 +184,14 @@ struct FCollectionTransformSelectionInvertDataflowNode : public FDataflowNode
 
 public:
 	/** Array of the selected bone indicies */
-	UPROPERTY(meta = (DataflowInput, DataflowOutput, DisplayName = "TransformSelection"))
+	UPROPERTY(meta = (DataflowInput, DataflowOutput, DisplayName = "TransformSelection", DataflowPassthrough = "TransformSelection", DataflowIntrinsic))
 	FDataflowTransformSelection TransformSelection;
 
 	FCollectionTransformSelectionInvertDataflowNode(const Dataflow::FNodeParameters& InParam, FGuid InGuid = FGuid::NewGuid())
 		: FDataflowNode(InParam, InGuid)
 	{
 		RegisterInputConnection(&TransformSelection);
-		RegisterOutputConnection(&TransformSelection);
+		RegisterOutputConnection(&TransformSelection, &TransformSelection);
 	}
 
 	virtual void Evaluate(Dataflow::FContext& Context, const FDataflowOutput* Out) const override;
@@ -221,7 +224,7 @@ public:
 	float RandomThreshold = 0.5f;
 
 	/** GeometryCollection for the selection */
-	UPROPERTY(meta = (DataflowInput, DataflowOutput))
+	UPROPERTY(meta = (DataflowInput, DataflowOutput, DataflowPassthrough = "Collection", DataflowIntrinsic))
 	FManagedArrayCollection Collection;
 
 	/** Array of the selected bone indicies */
@@ -236,7 +239,7 @@ public:
 		RegisterInputConnection(&RandomSeed);
 		RegisterInputConnection(&RandomThreshold);
 		RegisterOutputConnection(&TransformSelection);
-		RegisterOutputConnection(&Collection);
+		RegisterOutputConnection(&Collection, &Collection);
 	}
 
 	virtual void Evaluate(Dataflow::FContext& Context, const FDataflowOutput* Out) const override;
@@ -257,7 +260,7 @@ struct FCollectionTransformSelectionRootDataflowNode : public FDataflowNode
 
 public:
 	/** GeometryCollection for the selection */
-	UPROPERTY(meta = (DataflowInput, DataflowOutput))
+	UPROPERTY(meta = (DataflowInput, DataflowOutput, DataflowPassthrough = "Collection", DataflowIntrinsic))
 	FManagedArrayCollection Collection;
 
 	/** Array of the selected bone indicies */
@@ -269,7 +272,7 @@ public:
 	{
 		RegisterInputConnection(&Collection);
 		RegisterOutputConnection(&TransformSelection);
-		RegisterOutputConnection(&Collection);
+		RegisterOutputConnection(&Collection, &Collection);
 	}
 
 	virtual void Evaluate(Dataflow::FContext& Context, const FDataflowOutput* Out) const override;
@@ -352,7 +355,7 @@ public:
 
 /**
  *
- * Outputs the specified percentage of the incoming bone selection
+ * Outputs the specified percentage of the selected bones
  *
  */
 USTRUCT(meta = (DataflowGeometryCollection))
@@ -395,7 +398,7 @@ public:
 
 /**
  *
- * Selects the children of the incoming bone selection
+ * Selects the children of the selected bones
  *
  */
 USTRUCT(meta = (DataflowGeometryCollection))
@@ -429,7 +432,7 @@ public:
 
 /**
  *
- * Selects the siblings of the incoming bone selection
+ * Selects the siblings of the selected bones
  *
  */
 USTRUCT(meta = (DataflowGeometryCollection))
@@ -463,7 +466,7 @@ public:
 
 /**
  *
- * Selects the level of the incoming bone selection
+ * Selects the level of the selected bones
  *
  */
 USTRUCT(meta = (DataflowGeometryCollection))
@@ -497,7 +500,7 @@ public:
 
 /**
  *
- * Selects the contact(s) of the incoming bone selection
+ * Selects the contact(s) of the selected bones
  *
  */
 USTRUCT(meta = (DataflowGeometryCollection))
@@ -531,7 +534,7 @@ public:
 
 /**
  *
- * Selects the leaves in the GeometryCollection
+ * Selects the leaves in the Collection
  *
  */
 USTRUCT(meta = (DataflowGeometryCollection))
@@ -564,7 +567,7 @@ public:
 
 /**
  *
- * Selects the clusters in the GeometryCollection
+ * Selects the clusters in the Collection
  *
  */
 USTRUCT(meta = (DataflowGeometryCollection))
