Dataflow:
- Moved ExplodedView Dataflow node to GeometryCollectionFracturingNodes
- Moved exploded view compute code to FractureEngine library
- Updated ExplodedView Dataflow node to use the library
- Added CollectionExplodedVector facade
- Updated DataflowEngine rendering factory to use CollectionExplodedVector facade
- Updated GeometryCollection rendering to use Visibility attribute if it exists
- Updated VoronoiFracture, PlaneCutter node to use TransformSelection
- Updated CreateNonOverlappingConvexHulls node
- Added Internal FacesGroup attribute creation to VoronoiFracture, it will be added to all nodes which do fracture
- Added VertexSelectionCustom, FaceSelectionCustom, FaceSelectionInvert, BoolArrayToFaceSelection nodes
- Renamed and removed some Dataflow nodes
- Fixed issue with FBox and FSphere not initialized properly (UE-172071)

diff --git a/Engine/Plugins/Experimental/Dataflow/Source/DataflowEnginePlugin/Private/Dataflow/DataflowEngineRendering.cpp b/Engine/Plugins/Experimental/Dataflow/Source/DataflowEnginePlugin/Private/Dataflow/DataflowEngineRendering.cpp
--- a/Engine/Plugins/Experimental/Dataflow/Source/DataflowEnginePlugin/Private/Dataflow/DataflowEngineRendering.cpp
+++ b/Engine/Plugins/Experimental/Dataflow/Source/DataflowEnginePlugin/Private/Dataflow/DataflowEngineRendering.cpp
@@ -5,6 +5,7 @@
 #include "Dataflow/DataflowRenderingFactory.h"
 #include "GeometryCollection/GeometryCollectionUtility.h"
 #include "GeometryCollection/Facades/CollectionRenderingFacade.h"
+#include "GeometryCollection/Facades/CollectionExplodedVectorFacade.h"
 #include "GeometryCollection/ManagedArrayCollection.h"
 #include "GeometryCollection/GeometryCollectionAlgo.h"
 
@@ -38,42 +39,56 @@ namespace Dataflow
 						&& Collection.FindAttributeTyped<int32>("BoneMap", FGeometryCollection::VerticesGroup)
 						&& Collection.FindAttributeTyped<int32>(FTransformCollection::ParentAttribute, FTransformCollection::TransformGroup))
 					{
-						const TManagedArray<int32>& BoneIndex = Collection.GetAttribute<int32>("BoneMap", FGeometryCollection::VerticesGroup);
-						const TManagedArray<int32>& Parents = Collection.GetAttribute<int32>(FTransformCollection::ParentAttribute, FTransformCollection::TransformGroup);
-						const TManagedArray<FTransform>& Transforms = Collection.GetAttribute<FTransform>(FTransformCollection::TransformAttribute, FTransformCollection::TransformGroup);
+						if (Collection.NumElements(FTransformCollection::TransformGroup) > 0)
+						{
+							const TManagedArray<int32>& BoneIndex = Collection.GetAttribute<int32>("BoneMap", FGeometryCollection::VerticesGroup);
+							const TManagedArray<int32>& Parents = Collection.GetAttribute<int32>(FTransformCollection::ParentAttribute, FTransformCollection::TransformGroup);
+							const TManagedArray<FTransform>& Transforms = Collection.GetAttribute<FTransform>(FTransformCollection::TransformAttribute, FTransformCollection::TransformGroup);
 
-						TArray<FTransform> M;
-						GeometryCollectionAlgo::GlobalMatrices(Transforms, Parents, M);
+							TArray<FTransform> M;
+							GeometryCollectionAlgo::GlobalMatrices(Transforms, Parents, M);
 
-						auto ToD = [](FVector3f V) { return FVector3d(V.X, V.Y, V.Z); };
-						auto ToF = [](FVector3d V) { return FVector3f(V.X, V.Y, V.Z); };
+							// If Collection has "ExplodedVector" attribute then use it to modify the global matrices (ExplodedView node creates it)
+							GeometryCollection::Facades::FCollectionExplodedVectorFacade ExplodedViewFacade(Collection);
+							ExplodedViewFacade.UpdateGlobalMatricesWithExplodedVectors(M);
 
+							auto ToD = [](FVector3f V) { return FVector3d(V.X, V.Y, V.Z); };
+							auto ToF = [](FVector3d V) { return FVector3f(V.X, V.Y, V.Z); };
 
-						const TManagedArray<FVector3f>& Vertex = Collection.GetAttribute<FVector3f>("Vertex", FGeometryCollection::VerticesGroup);
-						const TManagedArray<FIntVector>& Faces = Collection.GetAttribute<FIntVector>("Indices", FGeometryCollection::FacesGroup);
 
-						TArray<FVector3f> Vertices; Vertices.AddUninitialized(Vertex.Num());
-						TArray<FIntVector> Tris; Tris.AddUninitialized(Faces.Num());
-						TArray<bool> Visited; Visited.Init(false, Vertices.Num());
+							const TManagedArray<FVector3f>& Vertex = Collection.GetAttribute<FVector3f>("Vertex", FGeometryCollection::VerticesGroup);
+							const TManagedArray<FIntVector>& Faces = Collection.GetAttribute<FIntVector>("Indices", FGeometryCollection::FacesGroup);
+							const TManagedArray<bool>* FaceVisible = Collection.FindAttribute<bool>("Visible", FGeometryCollection::FacesGroup);
 
-						int32 Tdx = 0;
-						for (const FIntVector& Face : Faces)
-						{
-							FIntVector Tri = FIntVector(Face[0], Face[1], Face[2]);
-							FTransform Ms[3] = { M[BoneIndex[Tri[0]]], M[BoneIndex[Tri[1]]], M[BoneIndex[Tri[2]]] };
+							TArray<FVector3f> Vertices; Vertices.AddUninitialized(Vertex.Num());
+							TArray<FIntVector> Tris; Tris.AddUninitialized(Faces.Num());
+							TArray<bool> Visited; Visited.Init(false, Vertices.Num());
 
-							Tris[Tdx++] = Tri;
-							if (!Visited[Tri[0]]) Vertices[Tri[0]] = ToF(Ms[0].TransformPosition(ToD(Vertex[Tri[0]])));
-							if (!Visited[Tri[1]]) Vertices[Tri[1]] = ToF(Ms[1].TransformPosition(ToD(Vertex[Tri[1]])));
-							if (!Visited[Tri[2]]) Vertices[Tri[2]] = ToF(Ms[2].TransformPosition(ToD(Vertex[Tri[2]])));
+							int32 Tdx = 0;
+							for (int32 FaceIdx = 0; FaceIdx < Faces.Num(); ++FaceIdx)
+							{				
+								if (FaceVisible && !(*FaceVisible)[FaceIdx]) continue;
 
-							Visited[Tri[0]] = true; Visited[Tri[1]] = true; Visited[Tri[2]] = true;
-						}
+								const FIntVector& Face = Faces[FaceIdx];
+
+								FIntVector Tri = FIntVector(Face[0], Face[1], Face[2]);
+								FTransform Ms[3] = { M[BoneIndex[Tri[0]]], M[BoneIndex[Tri[1]]], M[BoneIndex[Tri[2]]] };
 
-						// Maybe these buffers should be shrunk, but there are unused vertices in the buffer. 
-						for (int i = 0; i < Visited.Num(); i++) if (!Visited[i]) Vertices[i] = FVector3f(0);
+								Tris[Tdx++] = Tri;
+								if (!Visited[Tri[0]]) Vertices[Tri[0]] = ToF(Ms[0].TransformPosition(ToD(Vertex[Tri[0]])));
+								if (!Visited[Tri[1]]) Vertices[Tri[1]] = ToF(Ms[1].TransformPosition(ToD(Vertex[Tri[1]])));
+								if (!Visited[Tri[2]]) Vertices[Tri[2]] = ToF(Ms[2].TransformPosition(ToD(Vertex[Tri[2]])));
 
-						RenderCollection.AddSurface(MoveTemp(Vertices), MoveTemp(Tris));
+								Visited[Tri[0]] = true; Visited[Tri[1]] = true; Visited[Tri[2]] = true;
+							}
+
+							Tris.SetNum(Tdx);
+
+							// Maybe these buffers should be shrunk, but there are unused vertices in the buffer. 
+							for (int i = 0; i < Visited.Num(); i++) if (!Visited[i]) Vertices[i] = FVector3f(0);
+
+							RenderCollection.AddSurface(MoveTemp(Vertices), MoveTemp(Tris));
+						}
 					}
 				}
 			});
diff --git a/Engine/Plugins/Experimental/Fracture/Source/FractureEngine/FractureEngine.Build.cs b/Engine/Plugins/Experimental/Fracture/Source/FractureEngine/FractureEngine.Build.cs
--- a/Engine/Plugins/Experimental/Fracture/Source/FractureEngine/FractureEngine.Build.cs
+++ b/Engine/Plugins/Experimental/Fracture/Source/FractureEngine/FractureEngine.Build.cs
@@ -16,7 +16,8 @@ namespace UnrealBuildTool.Rules
 					"Chaos",
 					"Core",
 					"DataflowCore",
-					"PlanarCut"
+					"PlanarCut",
+					"Voronoi"
 				}
 				);
 
diff --git a/Engine/Plugins/Experimental/Fracture/Source/FractureEngine/Private/FractureEngineEdit.cpp b/Engine/Plugins/Experimental/Fracture/Source/FractureEngine/Private/FractureEngineEdit.cpp
--- a/Engine/Plugins/Experimental/Fracture/Source/FractureEngine/Private/FractureEngineEdit.cpp
+++ b/Engine/Plugins/Experimental/Fracture/Source/FractureEngine/Private/FractureEngineEdit.cpp
@@ -41,19 +41,19 @@ void FFractureEngineEdit::DeleteBranch(FGeometryCollection& GeometryCollection,
 }
 
 
-void FFractureEngineEdit::SetVisibilityInCollection(FManagedArrayCollection& InCollection, const TArray<int32>& InBoneSelection, bool bVisible)
+void FFractureEngineEdit::SetVisibilityInCollectionFromTransformSelection(FManagedArrayCollection& InCollection, const TArray<int32>& InTransformSelection, bool bVisible)
 {
 	GeometryCollection::Facades::FCollectionTransformSelectionFacade TransformSelectionFacade(InCollection);
 
-	const TManagedArray<int32>& TransformToGeometryIndex = InCollection.GetAttribute<int32>("TransformToGeometryIndex", FGeometryCollection::GeometryGroup);
+	const TManagedArray<int32>& TransformToGeometryIndex = InCollection.GetAttribute<int32>("TransformToGeometryIndex", FGeometryCollection::TransformGroup);
 	const TManagedArray<int32>& FaceStart = InCollection.GetAttribute<int32>("FaceStart", FGeometryCollection::GeometryGroup);
 	const TManagedArray<int32>& FaceCount = InCollection.GetAttribute<int32>("FaceCount", FGeometryCollection::GeometryGroup);
-	TManagedArray<bool>* Visible = InCollection.FindAttribute<bool>("Visible", FGeometryCollection::GeometryGroup);
+	TManagedArray<bool>& Visible = InCollection.ModifyAttribute<bool>("Visible", FGeometryCollection::FacesGroup);
 
-	TArray<int32> BoneIndicies = InBoneSelection;
-	TransformSelectionFacade.ConvertSelectionToRigidNodes(BoneIndicies);
+	TArray<int32> TransformIndicies = InTransformSelection;
+	TransformSelectionFacade.ConvertSelectionToRigidNodes(TransformIndicies);
 
-	for (int32 Index : BoneIndicies)
+	for (int32 Index : TransformIndicies)
 	{
 		// Iterate the faces in the geometry of this rigid node and set invisible.
 		if (TransformToGeometryIndex[Index] > INDEX_NONE)
@@ -61,13 +61,26 @@ void FFractureEngineEdit::SetVisibilityInCollection(FManagedArrayCollection& InC
 			int32 CurrFace = FaceStart[TransformToGeometryIndex[Index]];
 			for (int32 FaceOffset = 0; FaceOffset < FaceCount[TransformToGeometryIndex[Index]]; ++FaceOffset)
 			{
-				(*Visible)[CurrFace + FaceOffset] = bVisible;
+				Visible[CurrFace + FaceOffset] = bVisible;
 			}
 		}
 	}
 }
 
 
+void FFractureEngineEdit::SetVisibilityInCollectionFromFaceSelection(FManagedArrayCollection& InCollection, const TArray<int32>& InFaceSelection, bool bVisible)
+{
+	TManagedArray<bool>& Visible = InCollection.ModifyAttribute<bool>("Visible", FGeometryCollection::FacesGroup);
+
+	int32 NumFaces = InCollection.NumElements(FGeometryCollection::FacesGroup);
+
+	for (int32 FaceIdx : InFaceSelection)
+	{
+		Visible[FaceIdx] = bVisible;
+	}
+}
+
+
 void FFractureEngineEdit::Merge(FGeometryCollection& GeometryCollection, const TArray<int32>& InBoneSelection)
 {
 	const FManagedArrayCollection& InCollection = (const FManagedArrayCollection&)GeometryCollection;
diff --git a/Engine/Plugins/Experimental/Fracture/Source/FractureEngine/Private/FractureEngineEmbed.cpp b/Engine/Plugins/Experimental/Fracture/Source/FractureEngine/Private/FractureEngineEmbed.cpp
new file mode 100644
--- /dev/null
+++ b/Engine/Plugins/Experimental/Fracture/Source/FractureEngine/Private/FractureEngineEmbed.cpp
@@ -0,0 +1,9 @@
+// Copyright Epic Games, Inc. All Rights Reserved.
+
+#include "FractureEngineEmbed.h"
+#include "GeometryCollection/GeometryCollectionClusteringUtility.h"
+#include "GeometryCollection/GeometryCollectionAlgo.h"
+#include "GeometryCollection/GeometryCollectionProximityUtility.h"
+
+
+
diff --git a/Engine/Plugins/Experimental/Fracture/Source/FractureEngine/Private/FractureEngineFracturing.cpp b/Engine/Plugins/Experimental/Fracture/Source/FractureEngine/Private/FractureEngineFracturing.cpp
new file mode 100644
--- /dev/null
+++ b/Engine/Plugins/Experimental/Fracture/Source/FractureEngine/Private/FractureEngineFracturing.cpp
@@ -0,0 +1,338 @@
+// Copyright Epic Games, Inc. All Rights Reserved.
+
+#include "FractureEngineFracturing.h"
+#include "GeometryCollection/GeometryCollectionClusteringUtility.h"
+#include "GeometryCollection/GeometryCollectionAlgo.h"
+#include "GeometryCollection/GeometryCollectionProximityUtility.h"
+#include "GeometryCollection/GeometryCollection.h"
+#include "GeometryCollection/ManagedArrayCollection.h"
+#include "Voronoi/Voronoi.h"
+#include "PlanarCut.h"
+#include "GeometryCollection/Facades/CollectionMeshFacade.h"
+
+
+static void AddAdditionalAttributesIfRequired(FManagedArrayCollection& InOutCollection)
+{
+	FGeometryCollectionClusteringUtility::UpdateHierarchyLevelOfChildren(InOutCollection, -1);
+}
+
+static bool GetValidGeoCenter(FManagedArrayCollection& InOutCollection, 
+	const TManagedArray<int32>& TransformToGeometryIndex, 
+	const TArray<FTransform>& Transforms, 
+	const TManagedArray<int32>& Parents,
+	const TManagedArray<TSet<int32>>& Children, 
+	const TManagedArray<FBox>& BoundingBoxes, 
+	const TManagedArray<int32>& SimulationTypes,
+	int32 TransformIndex, 
+	FVector& OutGeoCenter)
+{
+
+	if (SimulationTypes[TransformIndex] == FGeometryCollection::ESimulationTypes::FST_Rigid)
+	{
+		OutGeoCenter = Transforms[TransformIndex].TransformPosition(BoundingBoxes[TransformToGeometryIndex[TransformIndex]].GetCenter());
+
+		return true;
+	}
+	else if (SimulationTypes[TransformIndex] == FGeometryCollection::ESimulationTypes::FST_None) // ie this is embedded geometry
+	{
+		int32 Parent = Parents[TransformIndex];
+		int32 ParentGeo = Parent != INDEX_NONE ? TransformToGeometryIndex[Parent] : INDEX_NONE;
+		if (ensureMsgf(ParentGeo != INDEX_NONE, TEXT("Embedded geometry should always have a rigid geometry parent!  Geometry collection may be malformed.")))
+		{
+			OutGeoCenter = Transforms[Parents[TransformIndex]].TransformPosition(BoundingBoxes[ParentGeo].GetCenter());
+		}
+		else
+		{
+			return false; // no valid value to return
+		}
+
+		return true;
+	}
+	else
+	{
+		FVector AverageCenter;
+		int32 ValidVectors = 0;
+		for (int32 ChildIndex : Children[TransformIndex])
+		{
+
+			if (GetValidGeoCenter(InOutCollection, TransformToGeometryIndex, Transforms, Parents, Children, BoundingBoxes, SimulationTypes, ChildIndex, OutGeoCenter))
+			{
+				if (ValidVectors == 0)
+				{
+					AverageCenter = OutGeoCenter;
+				}
+				else
+				{
+					AverageCenter += OutGeoCenter;
+				}
+				++ValidVectors;
+			}
+		}
+
+		if (ValidVectors > 0)
+		{
+			OutGeoCenter = AverageCenter / ValidVectors;
+			return true;
+		}
+	}
+
+	return false;
+}
+
+//
+// TODO: Rewrite this using facades
+//
+void FFractureEngineFracturing::GenerateExplodedViewAttribute(FManagedArrayCollection& InOutCollection, const FVector& InScale, float InUniformScale)
+{
+	// Check if InOutCollection is not empty
+	if (InOutCollection.HasAttribute("Transform", FGeometryCollection::TransformGroup))
+	{
+		InOutCollection.AddAttribute<FVector3f>("ExplodedVector", FGeometryCollection::TransformGroup, FManagedArrayCollection::FConstructionParameters(FName(), false));
+		check(InOutCollection.HasAttribute("ExplodedVector", FGeometryCollection::TransformGroup));
+
+		TManagedArray<FVector3f>& ExplodedVectors = InOutCollection.ModifyAttribute<FVector3f>("ExplodedVector", FGeometryCollection::TransformGroup);
+		const TManagedArray<FTransform>& Transforms = InOutCollection.GetAttribute<FTransform>("Transform", FGeometryCollection::TransformGroup);
+		const TManagedArray<int32>& TransformToGeometryIndices = InOutCollection.GetAttribute<int32>("TransformToGeometryIndex", FGeometryCollection::TransformGroup);
+		const TManagedArray<FBox>& BoundingBoxes = InOutCollection.GetAttribute<FBox>("BoundingBox", FGeometryCollection::GeometryGroup);
+
+		// Make sure we have valid "Level"
+		AddAdditionalAttributesIfRequired(InOutCollection);
+
+		const TManagedArray<int32>& Levels = InOutCollection.GetAttribute<int32>("Level", FTransformCollection::TransformGroup);
+		const TManagedArray<int32>& Parents = InOutCollection.GetAttribute<int32>("Parent", FTransformCollection::TransformGroup);
+		const TManagedArray<TSet<int32>>& Children = InOutCollection.GetAttribute<TSet<int32>>("Children", FGeometryCollection::TransformGroup);
+		const TManagedArray<int32>& SimulationTypes = InOutCollection.GetAttribute<int32>("SimulationType", FGeometryCollection::TransformGroup);
+
+		int32 ViewFractureLevel = -1;
+		int32 MaxFractureLevel = ViewFractureLevel;
+		for (int32 Idx = 0, ni = Transforms.Num(); Idx < ni; ++Idx)
+		{
+			if (Levels[Idx] > MaxFractureLevel)
+				MaxFractureLevel = Levels[Idx];
+		}
+
+		TArray<FTransform> TransformArr;
+		GeometryCollectionAlgo::GlobalMatrices(Transforms, Parents, TransformArr);
+
+		TArray<FVector> TransformedCenters;
+		TransformedCenters.SetNumUninitialized(TransformArr.Num());
+
+		int32 TransformsCount = 0;
+
+		FVector Center(ForceInitToZero);
+		for (int32 Idx = 0, ni = Transforms.Num(); Idx < ni; ++Idx)
+		{
+			ExplodedVectors[Idx] = FVector3f::ZeroVector;
+			FVector GeoCenter;
+
+			if (GetValidGeoCenter(InOutCollection, TransformToGeometryIndices, TransformArr, Parents, Children, BoundingBoxes, SimulationTypes, Idx, GeoCenter))
+			{
+				TransformedCenters[Idx] = GeoCenter;
+				if ((ViewFractureLevel < 0) || Levels[Idx] == ViewFractureLevel)
+				{
+					Center += TransformedCenters[Idx];
+					++TransformsCount;
+				}
+			}
+		}
+
+		Center /= TransformsCount;
+
+		for (int Level = 1; Level <= MaxFractureLevel; Level++)
+		{
+			for (int32 Idx = 0, ni = TransformArr.Num(); Idx < ni; ++Idx)
+			{
+				if ((ViewFractureLevel < 0) || Levels[Idx] == ViewFractureLevel)
+				{
+					FVector ScaleVec = InScale * InUniformScale;
+					ExplodedVectors[Idx] = (FVector3f)(TransformedCenters[Idx] - Center) * (FVector3f)ScaleVec;
+				}
+				else
+				{
+					if (Parents[Idx] > -1)
+					{
+						ExplodedVectors[Idx] = ExplodedVectors[Parents[Idx]];
+					}
+				}
+			}
+		}
+	}
+}
+
+
+static float GetMaxVertexMovement(float Grout, float Amplitude, int OctaveNumber, float Persistence)
+{
+	float MaxDisp = Grout;
+	float AmplitudeScaled = Amplitude;
+
+	for (int32 OctaveIdx = 0; OctaveIdx < OctaveNumber; OctaveIdx++, AmplitudeScaled *= Persistence)
+	{
+		MaxDisp += FMath::Abs(AmplitudeScaled);
+	}
+
+	return MaxDisp;
+}
+
+
+void FFractureEngineFracturing::VoronoiFracture(FManagedArrayCollection& InOutCollection,
+	const FDataflowTransformSelection& InTransformSelection,
+	const TArray<FVector>& InSites,
+	float InRandomSeed,
+	float InChanceToFracture,
+	bool InGroupFracture,
+	float InGrout,
+	float InAmplitude,
+	float InFrequency,
+	float InPersistence,
+	float InLacunarity,
+	int32 InOctaveNumber,
+	float InPointSpacing,
+	bool InAddSamplesForCollision,
+	float InCollisionSampleSpacing)
+{
+	if (TUniquePtr<FGeometryCollection> GeomCollection = TUniquePtr<FGeometryCollection>(InOutCollection.NewCopy<FGeometryCollection>()))
+	{
+		if (InSites.Num() > 0)
+		{
+			//
+			// Compute BoundingBox for InCollection
+			//
+			FBox BoundingBox(ForceInit);
+
+			if (InOutCollection.HasAttribute("Transform", FGeometryCollection::TransformGroup) &&
+				InOutCollection.HasAttribute("Parent", FGeometryCollection::TransformGroup) &&
+				InOutCollection.HasAttribute("TransformIndex", FGeometryCollection::GeometryGroup) &&
+				InOutCollection.HasAttribute("BoundingBox", FGeometryCollection::GeometryGroup))
+			{
+				const TManagedArray<FTransform>& Transforms = InOutCollection.GetAttribute<FTransform>("Transform", FGeometryCollection::TransformGroup);
+				const TManagedArray<int32>& ParentIndices = InOutCollection.GetAttribute<int32>("Parent", FGeometryCollection::TransformGroup);
+				const TManagedArray<int32>& TransformIndices = InOutCollection.GetAttribute<int32>("TransformIndex", FGeometryCollection::GeometryGroup);
+				const TManagedArray<FBox>& BoundingBoxes = InOutCollection.GetAttribute<FBox>("BoundingBox", FGeometryCollection::GeometryGroup);
+
+				TArray<FMatrix> TmpGlobalMatrices;
+				GeometryCollectionAlgo::GlobalMatrices(Transforms, ParentIndices, TmpGlobalMatrices);
+
+				if (TmpGlobalMatrices.Num() > 0)
+				{
+					for (int32 BoxIdx = 0; BoxIdx < BoundingBoxes.Num(); ++BoxIdx)
+					{
+						const int32 TransformIndex = TransformIndices[BoxIdx];
+						BoundingBox += BoundingBoxes[BoxIdx].TransformBy(TmpGlobalMatrices[TransformIndex]);
+					}
+				}
+
+				//
+				// Compute Voronoi Bounds
+				//
+				FBox VoronoiBounds = BoundingBox;
+				VoronoiBounds += FBox(InSites);
+
+				VoronoiBounds = VoronoiBounds.ExpandBy(GetMaxVertexMovement(InGrout, InAmplitude, InOctaveNumber, InPersistence) + KINDA_SMALL_NUMBER);
+
+				//
+				// Voronoi Fracture
+				//
+				FNoiseSettings NoiseSettings;
+				NoiseSettings.Amplitude = InAmplitude;
+				NoiseSettings.Frequency = InFrequency;
+				NoiseSettings.Octaves = InOctaveNumber;
+				NoiseSettings.PointSpacing = InPointSpacing;
+				NoiseSettings.Lacunarity = InLacunarity;
+				NoiseSettings.Persistence = InPersistence;
+
+				FVoronoiDiagram Voronoi(InSites, VoronoiBounds, .1f);
+				
+				FPlanarCells VoronoiPlanarCells = FPlanarCells(InSites, Voronoi);
+				VoronoiPlanarCells.InternalSurfaceMaterials.NoiseSettings = NoiseSettings;
+
+				TArray<int32> TransformSelectionArr = InTransformSelection.AsArray();
+				
+				int ResultGeometryIndex = CutMultipleWithPlanarCells(VoronoiPlanarCells, *GeomCollection, TransformSelectionArr, InGrout, InCollisionSampleSpacing, InRandomSeed, FTransform().Identity);
+
+				{
+					// Add FacesGroup::Internal<bool> attribute to tag the newly created inside faces
+					const TManagedArray<int32>& MaterialIDs = GeomCollection->GetAttribute<int32>("MaterialID", FGeometryCollection::FacesGroup);
+
+					GeometryCollection::Facades::FCollectionMeshFacade::AddInternalAttribute(*GeomCollection, MaterialIDs.GetConstArray());
+				}
+
+				InOutCollection = (const FManagedArrayCollection&)(*GeomCollection);
+			}
+		}
+	}
+}
+
+
+void FFractureEngineFracturing::PlaneCutter(FManagedArrayCollection& InOutCollection,
+	const FDataflowTransformSelection& InTransformSelection,
+	const FBox& InBoundingBox,
+	int32 InNumPlanes,
+	float InRandomSeed,
+	float InGrout,
+	float InAmplitude,
+	float InFrequency,
+	float InPersistence,
+	float InLacunarity,
+	int32 InOctaveNumber,
+	float InPointSpacing,
+	bool InAddSamplesForCollision,
+	float InCollisionSampleSpacing)
+{
+
+	if (TUniquePtr<FGeometryCollection> GeomCollection = TUniquePtr<FGeometryCollection>(InOutCollection.NewCopy<FGeometryCollection>()))
+	{
+		TArray<FPlane> CuttingPlanes;
+		TArray<FTransform> CuttingPlaneTransforms;
+
+		FRandomStream RandStream(InRandomSeed);
+
+		FBox Bounds = InBoundingBox;
+		const FVector Extent(Bounds.Max - Bounds.Min);
+
+		CuttingPlaneTransforms.Reserve(CuttingPlaneTransforms.Num() + InNumPlanes);
+		for (int32 ii = 0; ii < InNumPlanes; ++ii)
+		{
+			FVector Position(Bounds.Min + FVector(RandStream.FRand(), RandStream.FRand(), RandStream.FRand()) * Extent);
+			CuttingPlaneTransforms.Emplace(FTransform(FRotator(RandStream.FRand() * 360.0f, RandStream.FRand() * 360.0f, 0.0f), Position));
+		}
+
+		for (const FTransform& Transform : CuttingPlaneTransforms)
+		{
+			CuttingPlanes.Add(FPlane(Transform.GetLocation(), Transform.GetUnitAxis(EAxis::Z)));
+		}
+
+		FInternalSurfaceMaterials InternalSurfaceMaterials;
+		FNoiseSettings NoiseSettings;
+
+		if (InAmplitude > 0.f)
+		{
+			NoiseSettings.Amplitude = InAmplitude;
+			NoiseSettings.Frequency = InFrequency;
+			NoiseSettings.Lacunarity = InLacunarity;
+			NoiseSettings.Persistence = InPersistence;
+			NoiseSettings.Octaves = InOctaveNumber;
+			NoiseSettings.PointSpacing = InPointSpacing;
+
+			InternalSurfaceMaterials.NoiseSettings = NoiseSettings;
+		}
+
+		float CollisionSampleSpacingVal = InCollisionSampleSpacing;
+		float GroutVal = InGrout;
+
+		TArray<int32> TransformSelectionArr = InTransformSelection.AsArray();
+
+		int ResultGeometryIndex = CutMultipleWithMultiplePlanes(CuttingPlanes, InternalSurfaceMaterials, *GeomCollection, TransformSelectionArr, GroutVal, CollisionSampleSpacingVal, InRandomSeed, FTransform().Identity);
+
+		{
+			// Add FacesGroup::Internal<bool> attribute to tag the newly created inside faces
+			const TManagedArray<int32>& MaterialIDs = GeomCollection->GetAttribute<int32>("MaterialID", FGeometryCollection::FacesGroup);
+
+			GeometryCollection::Facades::FCollectionMeshFacade::AddInternalAttribute(*GeomCollection, MaterialIDs.GetConstArray());
+		}
+
+		InOutCollection = (const FManagedArrayCollection&)(*GeomCollection);
+	}
+}
+
+
+
diff --git a/Engine/Plugins/Experimental/Fracture/Source/FractureEngine/Private/FractureEngineUtility.cpp b/Engine/Plugins/Experimental/Fracture/Source/FractureEngine/Private/FractureEngineUtility.cpp
new file mode 100644
--- /dev/null
+++ b/Engine/Plugins/Experimental/Fracture/Source/FractureEngine/Private/FractureEngineUtility.cpp
@@ -0,0 +1,10 @@
+// Copyright Epic Games, Inc. All Rights Reserved.
+
+#include "FractureEngineUtility.h"
+#include "GeometryCollection/GeometryCollectionClusteringUtility.h"
+#include "GeometryCollection/GeometryCollectionAlgo.h"
+#include "GeometryCollection/GeometryCollectionProximityUtility.h"
+
+
+
+
diff --git a/Engine/Plugins/Experimental/Fracture/Source/FractureEngine/Public/FractureEngineEdit.h b/Engine/Plugins/Experimental/Fracture/Source/FractureEngine/Public/FractureEngineEdit.h
--- a/Engine/Plugins/Experimental/Fracture/Source/FractureEngine/Public/FractureEngineEdit.h
+++ b/Engine/Plugins/Experimental/Fracture/Source/FractureEngine/Public/FractureEngineEdit.h
@@ -11,7 +11,9 @@ public:
 
 	static void DeleteBranch(FGeometryCollection& GeometryCollection, const TArray<int32>& InBoneSelection);
 
-	static void SetVisibilityInCollection(FManagedArrayCollection& InCollection, const TArray<int32>& InBoneSelection, bool bVisible);
+	static void SetVisibilityInCollectionFromTransformSelection(FManagedArrayCollection& InCollection, const TArray<int32>& InTransformSelection, bool bVisible);
+
+	static void SetVisibilityInCollectionFromFaceSelection(FManagedArrayCollection& InCollection, const TArray<int32>& InFaceSelection, bool bVisible);
 
 	static void Merge(FGeometryCollection& GeometryCollection, const TArray<int32>& InBoneSelection);
 
diff --git a/Engine/Plugins/Experimental/Fracture/Source/FractureEngine/Public/FractureEngineEmbed.h b/Engine/Plugins/Experimental/Fracture/Source/FractureEngine/Public/FractureEngineEmbed.h
new file mode 100644
--- /dev/null
+++ b/Engine/Plugins/Experimental/Fracture/Source/FractureEngine/Public/FractureEngineEmbed.h
@@ -0,0 +1,13 @@
+// Copyright Epic Games, Inc. All Rights Reserved.
+
+#pragma once
+
+#include "Algo/Count.h"
+#include "GeometryCollection/GeometryCollection.h"
+
+
+class FRACTUREENGINE_API FFractureEngineEmbed
+{
+public:
+
+};
\ No newline at end of file
diff --git a/Engine/Plugins/Experimental/Fracture/Source/FractureEngine/Public/FractureEngineFracturing.h b/Engine/Plugins/Experimental/Fracture/Source/FractureEngine/Public/FractureEngineFracturing.h
new file mode 100644
--- /dev/null
+++ b/Engine/Plugins/Experimental/Fracture/Source/FractureEngine/Public/FractureEngineFracturing.h
@@ -0,0 +1,50 @@
+// Copyright Epic Games, Inc. All Rights Reserved.
+
+#pragma once
+
+#include "Algo/Count.h"
+#include "GeometryCollection/GeometryCollection.h"
+#include "GeometryCollection/ManagedArrayCollection.h"
+#include "Dataflow/DataflowSelection.h"
+#include "CoreMinimal.h"
+#include "Dataflow/DataflowSelection.h"
+
+
+class FRACTUREENGINE_API FFractureEngineFracturing
+{
+public:
+
+	static void GenerateExplodedViewAttribute(FManagedArrayCollection& InOutCollection, const FVector& InScale, float InUniformScale);
+
+	static void VoronoiFracture(FManagedArrayCollection& InOutCollection,
+		const FDataflowTransformSelection& InTransformSelection,
+		const TArray<FVector>& InSites,
+		float InRandomSeed,
+		float InChanceToFracture,
+		bool InGroupFracture,
+		float InGrout,
+		float InAmplitude,
+		float InFrequency,
+		float InPersistence,
+		float InLacunarity,
+		int32 InOctaveNumber,
+		float InPointSpacing,
+		bool InAddSamplesForCollision,
+		float InCollisionSampleSpacing);
+
+	static void PlaneCutter(FManagedArrayCollection& InOutCollection,
+		const FDataflowTransformSelection& InTransformSelection,
+		const FBox& InBoundingBox,
+		int32 InNumPlanes,
+		float InRandomSeed,
+		float InGrout,
+		float InAmplitude,
+		float InFrequency,
+		float InPersistence,
+		float InLacunarity,
+		int32 InOctaveNumber,
+		float InPointSpacing,
+		bool InAddSamplesForCollision,
+		float InCollisionSampleSpacing);
+
+};
\ No newline at end of file
diff --git a/Engine/Plugins/Experimental/Fracture/Source/FractureEngine/Public/FractureEngineUtility.h b/Engine/Plugins/Experimental/Fracture/Source/FractureEngine/Public/FractureEngineUtility.h
new file mode 100644
--- /dev/null
+++ b/Engine/Plugins/Experimental/Fracture/Source/FractureEngine/Public/FractureEngineUtility.h
@@ -0,0 +1,13 @@
+// Copyright Epic Games, Inc. All Rights Reserved.
+
+#pragma once
+
+#include "Algo/Count.h"
+#include "GeometryCollection/GeometryCollection.h"
+
+
+class FRACTUREENGINE_API FFractureEngineUtility
+{
+public:
+
+};
\ No newline at end of file
diff --git a/Engine/Plugins/Experimental/GeometryCollectionPlugin/Source/GeometryCollectionNodes/Private/Dataflow/GeometryCollectionEditNodes.cpp b/Engine/Plugins/Experimental/GeometryCollectionPlugin/Source/GeometryCollectionNodes/Private/Dataflow/GeometryCollectionEditNodes.cpp
--- a/Engine/Plugins/Experimental/GeometryCollectionPlugin/Source/GeometryCollectionNodes/Private/Dataflow/GeometryCollectionEditNodes.cpp
+++ b/Engine/Plugins/Experimental/GeometryCollectionPlugin/Source/GeometryCollectionNodes/Private/Dataflow/GeometryCollectionEditNodes.cpp
@@ -87,30 +87,67 @@ void FSetVisibilityInCollectionDataflowNode::Evaluate(Dataflow::FContext& Contex
 {
 	if (Out->IsA<FManagedArrayCollection>(&Collection))
 	{
-		FManagedArrayCollection InCollection = GetValue<FManagedArrayCollection>(Context, &Collection);
 		const FDataflowTransformSelection& InTransformSelection = GetValue<FDataflowTransformSelection>(Context, &TransformSelection);
+		const FDataflowFaceSelection& InFaceSelection = GetValue<FDataflowFaceSelection>(Context, &FaceSelection);
 
-		const int32 NumTransforms = InCollection.NumElements(FGeometryCollection::TransformGroup);
-		if (InTransformSelection.Num() == NumTransforms)
+		if (IsConnected<FDataflowTransformSelection>(&TransformSelection))
 		{
-			if (InTransformSelection.AnySelected())
+			FManagedArrayCollection InCollection = GetValue<FManagedArrayCollection>(Context, &Collection);
+			const int32 NumTransforms = InCollection.NumElements(FGeometryCollection::TransformGroup);
+
+			if (InTransformSelection.Num() == NumTransforms)
+			{
+				if (InTransformSelection.AnySelected())
+				{
+					TArray<int32> BoneIndices;
+					InTransformSelection.AsArray(BoneIndices);
+
+					FFractureEngineEdit::SetVisibilityInCollectionFromTransformSelection(InCollection, BoneIndices, Visibility == EVisibiltyOptionsEnum::Dataflow_VisibilityOptions_Visible);
+
+					SetValue<FManagedArrayCollection>(Context, InCollection, &Collection);
+					return;
+				}
+			}
+			else
 			{
-				TArray<int32> BoneIndices;
-				InTransformSelection.AsArray(BoneIndices);
+				// ERROR: InTransformSelection's length doesn't much number of transforms in InCollection
+				FString ErrorStr = "TransformSelection's number of elements doesn't match number of transforms in Collection.";
+				UE_LOG(LogTemp, Error, TEXT("[Dataflow ERROR] %s"), *ErrorStr);
+			}
+		}
+		else if (IsConnected<FDataflowFaceSelection>(&FaceSelection))
+		{
+			FManagedArrayCollection InCollection = GetValue<FManagedArrayCollection>(Context, &Collection);
+			const int32 NumFaces = InCollection.NumElements(FGeometryCollection::FacesGroup);
+
+			if (InFaceSelection.Num() == NumFaces)
+			{
+				if (InFaceSelection.AnySelected())
+				{
+					TArray<int32> FaceIndexArr;
+					InFaceSelection.AsArray(FaceIndexArr);
 
-				FFractureEngineEdit::SetVisibilityInCollection(InCollection, BoneIndices, Visibility == EVisibiltyOptionsEnum::Dataflow_VisibilityOptions_Visible);
+					FFractureEngineEdit::SetVisibilityInCollectionFromFaceSelection(InCollection, FaceIndexArr, Visibility == EVisibiltyOptionsEnum::Dataflow_VisibilityOptions_Visible);
 
-				SetValue<FManagedArrayCollection>(Context, InCollection, &Collection);
-				return;
+					SetValue<FManagedArrayCollection>(Context, InCollection, &Collection);
+					return;
+				}
+			}
+			else
+			{
+				// ERROR: InTransformSelection's length doesn't much number of transforms in InCollection
+				FString ErrorStr = "TransformSelection's number of elements doesn't match number of transforms in Collection.";
+				UE_LOG(LogTemp, Error, TEXT("[Dataflow ERROR] %s"), *ErrorStr);
 			}
 		}
 		else
 		{
 			// ERROR: InTransformSelection's length doesn't much number of transforms in InCollection
-			FString ErrorStr = "TransformSelection's number of elements doesn't match number of transforms in Collection.";
+			FString ErrorStr = "A TransformSelection or a FaceSelection must be connected as input.";
 			UE_LOG(LogTemp, Error, TEXT("[Dataflow ERROR] %s"), *ErrorStr);
 		}
 
+		const FManagedArrayCollection& InCollection = GetValue<FManagedArrayCollection>(Context, &Collection);
 		SetValue<FManagedArrayCollection>(Context, InCollection, &Collection);
 	}
 }
diff --git a/Engine/Plugins/Experimental/GeometryCollectionPlugin/Source/GeometryCollectionNodes/Private/Dataflow/GeometryCollectionEmbedNodes.cpp b/Engine/Plugins/Experimental/GeometryCollectionPlugin/Source/GeometryCollectionNodes/Private/Dataflow/GeometryCollectionEmbedNodes.cpp
new file mode 100644
--- /dev/null
+++ b/Engine/Plugins/Experimental/GeometryCollectionPlugin/Source/GeometryCollectionNodes/Private/Dataflow/GeometryCollectionEmbedNodes.cpp
@@ -0,0 +1,19 @@
+// Copyright Epic Games, Inc. All Rights Reserved.
+
+#include "Dataflow/GeometryCollectionEmbedNodes.h"
+#include "Dataflow/DataflowCore.h"
+
+
+//#include UE_INLINE_GENERATED_CPP_BY_NAME(GeometryCollectionEmbedNodes)
+
+namespace Dataflow
+{
+
+	void GeometryCollectionEmbedNodes()
+	{
+		static const FLinearColor CDefaultNodeBodyTintColor = FLinearColor(0.f, 0.f, 0.f, 0.5f);
+
+	}
+}
+
+
diff --git a/Engine/Plugins/Experimental/GeometryCollectionPlugin/Source/GeometryCollectionNodes/Private/Dataflow/GeometryCollectionFracturingNodes.cpp b/Engine/Plugins/Experimental/GeometryCollectionPlugin/Source/GeometryCollectionNodes/Private/Dataflow/GeometryCollectionFracturingNodes.cpp
--- a/Engine/Plugins/Experimental/GeometryCollectionPlugin/Source/GeometryCollectionNodes/Private/Dataflow/GeometryCollectionFracturingNodes.cpp
+++ b/Engine/Plugins/Experimental/GeometryCollectionPlugin/Source/GeometryCollectionNodes/Private/Dataflow/GeometryCollectionFracturingNodes.cpp
@@ -27,6 +27,7 @@
 #include "GeometryCollection/GeometryCollectionProximityUtility.h"
 #include "FractureEngineClustering.h"
 #include "FractureEngineSelection.h"
+#include "FractureEngineFracturing.h"
 
 #include UE_INLINE_GENERATED_CPP_BY_NAME(GeometryCollectionFracturingNodes)
 
@@ -41,6 +42,7 @@ namespace Dataflow
 		DATAFLOW_NODE_REGISTER_CREATION_FACTORY(FRadialScatterPointsDataflowNode);
 		DATAFLOW_NODE_REGISTER_CREATION_FACTORY(FVoronoiFractureDataflowNode);
 		DATAFLOW_NODE_REGISTER_CREATION_FACTORY(FPlaneCutterDataflowNode);
+		DATAFLOW_NODE_REGISTER_CREATION_FACTORY(FExplodedViewDataflowNode);
 
 		// GeometryCollection|Fracture
 		DATAFLOW_NODE_REGISTER_CREATION_FACTORY_NODE_COLORS_BY_CATEGORY("GeometryCollection|Fracture", FLinearColor(1.f, 1.f, 0.8f), CDefaultNodeBodyTintColor);
@@ -107,94 +109,42 @@ void FRadialScatterPointsDataflowNode::Evaluate(Dataflow::FContext& Context, con
 }
 
 
-static float GetMaxVertexMovement(float Grout, float Amplitude, int OctaveNumber, float Persistence)
-{
-	float MaxDisp = Grout;
-	float AmplitudeScaled = Amplitude;
-	for (int32 OctaveIdx = 0; OctaveIdx < OctaveNumber; OctaveIdx++, AmplitudeScaled *= Persistence)
-	{
-		MaxDisp += FMath::Abs(AmplitudeScaled);
-	}
-	return MaxDisp;
-}
-
 void FVoronoiFractureDataflowNode::Evaluate(Dataflow::FContext& Context, const FDataflowOutput* Out) const
 {
 	if (Out->IsA<FManagedArrayCollection>(&Collection))
 	{
-		const FManagedArrayCollection& InCollection = GetValue<FManagedArrayCollection>(Context, &Collection);
-		if (TUniquePtr<FGeometryCollection> GeomCollection = TUniquePtr<FGeometryCollection>(InCollection.NewCopy<FGeometryCollection>()))
+		const FDataflowTransformSelection& InTransformSelection = GetValue<FDataflowTransformSelection>(Context, &TransformSelection);
+
+		if (IsConnected<FDataflowTransformSelection>(&TransformSelection))
 		{
-			const TArray<FVector>& Sites = GetValue<TArray<FVector>>(Context, &Points);
-			if (Sites.Num() > 0)
+			if (InTransformSelection.AnySelected())
 			{
-				//
-				// Compute BoundingBox for ManagedArrayIn
-				//
-				FBox BoundingBox(ForceInit);
-
-				if (InCollection.HasAttribute("Transform", FGeometryCollection::TransformGroup) &&
-					InCollection.HasAttribute("Parent", FGeometryCollection::TransformGroup) &&
-					InCollection.HasAttribute("TransformIndex", FGeometryCollection::GeometryGroup) &&
-					InCollection.HasAttribute("BoundingBox", FGeometryCollection::GeometryGroup))
-				{
-					const TManagedArray<FTransform>& Transforms = InCollection.GetAttribute<FTransform>("Transform", FGeometryCollection::TransformGroup);
-					const TManagedArray<int32>& ParentIndices = InCollection.GetAttribute<int32>("Parent", FGeometryCollection::TransformGroup);
-					const TManagedArray<int32>& TransformIndices = InCollection.GetAttribute<int32>("TransformIndex", FGeometryCollection::GeometryGroup);
-					const TManagedArray<FBox>& BoundingBoxes = InCollection.GetAttribute<FBox>("BoundingBox", FGeometryCollection::GeometryGroup);
-
-					TArray<FMatrix> TmpGlobalMatrices;
-					GeometryCollectionAlgo::GlobalMatrices(Transforms, ParentIndices, TmpGlobalMatrices);
-
-					if (TmpGlobalMatrices.Num() > 0)
-					{
-						for (int32 BoxIdx = 0; BoxIdx < BoundingBoxes.Num(); ++BoxIdx)
-						{
-							const int32 TransformIndex = TransformIndices[BoxIdx];
-							BoundingBox += BoundingBoxes[BoxIdx].TransformBy(TmpGlobalMatrices[TransformIndex]);
-						}
-					}
-
-					//
-					// Compute Voronoi Bounds
-					//
-					FBox VoronoiBounds = BoundingBox;
-					VoronoiBounds += FBox(Sites);
-
-					float GroutVal = GetValue<float>(Context, &Grout);
-					float AmplitudeVal = GetValue<float>(Context, &Amplitude);
-					int32 OctaveNumberVal = GetValue<int32>(Context, &OctaveNumber);
-					float PersistenceVal = GetValue<float>(Context, &Persistence);
-
-					VoronoiBounds = VoronoiBounds.ExpandBy(GetMaxVertexMovement(GroutVal, AmplitudeVal, OctaveNumberVal, PersistenceVal) + KINDA_SMALL_NUMBER);
-
-					//
-					// Voronoi Fracture
-					//
-					FNoiseSettings NoiseSettings;
-					NoiseSettings.Amplitude = AmplitudeVal;
-					NoiseSettings.Frequency = GetValue<float>(Context, &Frequency);
-					NoiseSettings.Octaves = OctaveNumberVal;
-					NoiseSettings.PointSpacing = GetValue<float>(Context, &PointSpacing);
-					NoiseSettings.Lacunarity = GetValue<float>(Context, &Lacunarity);
-					NoiseSettings.Persistence = GetValue<float>(Context, &Persistence);;
-
-					FVoronoiDiagram Voronoi(Sites, VoronoiBounds, .1f);
-
-					FPlanarCells VoronoiPlanarCells = FPlanarCells(Sites, Voronoi);
-					VoronoiPlanarCells.InternalSurfaceMaterials.NoiseSettings = NoiseSettings;
-
-					const TArrayView<const int32>& TransformIndicesArray(TransformIndices.GetConstArray());
-
-					float CollisionSampleSpacingVal = GetValue<float>(Context, &CollisionSampleSpacing);
-					float RandomSeedVal = GetValue<float>(Context, &RandomSeed);
-
-					int ResultGeometryIndex = CutMultipleWithPlanarCells(VoronoiPlanarCells, *GeomCollection, TransformIndicesArray, GroutVal, CollisionSampleSpacingVal, RandomSeedVal, FTransform().Identity);
-
-					SetValue<FManagedArrayCollection>(Context, (const FManagedArrayCollection&)(*GeomCollection), &Collection);
-				}
+				FManagedArrayCollection InCollection = GetValue<FManagedArrayCollection>(Context, &Collection);
+
+				FFractureEngineFracturing::VoronoiFracture(InCollection,
+					InTransformSelection,
+					GetValue<TArray<FVector>>(Context, &Points),
+					GetValue<float>(Context, &RandomSeed),
+					GetValue<float>(Context, &ChanceToFracture),
+					GroupFracture,
+					GetValue<float>(Context, &Grout),
+					GetValue<float>(Context, &Amplitude),
+					GetValue<float>(Context, &Frequency),
+					GetValue<float>(Context, &Persistence),
+					GetValue<float>(Context, &Lacunarity),
+					GetValue<int32>(Context, &OctaveNumber),
+					GetValue<float>(Context, &PointSpacing),
+					AddSamplesForCollision,
+					GetValue<float>(Context, &CollisionSampleSpacing));
+
+				SetValue<FManagedArrayCollection>(Context, InCollection, &Collection);
+
+				return;
 			}
 		}
+
+		const FManagedArrayCollection& InCollection = GetValue<FManagedArrayCollection>(Context, &Collection);
+		SetValue<FManagedArrayCollection>(Context, InCollection, &Collection);
 	}
 }
 
@@ -203,61 +153,53 @@ void FPlaneCutterDataflowNode::Evaluate(Dataflow::FContext& Context, const FData
 {
 	if (Out->IsA<FManagedArrayCollection>(&Collection))
 	{
-		const FManagedArrayCollection& InCollection = GetValue<FManagedArrayCollection>(Context, &Collection);
-		if (TUniquePtr<FGeometryCollection> GeomCollection = TUniquePtr<FGeometryCollection>(InCollection.NewCopy<FGeometryCollection>()))
-		{
-			TArray<FPlane> CuttingPlanes;
-			TArray<FTransform> CuttingPlaneTransforms;
-
-			float RandomSeedVal = GetValue<float>(Context, &RandomSeed);
-			FRandomStream RandStream(RandomSeedVal);
-
-			FBox Bounds = GetValue<FBox>(Context, &BoundingBox);
-			const FVector Extent(Bounds.Max - Bounds.Min);
-
-			CuttingPlaneTransforms.Reserve(CuttingPlaneTransforms.Num() + NumPlanes);
-			for (int32 ii = 0; ii < NumPlanes; ++ii)
-			{
-				FVector Position(Bounds.Min + FVector(RandStream.FRand(), RandStream.FRand(), RandStream.FRand()) * Extent);
-				CuttingPlaneTransforms.Emplace(FTransform(FRotator(RandStream.FRand() * 360.0f, RandStream.FRand() * 360.0f, 0.0f), Position));
-			}
+		const FDataflowTransformSelection& InTransformSelection = GetValue<FDataflowTransformSelection>(Context, &TransformSelection);
 
-			for (const FTransform& Transform : CuttingPlaneTransforms)
+		if (IsConnected<FDataflowTransformSelection>(&TransformSelection))
+		{
+			if (InTransformSelection.AnySelected())
 			{
-				CuttingPlanes.Add(FPlane(Transform.GetLocation(), Transform.GetUnitAxis(EAxis::Z)));
+				FManagedArrayCollection InCollection = GetValue<FManagedArrayCollection>(Context, &Collection);
+
+				FFractureEngineFracturing::PlaneCutter(InCollection,
+					InTransformSelection,
+					GetValue<FBox>(Context, &BoundingBox),
+					NumPlanes,
+					GetValue<float>(Context, &RandomSeed),
+					GetValue<float>(Context, &Grout),
+					GetValue<float>(Context, &Amplitude),
+					GetValue<float>(Context, &Frequency),
+					GetValue<float>(Context, &Persistence),
+					GetValue<float>(Context, &Lacunarity),
+					GetValue<int32>(Context, &OctaveNumber),
+					GetValue<float>(Context, &PointSpacing),
+					GetValue<bool>(Context, &AddSamplesForCollision),
+					GetValue<float>(Context, &CollisionSampleSpacing));
+
+				SetValue<FManagedArrayCollection>(Context, InCollection, &Collection);
+
+				return;
 			}
+		}
 
-			FInternalSurfaceMaterials InternalSurfaceMaterials;
-			FNoiseSettings NoiseSettings;
+		const FManagedArrayCollection& InCollection = GetValue<FManagedArrayCollection>(Context, &Collection);
+		SetValue<FManagedArrayCollection>(Context, InCollection, &Collection);
+	}
+}
 
-			float AmplitudeVal = GetValue<float>(Context, &Amplitude);
-			if (AmplitudeVal > 0.0f)
-			{
-				NoiseSettings.Amplitude = AmplitudeVal;
-				NoiseSettings.Frequency = GetValue<float>(Context, &Frequency);
-				NoiseSettings.Lacunarity = GetValue<float>(Context, &Lacunarity);
-				NoiseSettings.Persistence = GetValue<float>(Context, &Persistence);
-				NoiseSettings.Octaves = GetValue<int32>(Context, &OctaveNumber);
-				NoiseSettings.PointSpacing = GetValue<float>(Context, &PointSpacing);
-
-				InternalSurfaceMaterials.NoiseSettings = NoiseSettings;
-			}
 
-			if (GeomCollection->HasAttribute("TransformIndex", "Geometry"))
-			{
-				const TManagedArray<int32>& TransformIndices = GeomCollection->GetAttribute<int32>("TransformIndex", "Geometry");
-				const TArrayView<const int32>& TransformIndicesArray(TransformIndices.GetConstArray());
-
-				float CollisionSampleSpacingVal = GetValue<float>(Context, &CollisionSampleSpacing);
-				float GroutVal = GetValue<float>(Context, &Grout);
+void FExplodedViewDataflowNode::Evaluate(Dataflow::FContext& Context, const FDataflowOutput* Out) const
+{
+	if (Out->IsA<FManagedArrayCollection>(&Collection))
+	{
+		FManagedArrayCollection InCollection = GetValue<FManagedArrayCollection>(Context, &Collection);
 
-				int ResultGeometryIndex = CutMultipleWithMultiplePlanes(CuttingPlanes, InternalSurfaceMaterials, *GeomCollection, TransformIndicesArray, GroutVal, CollisionSampleSpacingVal, RandomSeedVal, FTransform().Identity);
+		FFractureEngineFracturing::GenerateExplodedViewAttribute(InCollection, GetValue<FVector>(Context, &Scale), GetValue<float>(Context, &UniformScale));
 
-				SetValue<FManagedArrayCollection>(Context, (const FManagedArrayCollection&)(*GeomCollection), &Collection);
-			}
-		}
+		SetValue<FManagedArrayCollection>(Context, MoveTemp(InCollection), &Collection);
 	}
 }
 
 
 
+
diff --git a/Engine/Plugins/Experimental/GeometryCollectionPlugin/Source/GeometryCollectionNodes/Private/Dataflow/GeometryCollectionMathNodes.cpp b/Engine/Plugins/Experimental/GeometryCollectionPlugin/Source/GeometryCollectionNodes/Private/Dataflow/GeometryCollectionMathNodes.cpp
new file mode 100644
--- /dev/null
+++ b/Engine/Plugins/Experimental/GeometryCollectionPlugin/Source/GeometryCollectionNodes/Private/Dataflow/GeometryCollectionMathNodes.cpp
@@ -0,0 +1,19 @@
+// Copyright Epic Games, Inc. All Rights Reserved.
+
+#include "Dataflow/GeometryCollectionMathNodes.h"
+#include "Dataflow/DataflowCore.h"
+
+
+//#include UE_INLINE_GENERATED_CPP_BY_NAME(GeometryCollectionMathNodes)
+
+namespace Dataflow
+{
+
+	void GeometryCollectionMathNodes()
+	{
+		static const FLinearColor CDefaultNodeBodyTintColor = FLinearColor(0.f, 0.f, 0.f, 0.5f);
+
+	}
+}
+
+
diff --git a/Engine/Plugins/Experimental/GeometryCollectionPlugin/Source/GeometryCollectionNodes/Private/Dataflow/GeometryCollectionNodes.cpp b/Engine/Plugins/Experimental/GeometryCollectionPlugin/Source/GeometryCollectionNodes/Private/Dataflow/GeometryCollectionNodes.cpp
--- a/Engine/Plugins/Experimental/GeometryCollectionPlugin/Source/GeometryCollectionNodes/Private/Dataflow/GeometryCollectionNodes.cpp
+++ b/Engine/Plugins/Experimental/GeometryCollectionPlugin/Source/GeometryCollectionNodes/Private/Dataflow/GeometryCollectionNodes.cpp
@@ -42,12 +42,8 @@ namespace Dataflow
 	{
 		static const FLinearColor CDefaultNodeBodyTintColor = FLinearColor(0.f, 0.f, 0.f, 0.5f);
 
-		DATAFLOW_NODE_REGISTER_CREATION_FACTORY(FGetCollectionAssetDataflowNode);
-		DATAFLOW_NODE_REGISTER_CREATION_FACTORY(FExampleCollectionEditDataflowNode);
-		DATAFLOW_NODE_REGISTER_CREATION_FACTORY(FSetCollectionAssetDataflowNode);
+		DATAFLOW_NODE_REGISTER_CREATION_FACTORY(FGetCollectionFromAssetDataflowNode);
 		DATAFLOW_NODE_REGISTER_CREATION_FACTORY(FAppendCollectionAssetsDataflowNode);
-		DATAFLOW_NODE_REGISTER_CREATION_FACTORY(FResetGeometryCollectionDataflowNode);
-
 		DATAFLOW_NODE_REGISTER_CREATION_FACTORY(FPrintStringDataflowNode);
 		DATAFLOW_NODE_REGISTER_CREATION_FACTORY(FLogStringDataflowNode);
 		DATAFLOW_NODE_REGISTER_CREATION_FACTORY(FMakeLiteralStringDataflowNode);
@@ -73,8 +69,6 @@ namespace Dataflow
 		DATAFLOW_NODE_REGISTER_CREATION_FACTORY(FRandomUnitVectorInConeDataflowNode);
 		DATAFLOW_NODE_REGISTER_CREATION_FACTORY(FRadiansToDegreesDataflowNode);
 		DATAFLOW_NODE_REGISTER_CREATION_FACTORY(FDegreesToRadiansDataflowNode);
-		DATAFLOW_NODE_REGISTER_CREATION_FACTORY(FExplodedViewDataflowNode);
-		DATAFLOW_NODE_REGISTER_CREATION_FACTORY(FCreateNonOverlappingConvexHullsDataflowNode);
 		DATAFLOW_NODE_REGISTER_CREATION_FACTORY(FHashStringDataflowNode);
 		DATAFLOW_NODE_REGISTER_CREATION_FACTORY(FHashVectorDataflowNode);
 		DATAFLOW_NODE_REGISTER_CREATION_FACTORY(FFloatToIntDataflowNode);
@@ -95,6 +89,7 @@ namespace Dataflow
 		DATAFLOW_NODE_REGISTER_CREATION_FACTORY(FGetNumElementsInCollectionGroupDataflowNode);
 		DATAFLOW_NODE_REGISTER_CREATION_FACTORY(FGetCollectionAttributeDataTypedDataflowNode);
 		DATAFLOW_NODE_REGISTER_CREATION_FACTORY(FSetCollectionAttributeDataTypedDataflowNode);
+		DATAFLOW_NODE_REGISTER_CREATION_FACTORY(FBoolArrayToFaceSelectionDataflowNode);
 
 		// GeometryCollection
 		DATAFLOW_NODE_REGISTER_CREATION_FACTORY_NODE_COLORS_BY_CATEGORY("GeometryCollection", FLinearColor(0.55f, 0.45f, 1.0f), CDefaultNodeBodyTintColor);
@@ -113,56 +108,29 @@ namespace Dataflow
 	}
 }
 
-void FGetCollectionAssetDataflowNode::Evaluate(Dataflow::FContext& Context, const FDataflowOutput* Out) const
+void FGetCollectionFromAssetDataflowNode::Evaluate(Dataflow::FContext& Context, const FDataflowOutput* Out) const
 {
-	if (Out->IsA<DataType>(&Output))
+	if (Out->IsA<FManagedArrayCollection>(&Collection))
 	{
-		if (const Dataflow::FEngineContext* EngineContext = Context.AsType<Dataflow::FEngineContext>())
+		if (CollectionAsset)
 		{
-			if (UGeometryCollection* CollectionAsset = Cast<UGeometryCollection>(EngineContext->Owner))
+			if (const TSharedPtr<FGeometryCollection, ESPMode::ThreadSafe> AssetCollection = CollectionAsset->GetGeometryCollection())
 			{
-				if (const TSharedPtr<FGeometryCollection, ESPMode::ThreadSafe> AssetCollection = CollectionAsset->GetGeometryCollection())
-				{
-					SetValue<DataType>(Context, DataType(*AssetCollection), &Output);
-				}
+				SetValue<FManagedArrayCollection>(Context, (const FManagedArrayCollection&)(*AssetCollection), &Collection);
 			}
-		}
-	}
-}
-
-void FExampleCollectionEditDataflowNode::Evaluate(Dataflow::FContext& Context, const FDataflowOutput* Out) const
-{
-	if (Out->IsA<DataType>(&Collection))
-	{
-		DataType InCollection = GetValue<DataType>(Context, &Collection);
-
-		if (bActive)
-		{
-			TManagedArray<FVector3f>* Vertex = InCollection.FindAttribute<FVector3f>("Vertex", "Vertices");
-			for (int i = 0; i < Vertex->Num(); i++)
+			else
 			{
-				(*Vertex)[i][1] *= Scale;
+				SetValue<FManagedArrayCollection>(Context, FManagedArrayCollection(), &Collection);
 			}
 		}
-		SetValue<DataType>(Context, InCollection, &Collection);
-	}
-}
-
-void FSetCollectionAssetDataflowNode::Evaluate(Dataflow::FContext& Context, const FDataflowOutput* Out) const
-{
-	DataType InCollection = GetValue<DataType>(Context, &Collection);
-
-	if (const Dataflow::FEngineContext* EngineContext = Context.AsType<Dataflow::FEngineContext>())
-	{
-		if (UGeometryCollection* CollectionAsset = Cast<UGeometryCollection>(EngineContext->Owner))
+		else
 		{
-			TSharedPtr<FGeometryCollection, ESPMode::ThreadSafe> NewCollection(InCollection.NewCopy<FGeometryCollection>());
-			CollectionAsset->SetGeometryCollection(NewCollection);
-			CollectionAsset->InvalidateCollection();
+			SetValue<FManagedArrayCollection>(Context, FManagedArrayCollection(), &Collection);
 		}
 	}
 }
 
+
 void FAppendCollectionAssetsDataflowNode::Evaluate(Dataflow::FContext& Context, const FDataflowOutput* Out) const
 {
 	if (Out->IsA<DataType>(&Collection1))
@@ -175,71 +143,6 @@ void FAppendCollectionAssetsDataflowNode::Evaluate(Dataflow::FContext& Context,
 	}
 }
 
-void FResetGeometryCollectionDataflowNode::Evaluate(Dataflow::FContext& Context, const FDataflowOutput* Out) const
-{
-	if (Out->IsA<DataType>(&Collection))
-	{
-		SetValue<DataType>(Context, Collection, &Collection); // prime to avoid ensure
-
-		if (const Dataflow::FEngineContext* EngineContext = Context.AsType<Dataflow::FEngineContext>())
-		{
-			if (UGeometryCollection* GeometryCollectionObject = Cast<UGeometryCollection>(EngineContext->Owner))
-			{
-				GeometryCollectionObject->Reset();
-
-				const UObject* Owner = EngineContext->Owner;
-				FName AName("GeometrySource");
-				if (Owner && Owner->GetClass())
-				{
-					if (const ::FProperty* UEProperty = Owner->GetClass()->FindPropertyByName(AName))
-					{
-						if (const FArrayProperty* ArrayProperty = CastField<FArrayProperty>(UEProperty))
-						{
-							FScriptArrayHelper_InContainer ArrayHelper(ArrayProperty, Owner);
-							const int32 ArraySize = ArrayHelper.Num();
-							for (int32 Index = 0; Index < ArraySize; ++Index)
-							{
-								if (FGeometryCollectionSource* SourceObject = (FGeometryCollectionSource*)(ArrayHelper.GetRawPtr(Index)))
-								{
-									if (UObject* ResolvedObject = SourceObject->SourceGeometryObject.ResolveObject())
-									{
-										if (UStaticMesh* StaticMesh = Cast<UStaticMesh>(ResolvedObject))
-										{
-											TArray<UMaterialInterface*> Materials;
-											Materials.Reserve(StaticMesh->GetStaticMaterials().Num());
-
-											for (int32 Index2 = 0; Index2 < StaticMesh->GetStaticMaterials().Num(); ++Index2)
-											{
-												UMaterialInterface* CurrMaterial = StaticMesh->GetMaterial(Index2);
-												Materials.Add(CurrMaterial);
-											}
-
-											// Geometry collections usually carry the selection material, which we'll delete before appending
-											UMaterialInterface* BoneSelectedMaterial = LoadObject<UMaterialInterface>(nullptr, UGeometryCollection::GetSelectedMaterialPath(), nullptr, LOAD_None, nullptr);
-											GeometryCollectionObject->Materials.Remove(BoneSelectedMaterial);
-											Materials.Remove(BoneSelectedMaterial);
-
-											FGeometryCollectionEngineConversion::AppendStaticMesh(StaticMesh, Materials, FTransform(), GeometryCollectionObject);
-
-										}
-									}
-								}
-							}
-						}
-					}
-				}
-				GeometryCollectionObject->UpdateGeometryDependentProperties();
-				GeometryCollectionObject->InitializeMaterials();
-				GeometryCollectionObject->InvalidateCollection();
-
-				if (const TSharedPtr<FGeometryCollection, ESPMode::ThreadSafe> AssetCollection = GeometryCollectionObject->GetGeometryCollection())
-				{
-					SetValue<DataType>(Context, *AssetCollection, &Collection);
-				}
-			}
-		}
-	}
-}
 
 void FPrintStringDataflowNode::Evaluate(Dataflow::FContext& Context, const FDataflowOutput* Out) const
 {
@@ -554,168 +457,6 @@ void FDegreesToRadiansDataflowNode::Evaluate(Dataflow::FContext& Context, const
 	}
 }
 
-static void AddAdditionalAttributesIfRequired(FGeometryCollection* GeometryCollection)
-{
-	FGeometryCollectionClusteringUtility::UpdateHierarchyLevelOfChildren(GeometryCollection, -1);
-}
-
-bool FExplodedViewDataflowNode::GetValidGeoCenter(FGeometryCollection* Collection, const TManagedArray<int32>& TransformToGeometryIndex, const TArray<FTransform>& Transforms, const TManagedArray<TSet<int32>>& Children, const TManagedArray<FBox>& BoundingBox, int32 TransformIndex, FVector& OutGeoCenter)
-{
-	if (Collection->IsRigid(TransformIndex))
-	{
-		OutGeoCenter = Transforms[TransformIndex].TransformPosition(BoundingBox[TransformToGeometryIndex[TransformIndex]].GetCenter());
-
-		return true;
-	}
-	else if (Collection->SimulationType[TransformIndex] == FGeometryCollection::ESimulationTypes::FST_None) // ie this is embedded geometry
-	{
-		int32 Parent = Collection->Parent[TransformIndex];
-		int32 ParentGeo = Parent != INDEX_NONE ? TransformToGeometryIndex[Parent] : INDEX_NONE;
-		if (ensureMsgf(ParentGeo != INDEX_NONE, TEXT("Embedded geometry should always have a rigid geometry parent!  Geometry collection may be malformed.")))
-		{
-			OutGeoCenter = Transforms[Collection->Parent[TransformIndex]].TransformPosition(BoundingBox[ParentGeo].GetCenter());
-		}
-		else
-		{
-			return false; // no valid value to return
-		}
-
-		return true;
-	}
-	else
-	{
-		FVector AverageCenter;
-		int32 ValidVectors = 0;
-		for (int32 ChildIndex : Children[TransformIndex])
-		{
-
-			if (GetValidGeoCenter(Collection, TransformToGeometryIndex, Transforms, Children, BoundingBox, ChildIndex, OutGeoCenter))
-			{
-				if (ValidVectors == 0)
-				{
-					AverageCenter = OutGeoCenter;
-				}
-				else
-				{
-					AverageCenter += OutGeoCenter;
-				}
-				++ValidVectors;
-			}
-		}
-
-		if (ValidVectors > 0)
-		{
-			OutGeoCenter = AverageCenter / ValidVectors;
-			return true;
-		}
-	}
-	return false;
-}
-
-void FExplodedViewDataflowNode::Evaluate(Dataflow::FContext& Context, const FDataflowOutput* Out) const
-{
-	if (Out->IsA<FManagedArrayCollection>(&Collection))
-	{
-		const FManagedArrayCollection& InCollection = GetValue<FManagedArrayCollection>(Context, &Collection);
-
-		if (TUniquePtr<FGeometryCollection> GeomCollection = TUniquePtr<FGeometryCollection>(InCollection.NewCopy<FGeometryCollection>()))
-		{
-			GeomCollection->AddAttribute<FVector3f>("ExplodedVector", FGeometryCollection::TransformGroup, FManagedArrayCollection::FConstructionParameters(FName()));
-			check(GeomCollection->HasAttribute("ExplodedVector", FGeometryCollection::TransformGroup));
-
-			TManagedArray<FVector3f>& ExplodedVectors = GeomCollection->ModifyAttribute<FVector3f>("ExplodedVector", FGeometryCollection::TransformGroup);
-			const TManagedArray<FTransform>& Transform = GeomCollection->GetAttribute<FTransform>("Transform", FGeometryCollection::TransformGroup);
-			const TManagedArray<int32>& TransformToGeometryIndex = GeomCollection->GetAttribute<int32>("TransformToGeometryIndex", FGeometryCollection::TransformGroup);
-			const TManagedArray<FBox>& BoundingBox = GeomCollection->GetAttribute<FBox>("BoundingBox", FGeometryCollection::GeometryGroup);
-
-			// Make sure we have valid "Level"
-			AddAdditionalAttributesIfRequired(GeomCollection.Get());
-
-			const TManagedArray<int32>& Levels = GeomCollection->GetAttribute<int32>("Level", FTransformCollection::TransformGroup);
-			const TManagedArray<int32>& Parent = GeomCollection->GetAttribute<int32>("Parent", FTransformCollection::TransformGroup);
-			const TManagedArray<TSet<int32>>& Children = GeomCollection->GetAttribute<TSet<int32>>("Children", FGeometryCollection::TransformGroup);
-
-			int32 ViewFractureLevel = -1;
-			int32 MaxFractureLevel = ViewFractureLevel;
-			for (int32 Idx = 0, ni = Transform.Num(); Idx < ni; ++Idx)
-			{
-				if (Levels[Idx] > MaxFractureLevel)
-					MaxFractureLevel = Levels[Idx];
-			}
-
-			TArray<FTransform> Transforms;
-			GeometryCollectionAlgo::GlobalMatrices(Transform, GeomCollection->Parent, Transforms);
-
-			TArray<FVector> TransformedCenters;
-			TransformedCenters.SetNumUninitialized(Transforms.Num());
-
-			int32 TransformsCount = 0;
-
-			FVector Center(ForceInitToZero);
-			for (int32 Idx = 0, ni = Transform.Num(); Idx < ni; ++Idx)
-			{
-				ExplodedVectors[Idx] = FVector3f::ZeroVector;
-				FVector GeoCenter;
-
-				if (GetValidGeoCenter(GeomCollection.Get(), TransformToGeometryIndex, Transforms, Children, BoundingBox, Idx, GeoCenter))
-				{
-					TransformedCenters[Idx] = GeoCenter;
-					if ((ViewFractureLevel < 0) || Levels[Idx] == ViewFractureLevel)
-					{
-						Center += TransformedCenters[Idx];
-						++TransformsCount;
-					}
-				}
-			}
-
-			Center /= TransformsCount;
-
-			for (int Level = 1; Level <= MaxFractureLevel; Level++)
-			{
-				for (int32 Idx = 0, ni = Transforms.Num(); Idx < ni; ++Idx)
-				{
-					if ((ViewFractureLevel < 0) || Levels[Idx] == ViewFractureLevel)
-					{
-						FVector ScaleVal = GetValue<FVector>(Context, &Scale);
-						float UniformScaleVal = GetValue<float>(Context, &UniformScale);
-
-						FVector ScaleVec = ScaleVal * UniformScaleVal;
-						ExplodedVectors[Idx] = (FVector3f)(TransformedCenters[Idx] - Center) * (FVector3f)ScaleVec;
-					}
-					else
-					{
-						if (Parent[Idx] > -1)
-						{
-							ExplodedVectors[Idx] = ExplodedVectors[Parent[Idx]];
-						}
-					}
-				}
-			}
-
-			SetValue<FManagedArrayCollection>(Context, (const FManagedArrayCollection&)(*GeomCollection), &Collection);
-		}
-	}
-}
-
-void FCreateNonOverlappingConvexHullsDataflowNode::Evaluate(Dataflow::FContext& Context, const FDataflowOutput* Out) const
-{
-	if (Out->IsA<FManagedArrayCollection>(&Collection))
-	{
-		const FManagedArrayCollection& InCollection = GetValue<FManagedArrayCollection>(Context, &Collection);
-
-		if (TUniquePtr<FGeometryCollection> GeomCollection = TUniquePtr<FGeometryCollection>(InCollection.NewCopy<FGeometryCollection>()))
-		{
-			float CanRemoveFractionVal = GetValue<float>(Context, &CanRemoveFraction);
-			float CanExceedFractionVal = GetValue<float>(Context, &CanExceedFraction);
-			float SimplificationDistanceThresholdVal = GetValue<float>(Context, &SimplificationDistanceThreshold);
-
-			FGeometryCollectionConvexUtility::FGeometryCollectionConvexData ConvexData = FGeometryCollectionConvexUtility::CreateNonOverlappingConvexHullData(GeomCollection.Get(), CanRemoveFractionVal, SimplificationDistanceThresholdVal, CanExceedFractionVal);
-
-			SetValue<FManagedArrayCollection>(Context, (const FManagedArrayCollection&)(*GeomCollection), &Collection);
-		}
-	}
-}
-
 
 void FHashStringDataflowNode::Evaluate(Dataflow::FContext& Context, const FDataflowOutput* Out) const
 {
@@ -1472,4 +1213,17 @@ void FSetCollectionAttributeDataTypedDataflowNode::Evaluate(Dataflow::FContext&
 }
 
 
+void FBoolArrayToFaceSelectionDataflowNode::Evaluate(Dataflow::FContext& Context, const FDataflowOutput* Out) const
+{
+	if (Out->IsA<FDataflowFaceSelection>(&FaceSelection))
+	{
+		const TArray<bool>& InBoolAttributeData = GetValue<TArray<bool>>(Context, &BoolAttributeData);
+
+		FDataflowFaceSelection NewFaceSelection;
+		NewFaceSelection.Initialize(InBoolAttributeData.Num(), false);
+		NewFaceSelection.SetFromArray(InBoolAttributeData);
+
+		SetValue<FDataflowFaceSelection>(Context, NewFaceSelection, &FaceSelection);
+	}
+}
 
diff --git a/Engine/Plugins/Experimental/GeometryCollectionPlugin/Source/GeometryCollectionNodes/Private/Dataflow/GeometryCollectionNodesPlugin.cpp b/Engine/Plugins/Experimental/GeometryCollectionPlugin/Source/GeometryCollectionNodes/Private/Dataflow/GeometryCollectionNodesPlugin.cpp
--- a/Engine/Plugins/Experimental/GeometryCollectionPlugin/Source/GeometryCollectionNodes/Private/Dataflow/GeometryCollectionNodesPlugin.cpp
+++ b/Engine/Plugins/Experimental/GeometryCollectionPlugin/Source/GeometryCollectionNodes/Private/Dataflow/GeometryCollectionNodesPlugin.cpp
@@ -14,6 +14,7 @@
 #include "Dataflow/GeometryCollectionClusteringNodes.h"
 #include "Dataflow/GeometryCollectionFracturingNodes.h"
 #include "Dataflow/GeometryCollectionEditNodes.h"
+#include "Dataflow/GeometryCollectionUtilityNodes.h"
 
 #define LOCTEXT_NAMESPACE "DataflowNodes"
 
@@ -29,6 +30,7 @@ void IGeometryCollectionNodesPlugin::StartupModule()
 	Dataflow::GeometryCollectionClusteringNodes();
 	Dataflow::GeometryCollectionFracturingNodes();
 	Dataflow::GeometryCollectionEditNodes();
+	Dataflow::GeometryCollectionUtilityNodes();
 }
 
 void IGeometryCollectionNodesPlugin::ShutdownModule()
diff --git a/Engine/Plugins/Experimental/GeometryCollectionPlugin/Source/GeometryCollectionNodes/Private/Dataflow/GeometryCollectionSelectionNodes.cpp b/Engine/Plugins/Experimental/GeometryCollectionPlugin/Source/GeometryCollectionNodes/Private/Dataflow/GeometryCollectionSelectionNodes.cpp
--- a/Engine/Plugins/Experimental/GeometryCollectionPlugin/Source/GeometryCollectionNodes/Private/Dataflow/GeometryCollectionSelectionNodes.cpp
+++ b/Engine/Plugins/Experimental/GeometryCollectionPlugin/Source/GeometryCollectionNodes/Private/Dataflow/GeometryCollectionSelectionNodes.cpp
@@ -63,6 +63,7 @@ namespace Dataflow
 		DATAFLOW_NODE_REGISTER_CREATION_FACTORY(FCollectionVertexSelectionCustomDataflowNode);
 		DATAFLOW_NODE_REGISTER_CREATION_FACTORY(FCollectionFaceSelectionCustomDataflowNode);
 		DATAFLOW_NODE_REGISTER_CREATION_FACTORY(FCollectionSelectionConvertDataflowNode);
+		DATAFLOW_NODE_REGISTER_CREATION_FACTORY(FCollectionFaceSelectionInvertDataflowNode);
 
 		// GeometryCollection|Selection
 		DATAFLOW_NODE_REGISTER_CREATION_FACTORY_NODE_COLORS_BY_CATEGORY("GeometryCollection|Selection", FLinearColor(1.f, 1.f, 0.05f), CDefaultNodeBodyTintColor);
@@ -984,3 +985,17 @@ void FCollectionSelectionConvertDataflowNode::Evaluate(Dataflow::FContext& Conte
 		SetValue<FManagedArrayCollection>(Context, InCollection, &Collection);
 	}
 }
+
+
+void FCollectionFaceSelectionInvertDataflowNode::Evaluate(Dataflow::FContext& Context, const FDataflowOutput* Out) const
+{
+	if (Out->IsA<FDataflowFaceSelection>(&FaceSelection))
+	{
+		FDataflowFaceSelection InFaceSelection = GetValue<FDataflowFaceSelection>(Context, &FaceSelection);
+
+		InFaceSelection.Invert();
+
+		SetValue<FDataflowFaceSelection>(Context, MoveTemp(InFaceSelection), &FaceSelection);
+	}
+}
+
diff --git a/Engine/Plugins/Experimental/GeometryCollectionPlugin/Source/GeometryCollectionNodes/Private/Dataflow/GeometryCollectionUtilityNodes.cpp b/Engine/Plugins/Experimental/GeometryCollectionPlugin/Source/GeometryCollectionNodes/Private/Dataflow/GeometryCollectionUtilityNodes.cpp
new file mode 100644
--- /dev/null
+++ b/Engine/Plugins/Experimental/GeometryCollectionPlugin/Source/GeometryCollectionNodes/Private/Dataflow/GeometryCollectionUtilityNodes.cpp
@@ -0,0 +1,48 @@
+// Copyright Epic Games, Inc. All Rights Reserved.
+
+#include "Dataflow/GeometryCollectionUtilityNodes.h"
+#include "Dataflow/DataflowCore.h"
+
+#include "GeometryCollection/ManagedArrayCollection.h"
+#include "GeometryCollection/GeometryCollection.h"
+#include "GeometryCollection/GeometryCollectionConvexUtility.h"
+
+//#include UE_INLINE_GENERATED_CPP_BY_NAME(GeometryCollectionUtilityNodes)
+
+namespace Dataflow
+{
+
+	void GeometryCollectionUtilityNodes()
+	{
+		static const FLinearColor CDefaultNodeBodyTintColor = FLinearColor(0.f, 0.f, 0.f, 0.5f);
+
+		DATAFLOW_NODE_REGISTER_CREATION_FACTORY(FCreateNonOverlappingConvexHullsDataflowNode);
+
+	}
+}
+
+
+void FCreateNonOverlappingConvexHullsDataflowNode::Evaluate(Dataflow::FContext& Context, const FDataflowOutput* Out) const
+{
+	if (Out->IsA<FManagedArrayCollection>(&Collection))
+	{
+		const FManagedArrayCollection& InCollection = GetValue<FManagedArrayCollection>(Context, &Collection);
+
+		if (TUniquePtr<FGeometryCollection> GeomCollection = TUniquePtr<FGeometryCollection>(InCollection.NewCopy<FGeometryCollection>()))
+		{
+			float InCanRemoveFraction = GetValue<float>(Context, &CanRemoveFraction);
+			float InCanExceedFraction = GetValue<float>(Context, &CanExceedFraction);
+			float InSimplificationDistanceThreshold = GetValue<float>(Context, &SimplificationDistanceThreshold);
+			float InOverlapRemovalShrinkPercent = GetValue<float>(Context, &OverlapRemovalShrinkPercent);
+
+			FGeometryCollectionConvexUtility::FGeometryCollectionConvexData ConvexData = FGeometryCollectionConvexUtility::CreateNonOverlappingConvexHullData(GeomCollection.Get(), 
+				InCanRemoveFraction, 
+				InSimplificationDistanceThreshold, 
+				InCanExceedFraction,
+				(EConvexOverlapRemoval)OverlapRemovalMethod,
+				InOverlapRemovalShrinkPercent);
+
+			SetValue<FManagedArrayCollection>(Context, (const FManagedArrayCollection&)(*GeomCollection), &Collection);
+		}
+	}
+}
diff --git a/Engine/Plugins/Experimental/GeometryCollectionPlugin/Source/GeometryCollectionNodes/Public/Dataflow/GeometryCollectionEditNodes.h b/Engine/Plugins/Experimental/GeometryCollectionPlugin/Source/GeometryCollectionNodes/Public/Dataflow/GeometryCollectionEditNodes.h
--- a/Engine/Plugins/Experimental/GeometryCollectionPlugin/Source/GeometryCollectionNodes/Public/Dataflow/GeometryCollectionEditNodes.h
+++ b/Engine/Plugins/Experimental/GeometryCollectionPlugin/Source/GeometryCollectionNodes/Public/Dataflow/GeometryCollectionEditNodes.h
@@ -28,7 +28,7 @@ public:
 	UPROPERTY(meta = (DataflowInput, DataflowOutput, DataflowPassthrough = "Collection", DataflowIntrinsic))
 	FManagedArrayCollection Collection;
 
-	/** Bone selection for pruning */
+	/** Transform selection for pruning */
 	UPROPERTY(meta = (DataflowInput, DisplayName = "TransformSelection", DataflowIntrinsic))
 	FDataflowTransformSelection TransformSelection;
 
@@ -71,19 +71,24 @@ public:
 	UPROPERTY(EditAnywhere, Category = "Visibility");
 	EVisibiltyOptionsEnum Visibility = EVisibiltyOptionsEnum::Dataflow_VisibilityOptions_Invisible;
 
-	/** Fractured GeometryCollection to prune */
+	/** Fractured GeometryCollection to set visibility */
 	UPROPERTY(meta = (DataflowInput, DataflowOutput, DataflowPassthrough = "Collection", DataflowIntrinsic))
 	FManagedArrayCollection Collection;
 
-	/** Bone selection for pruning */
+	/** Transform selection for setting visibility */
 	UPROPERTY(meta = (DataflowInput, DisplayName = "TransformSelection", DataflowIntrinsic))
 	FDataflowTransformSelection TransformSelection;
 
+	/** Face selection for setting visibility */
+	UPROPERTY(meta = (DataflowInput, DisplayName = "TransformSelection", DataflowIntrinsic))
+	FDataflowFaceSelection FaceSelection;
+
 	FSetVisibilityInCollectionDataflowNode(const Dataflow::FNodeParameters& InParam, FGuid InGuid = FGuid::NewGuid())
 		: FDataflowNode(InParam, InGuid)
 	{
 		RegisterInputConnection(&Collection);
 		RegisterInputConnection(&TransformSelection);
+		RegisterInputConnection(&FaceSelection);
 		RegisterOutputConnection(&Collection, &Collection);
 	}
 
@@ -108,7 +113,7 @@ public:
 	UPROPERTY(meta = (DataflowInput, DataflowOutput, DataflowPassthrough = "Collection", DataflowIntrinsic))
 	FManagedArrayCollection Collection;
 
-	/** Bone selection for merging */
+	/** Transform selection for merging */
 	UPROPERTY(meta = (DataflowInput, DisplayName = "TransformSelection", DataflowIntrinsic))
 	FDataflowTransformSelection TransformSelection;
 
diff --git a/Engine/Plugins/Experimental/GeometryCollectionPlugin/Source/GeometryCollectionNodes/Public/Dataflow/GeometryCollectionEmbedNodes.h b/Engine/Plugins/Experimental/GeometryCollectionPlugin/Source/GeometryCollectionNodes/Public/Dataflow/GeometryCollectionEmbedNodes.h
new file mode 100644
--- /dev/null
+++ b/Engine/Plugins/Experimental/GeometryCollectionPlugin/Source/GeometryCollectionNodes/Public/Dataflow/GeometryCollectionEmbedNodes.h
@@ -0,0 +1,22 @@
+// Copyright Epic Games, Inc. All Rights Reserved.
+
+#pragma once 
+
+#include "CoreMinimal.h"
+#include "Dataflow/DataflowEngine.h"
+#include "GeometryCollection/ManagedArrayCollection.h"
+#include "Dataflow/DataflowSelection.h"
+
+//#include "GeometryCollectionEmbedNodes.generated.h"
+
+
+class FGeometryCollection;
+
+
+
+
+namespace Dataflow
+{
+	void GeometryCollectionEmbedNodes();
+}
+
diff --git a/Engine/Plugins/Experimental/GeometryCollectionPlugin/Source/GeometryCollectionNodes/Public/Dataflow/GeometryCollectionFracturingNodes.h b/Engine/Plugins/Experimental/GeometryCollectionPlugin/Source/GeometryCollectionNodes/Public/Dataflow/GeometryCollectionFracturingNodes.h
--- a/Engine/Plugins/Experimental/GeometryCollectionPlugin/Source/GeometryCollectionNodes/Public/Dataflow/GeometryCollectionFracturingNodes.h
+++ b/Engine/Plugins/Experimental/GeometryCollectionPlugin/Source/GeometryCollectionNodes/Public/Dataflow/GeometryCollectionFracturingNodes.h
@@ -125,6 +125,10 @@ public:
 	UPROPERTY(meta = (DataflowInput, DataflowIntrinsic))
 	TArray<FVector> Points;
 
+	/**   */
+	UPROPERTY(meta = (DataflowInput, DisplayName = "TransformSelection", DataflowIntrinsic))
+	FDataflowTransformSelection TransformSelection;
+
 	UPROPERTY(EditAnywhere, Category = "Fracture", meta = (DataflowInput));
 	float RandomSeed = -1.f;
 
@@ -166,6 +170,7 @@ public:
 	{
 		RegisterInputConnection(&Collection);
 		RegisterInputConnection(&Points);
+		RegisterInputConnection(&TransformSelection);
 		RegisterInputConnection(&RandomSeed);
 		RegisterInputConnection(&ChanceToFracture);
 		RegisterInputConnection(&Grout);
@@ -193,7 +198,7 @@ USTRUCT(meta = (DataflowGeometryCollection))
 struct FPlaneCutterDataflowNode : public FDataflowNode
 {
 	GENERATED_USTRUCT_BODY()
-		DATAFLOW_NODE_DEFINE_INTERNAL(FPlaneCutterDataflowNode, "PlaneCutter", "GeometryCollection|Fracture", "")
+	DATAFLOW_NODE_DEFINE_INTERNAL(FPlaneCutterDataflowNode, "PlaneCutter", "GeometryCollection|Fracture", "")
 
 public:
 	UPROPERTY(meta = (DataflowInput, DataflowOutput, DataflowPassthrough = "Collection", DataflowIntrinsic))
@@ -202,6 +207,10 @@ public:
 	UPROPERTY(meta = (DataflowInput, DataflowIntrinsic))
 	FBox BoundingBox = FBox(ForceInit);
 
+	/**   */
+	UPROPERTY(meta = (DataflowInput, DisplayName = "TransformSelection", DataflowIntrinsic))
+	FDataflowTransformSelection TransformSelection;
+
 	UPROPERTY(EditAnywhere, Category = "Fracture", meta = (DataflowInput, UIMin = 1))
 	int32 NumPlanes = 1;
 
@@ -240,6 +249,7 @@ public:
 	{
 		RegisterInputConnection(&Collection);
 		RegisterInputConnection(&BoundingBox);
+		RegisterInputConnection(&TransformSelection);
 		RegisterInputConnection(&NumPlanes);
 		RegisterInputConnection(&RandomSeed);
 		RegisterInputConnection(&Grout);
@@ -249,6 +259,7 @@ public:
 		RegisterInputConnection(&Lacunarity);
 		RegisterInputConnection(&OctaveNumber);
 		RegisterInputConnection(&PointSpacing);
+		RegisterInputConnection(&AddSamplesForCollision);
 		RegisterInputConnection(&CollisionSampleSpacing);
 		RegisterOutputConnection(&Collection, &Collection);
 	}
@@ -256,6 +267,68 @@ public:
 	virtual void Evaluate(Dataflow::FContext& Context, const FDataflowOutput* Out) const override;
 
 };
+
+
+/**
+ *
+ * "Explodes" the pieces from the Collection for better visualization
+ *
+ */
+USTRUCT(meta = (DataflowGeometryCollection))
+struct FExplodedViewDataflowNode : public FDataflowNode
+{
+	GENERATED_USTRUCT_BODY()
+	DATAFLOW_NODE_DEFINE_INTERNAL(FExplodedViewDataflowNode, "ExplodedView", "GeometryCollection|Fracture|Utilities", "")
+	DATAFLOW_NODE_RENDER_TYPE(FGeometryCollection::StaticType(), "Collection")
+
+public:
+	UPROPERTY(meta = (DataflowInput, DataflowOutput, DataflowPassthrough = "Collection", DataflowIntrinsic))
+	FManagedArrayCollection Collection;
+
+	UPROPERTY(EditAnywhere, Category = "Scale", meta = (DataflowInput))
+	float UniformScale = 1.f;
+
+	UPROPERTY(EditAnywhere, Category = "Scale", meta = (DataflowInput))
+	FVector Scale = FVector(1.0);
+
+	FExplodedViewDataflowNode(const Dataflow::FNodeParameters& InParam, FGuid InGuid = FGuid::NewGuid())
+		: FDataflowNode(InParam, InGuid)
+	{
+		RegisterInputConnection(&Collection);
+		RegisterInputConnection(&UniformScale);
+		RegisterInputConnection(&Scale);
+		RegisterOutputConnection(&Collection, &Collection);
+	}
+
+	virtual void Evaluate(Dataflow::FContext& Context, const FDataflowOutput* Out) const override;
+
+private:
+	// todo(chaos) this is a copy of a function in FractureEditorModeToolkit, we should move this to a common place  
+	static bool GetValidGeoCenter(FGeometryCollection* Collection, const TManagedArray<int32>& TransformToGeometryIndex, const TArray<FTransform>& Transforms, const TManagedArray<TSet<int32>>& Children, const TManagedArray<FBox>& BoundingBox, int32 TransformIndex, FVector& OutGeoCenter);
+
+};
+
+
+
+
+
+
+
+
+
+
+
+
+
+
+
+
+
+
+
+
+
+
 
 
 namespace Dataflow
diff --git a/Engine/Plugins/Experimental/GeometryCollectionPlugin/Source/GeometryCollectionNodes/Public/Dataflow/GeometryCollectionMathNodes.h b/Engine/Plugins/Experimental/GeometryCollectionPlugin/Source/GeometryCollectionNodes/Public/Dataflow/GeometryCollectionMathNodes.h
new file mode 100644
--- /dev/null
+++ b/Engine/Plugins/Experimental/GeometryCollectionPlugin/Source/GeometryCollectionNodes/Public/Dataflow/GeometryCollectionMathNodes.h
@@ -0,0 +1,22 @@
+// Copyright Epic Games, Inc. All Rights Reserved.
+
+#pragma once 
+
+#include "CoreMinimal.h"
+#include "Dataflow/DataflowEngine.h"
+#include "GeometryCollection/ManagedArrayCollection.h"
+#include "Dataflow/DataflowSelection.h"
+
+//#include "GeometryCollectionMathNodes.generated.h"
+
+
+class FGeometryCollection;
+
+
+
+
+namespace Dataflow
+{
+	void GeometryCollectionMathNodes();
+}
+
diff --git a/Engine/Plugins/Experimental/GeometryCollectionPlugin/Source/GeometryCollectionNodes/Public/Dataflow/GeometryCollectionNodes.h b/Engine/Plugins/Experimental/GeometryCollectionPlugin/Source/GeometryCollectionNodes/Public/Dataflow/GeometryCollectionNodes.h
--- a/Engine/Plugins/Experimental/GeometryCollectionPlugin/Source/GeometryCollectionNodes/Public/Dataflow/GeometryCollectionNodes.h
+++ b/Engine/Plugins/Experimental/GeometryCollectionPlugin/Source/GeometryCollectionNodes/Public/Dataflow/GeometryCollectionNodes.h
@@ -14,6 +14,7 @@
 
 
 class FGeometryCollection;
+class UGeometryCollection;
 class UStaticMesh;
 
 /**
@@ -22,82 +23,23 @@ class UStaticMesh;
  *
  */
 USTRUCT(meta = (DataflowGeometryCollection))
-struct FGetCollectionAssetDataflowNode : public FDataflowNode
+struct FGetCollectionFromAssetDataflowNode : public FDataflowNode
 {
 	GENERATED_USTRUCT_BODY()
-	DATAFLOW_NODE_DEFINE_INTERNAL(FGetCollectionAssetDataflowNode, "GetCollectionAsset", "GeometryCollection|Asset", "")
+	DATAFLOW_NODE_DEFINE_INTERNAL(FGetCollectionFromAssetDataflowNode, "GetCollectionFromAsset", "GeometryCollection|Asset", "")
 
 public:
-	typedef FManagedArrayCollection DataType;
-
-	UPROPERTY(meta = (DataflowOutput, DisplayName = "Collection"))
-	mutable FManagedArrayCollection Output;
-
-	FGetCollectionAssetDataflowNode(const Dataflow::FNodeParameters& InParam, FGuid InGuid = FGuid::NewGuid())
-		: FDataflowNode(InParam, InGuid)
-	{
-		RegisterOutputConnection(&Output);
-	}
-
-	virtual void Evaluate(Dataflow::FContext& Context, const FDataflowOutput* Out) const override;
-
-};
-
-
-/**
- *
- * Description for this node
- *
- */
-USTRUCT(meta = (DataflowGeometryCollection))
-struct FExampleCollectionEditDataflowNode : public FDataflowNode
-{
-	GENERATED_USTRUCT_BODY()
-	DATAFLOW_NODE_DEFINE_INTERNAL(FExampleCollectionEditDataflowNode, "ExampleCollectionEdit", "GeometryCollection", "")
+	UPROPERTY(EditAnywhere, Category = "Dataflow", meta = (DataflowInput, DisplayName = "Asset"))
+	TObjectPtr<UGeometryCollection> CollectionAsset;
 
-public:
-	typedef FManagedArrayCollection DataType;
-
-	/** Description for this parameter */
-	UPROPERTY(EditAnywhere, Category = "Dataflow", meta = (EditCondition = "false", EditConditionHides));
-	float Scale = 1.0;
-
-	UPROPERTY(meta = (DataflowInput, DataflowOutput, DataflowPassthrough = "Collection"))
+	UPROPERTY(meta = (DataflowOutput))
 	FManagedArrayCollection Collection;
 
-	FExampleCollectionEditDataflowNode(const Dataflow::FNodeParameters& InParam, FGuid InGuid = FGuid::NewGuid())
+	FGetCollectionFromAssetDataflowNode(const Dataflow::FNodeParameters& InParam, FGuid InGuid = FGuid::NewGuid())
 		: FDataflowNode(InParam, InGuid)
 	{
-		RegisterOutputConnection(&Collection, &Collection);
-		RegisterInputConnection(&Collection);
-	}
-
-	virtual void Evaluate(Dataflow::FContext& Context, const FDataflowOutput* Out) const override;
-
-};
-
-
-/**
- *
- * Description for this node
- *
- */
-USTRUCT(meta = (DataflowGeometryCollection, DataflowTerminal))
-struct FSetCollectionAssetDataflowNode : public FDataflowTerminalNode
-{
-	GENERATED_USTRUCT_BODY()
-	DATAFLOW_NODE_DEFINE_INTERNAL(FSetCollectionAssetDataflowNode, "SetCollectionAsset", "Terminal", "")
-
-public:
-	typedef FManagedArrayCollection DataType;
-
-	UPROPERTY(meta = (DataflowInput))
-	FManagedArrayCollection Collection;
-
-	FSetCollectionAssetDataflowNode(const Dataflow::FNodeParameters& InParam, FGuid InGuid = FGuid::NewGuid())
-		: FDataflowTerminalNode(InParam, InGuid)
-	{
-		RegisterInputConnection(&Collection);
+		RegisterInputConnection(&CollectionAsset);
+		RegisterOutputConnection(&Collection);
 	}
 
 	virtual void Evaluate(Dataflow::FContext& Context, const FDataflowOutput* Out) const override;
@@ -138,33 +80,6 @@ public:
 
 };
 
-/**
- *
- * Description for this node
- *
- */
-USTRUCT(meta = (DataflowGeometryCollection))
-struct FResetGeometryCollectionDataflowNode : public FDataflowNode
-{
-	GENERATED_USTRUCT_BODY()
-	DATAFLOW_NODE_DEFINE_INTERNAL(FResetGeometryCollectionDataflowNode, "ResetGeometryCollection", "GeometryCollection", "")
-
-public:
-	typedef FManagedArrayCollection DataType;
-
-	UPROPERTY(meta = (DataflowOutput))
-	FManagedArrayCollection Collection;
-
-	FResetGeometryCollectionDataflowNode(const Dataflow::FNodeParameters& InParam, FGuid InGuid = FGuid::NewGuid())
-		: FDataflowNode(InParam, InGuid)
-	{
-		RegisterOutputConnection(&Collection);
-	}
-
-	virtual void Evaluate(Dataflow::FContext& Context, const FDataflowOutput* Out) const override;
-
-};
-
 
 /**
  *
@@ -501,7 +416,7 @@ public:
 	float Radius = 0.f;
 
 	UPROPERTY(meta = (DataflowOutput));
-	FSphere Sphere;
+	FSphere Sphere = FSphere(ForceInit);
 
 	FMakeSphereDataflowNode(const Dataflow::FNodeParameters& InParam, FGuid InGuid = FGuid::NewGuid())
 		: FDataflowNode(InParam, InGuid)
@@ -1007,85 +922,6 @@ public:
 };
 
 
-/**
- *
- * "Explodes" the pieces from the Collection for better visualization
- *
- */
-USTRUCT(meta = (DataflowGeometryCollection))
-struct FExplodedViewDataflowNode : public FDataflowNode
-{
-	GENERATED_USTRUCT_BODY()
-	DATAFLOW_NODE_DEFINE_INTERNAL(FExplodedViewDataflowNode, "ExplodedView", "GeometryCollection|Fracture|Utilities", "")
-	DATAFLOW_NODE_RENDER_TYPE(FGeometryCollection::StaticType(), "Collection")
-
-public:
-	UPROPERTY(meta = (DataflowInput, DataflowOutput))
-	FManagedArrayCollection Collection;
-
-	UPROPERTY(EditAnywhere, Category = "Scale", meta = (DataflowInput))
-	float UniformScale = 1.f;
-
-	UPROPERTY(EditAnywhere, Category = "Scale", meta = (DataflowInput))
-	FVector Scale = FVector(1.0);
-
-	FExplodedViewDataflowNode(const Dataflow::FNodeParameters& InParam, FGuid InGuid = FGuid::NewGuid())
-		: FDataflowNode(InParam, InGuid)
-	{
-		RegisterInputConnection(&Collection);
-		RegisterInputConnection(&UniformScale);
-		RegisterInputConnection(&Scale);
-		RegisterOutputConnection(&Collection);
-	}
-
-	virtual void Evaluate(Dataflow::FContext& Context, const FDataflowOutput* Out) const override;
-
-private:
-	// todo(chaos) this is a copy of a function in FractureEditorModeToolkit, we should move this to a common place  
-	static bool GetValidGeoCenter(FGeometryCollection* Collection, const TManagedArray<int32>& TransformToGeometryIndex, const TArray<FTransform>& Transforms, const TManagedArray<TSet<int32>>& Children, const TManagedArray<FBox>& BoundingBox, int32 TransformIndex, FVector& OutGeoCenter);
-
-};
-
-
-/**
- *
- * Generates convex hull representation for the bones for simulation
- *
- */
-USTRUCT(meta = (DataflowGeometryCollection))
-struct FCreateNonOverlappingConvexHullsDataflowNode : public FDataflowNode
-{
-	GENERATED_USTRUCT_BODY()
-	DATAFLOW_NODE_DEFINE_INTERNAL(FCreateNonOverlappingConvexHullsDataflowNode, "CreateNonOverlappingConvexHulls", "GeometryCollection|Utilities", "")
-
-public:
-	UPROPERTY(meta = (DataflowInput, DataflowOutput))
-	FManagedArrayCollection Collection;
-
-	UPROPERTY(EditAnywhere, Category = "Convex", meta = (DataflowInput, UIMin = 0.01f, UIMax = 1.f))
-	float CanRemoveFraction = 0.5f;
-
-	UPROPERTY(EditAnywhere, Category = "Convex", meta = (DataflowInput, UIMin = 0.f))
-	float CanExceedFraction = 0.5f;
-
-	UPROPERTY(EditAnywhere, Category = "Convex", meta = (DataflowInput, UIMin = 0.f))
-	float SimplificationDistanceThreshold = 10.f;
-
-	FCreateNonOverlappingConvexHullsDataflowNode(const Dataflow::FNodeParameters& InParam, FGuid InGuid = FGuid::NewGuid())
-		: FDataflowNode(InParam, InGuid)
-	{
-		RegisterInputConnection(&Collection);
-		RegisterInputConnection(&CanRemoveFraction);
-		RegisterInputConnection(&CanExceedFraction);
-		RegisterInputConnection(&SimplificationDistanceThreshold);
-		RegisterOutputConnection(&Collection);
-	}
-
-	virtual void Evaluate(Dataflow::FContext& Context, const FDataflowOutput* Out) const override;
-
-};
-
-
 /**
  *
  * Generates a hash value from a string
@@ -2001,10 +1837,6 @@ public:
 	UPROPERTY(meta = (DataflowInput));
 	TArray<FVector> VectorAttributeData;
 
-	/** If yes then create attribute if it doesn't exist */
-//	UPROPERTY(EditAnywhere, Category = "Attribute", meta = (DisplayName = "Attribute"))
-//	FString AttrName = FString("");
-
 	FSetCollectionAttributeDataTypedDataflowNode(const Dataflow::FNodeParameters& InParam, FGuid InGuid = FGuid::NewGuid())
 		: FDataflowNode(InParam, InGuid)
 	{
@@ -2021,6 +1853,37 @@ public:
 };
 
 
+/**
+ *
+ * Converts a TArray<bool> to a FDataflowFaceSelection
+ *
+ */
+USTRUCT()
+struct FBoolArrayToFaceSelectionDataflowNode : public FDataflowNode
+{
+	GENERATED_USTRUCT_BODY()
+	DATAFLOW_NODE_DEFINE_INTERNAL(FBoolArrayToFaceSelectionDataflowNode, "BoolArrayToFaceSelection", "Utilities|Array", "")
+
+public:
+	/** TArray<bool> data */
+	UPROPERTY(meta = (DataflowInput));
+	TArray<bool> BoolAttributeData;
+
+	UPROPERTY(meta = (DataflowOutput, DisplayName = "FaceSelection"))
+	FDataflowFaceSelection FaceSelection;
+
+	FBoolArrayToFaceSelectionDataflowNode(const Dataflow::FNodeParameters& InParam, FGuid InGuid = FGuid::NewGuid())
+		: FDataflowNode(InParam, InGuid)
+	{
+		RegisterInputConnection(&BoolAttributeData);
+		RegisterOutputConnection(&FaceSelection);
+	}
+
+	virtual void Evaluate(Dataflow::FContext& Context, const FDataflowOutput* Out) const override;
+
+};
+
+
 namespace Dataflow
 {
 	void GeometryCollectionEngineNodes();
diff --git a/Engine/Plugins/Experimental/GeometryCollectionPlugin/Source/GeometryCollectionNodes/Public/Dataflow/GeometryCollectionSelectionNodes.h b/Engine/Plugins/Experimental/GeometryCollectionPlugin/Source/GeometryCollectionNodes/Public/Dataflow/GeometryCollectionSelectionNodes.h
--- a/Engine/Plugins/Experimental/GeometryCollectionPlugin/Source/GeometryCollectionNodes/Public/Dataflow/GeometryCollectionSelectionNodes.h
+++ b/Engine/Plugins/Experimental/GeometryCollectionPlugin/Source/GeometryCollectionNodes/Public/Dataflow/GeometryCollectionSelectionNodes.h
@@ -23,7 +23,7 @@ USTRUCT(meta = (DataflowGeometryCollection))
 struct FCollectionTransformSelectionAllDataflowNode : public FDataflowNode
 {
 	GENERATED_USTRUCT_BODY()
-	DATAFLOW_NODE_DEFINE_INTERNAL(FCollectionTransformSelectionAllDataflowNode, "CollectionTransformSelectAll", "GeometryCollection|Selection", "")
+	DATAFLOW_NODE_DEFINE_INTERNAL(FCollectionTransformSelectionAllDataflowNode, "CollectionTransformSelectAll", "GeometryCollection|Selection|Transform", "")
 
 public:
 	/** GeometryCollection for the selection */
@@ -70,7 +70,7 @@ USTRUCT(meta = (DataflowGeometryCollection))
 struct FCollectionTransformSelectionSetOperationDataflowNode : public FDataflowNode
 {
 	GENERATED_USTRUCT_BODY()
-	DATAFLOW_NODE_DEFINE_INTERNAL(FCollectionTransformSelectionSetOperationDataflowNode, "CollectionTransformSelectionSetOperation", "GeometryCollection|Selection", "")
+	DATAFLOW_NODE_DEFINE_INTERNAL(FCollectionTransformSelectionSetOperationDataflowNode, "CollectionTransformSelectionSetOperation", "GeometryCollection|Selection|Transform", "")
 
 public:
 	/** Boolean operation */
@@ -111,7 +111,7 @@ USTRUCT(meta = (DataflowGeometryCollection))
 struct FCollectionTransformSelectionInfoDataflowNode : public FDataflowNode
 {
 	GENERATED_USTRUCT_BODY()
-	DATAFLOW_NODE_DEFINE_INTERNAL(FCollectionTransformSelectionInfoDataflowNode, "CollectionTransformSelectionInfo", "GeometryCollection|Selection", "")
+	DATAFLOW_NODE_DEFINE_INTERNAL(FCollectionTransformSelectionInfoDataflowNode, "CollectionTransformSelectionInfo", "GeometryCollection|Selection|Transform", "")
 
 public:
 	/** Array of the selected bone indicies */
@@ -148,7 +148,7 @@ USTRUCT(meta = (DataflowGeometryCollection))
 struct FCollectionTransformSelectionNoneDataflowNode : public FDataflowNode
 {
 	GENERATED_USTRUCT_BODY()
-	DATAFLOW_NODE_DEFINE_INTERNAL(FCollectionTransformSelectionNoneDataflowNode, "CollectionTransformSelectNone", "GeometryCollection|Selection", "")
+	DATAFLOW_NODE_DEFINE_INTERNAL(FCollectionTransformSelectionNoneDataflowNode, "CollectionTransformSelectNone", "GeometryCollection|Selection|Transform", "")
 
 public:
 	/** GeometryCollection for the selection */
@@ -181,7 +181,7 @@ USTRUCT(meta = (DataflowGeometryCollection))
 struct FCollectionTransformSelectionInvertDataflowNode : public FDataflowNode
 {
 	GENERATED_USTRUCT_BODY()
-	DATAFLOW_NODE_DEFINE_INTERNAL(FCollectionTransformSelectionInvertDataflowNode, "CollectionTransformSelectInvert", "GeometryCollection|Selection", "")
+	DATAFLOW_NODE_DEFINE_INTERNAL(FCollectionTransformSelectionInvertDataflowNode, "CollectionTransformSelectInvert", "GeometryCollection|Selection|Transform", "")
 
 public:
 	/** Array of the selected bone indicies */
@@ -209,7 +209,7 @@ USTRUCT(meta = (DataflowGeometryCollection))
 struct FCollectionTransformSelectionRandomDataflowNode : public FDataflowNode
 {
 	GENERATED_USTRUCT_BODY()
-	DATAFLOW_NODE_DEFINE_INTERNAL(FCollectionTransformSelectionRandomDataflowNode, "CollectionTransformSelectRandom", "GeometryCollection|Selection", "")
+	DATAFLOW_NODE_DEFINE_INTERNAL(FCollectionTransformSelectionRandomDataflowNode, "CollectionTransformSelectRandom", "GeometryCollection|Selection|Transform", "")
 
 public:
 	/** If true, it always generates the same result for the same RandomSeed */
@@ -257,7 +257,7 @@ USTRUCT(meta = (DataflowGeometryCollection))
 struct FCollectionTransformSelectionRootDataflowNode : public FDataflowNode
 {
 	GENERATED_USTRUCT_BODY()
-	DATAFLOW_NODE_DEFINE_INTERNAL(FCollectionTransformSelectionRootDataflowNode, "CollectionTransformSelectRoot", "GeometryCollection|Selection", "")
+	DATAFLOW_NODE_DEFINE_INTERNAL(FCollectionTransformSelectionRootDataflowNode, "CollectionTransformSelectRoot", "GeometryCollection|Selection|Transform", "")
 
 public:
 	/** GeometryCollection for the selection */
@@ -291,7 +291,7 @@ USTRUCT(meta = (DataflowGeometryCollection))
 struct FCollectionTransformSelectionCustomDataflowNode : public FDataflowNode
 {
 	GENERATED_USTRUCT_BODY()
-	DATAFLOW_NODE_DEFINE_INTERNAL(FCollectionTransformSelectionCustomDataflowNode, "CollectionTransformSelectCustom", "GeometryCollection|Selection", "")
+	DATAFLOW_NODE_DEFINE_INTERNAL(FCollectionTransformSelectionCustomDataflowNode, "CollectionTransformSelectCustom", "GeometryCollection|Selection|Transform", "")
 
 public:
 	/** GeometryCollection for the selection */
@@ -329,7 +329,7 @@ USTRUCT(meta = (DataflowGeometryCollection))
 struct FCollectionTransformSelectionParentDataflowNode : public FDataflowNode
 {
 	GENERATED_USTRUCT_BODY()
-	DATAFLOW_NODE_DEFINE_INTERNAL(FCollectionTransformSelectionParentDataflowNode, "CollectionTransformSelectParent", "GeometryCollection|Selection", "")
+	DATAFLOW_NODE_DEFINE_INTERNAL(FCollectionTransformSelectionParentDataflowNode, "CollectionTransformSelectParent", "GeometryCollection|Selection|Transform", "")
 
 public:
 	/** Array of the selected bone indicies */
@@ -363,7 +363,7 @@ USTRUCT(meta = (DataflowGeometryCollection))
 struct FCollectionTransformSelectionByPercentageDataflowNode : public FDataflowNode
 {
 	GENERATED_USTRUCT_BODY()
-	DATAFLOW_NODE_DEFINE_INTERNAL(FCollectionTransformSelectionByPercentageDataflowNode, "CollectionTransformSelectByPercentage", "GeometryCollection|Selection", "")
+	DATAFLOW_NODE_DEFINE_INTERNAL(FCollectionTransformSelectionByPercentageDataflowNode, "CollectionTransformSelectByPercentage", "GeometryCollection|Selection|Transform", "")
 
 public:
 	/** Array of the selected bone indicies */
@@ -406,7 +406,7 @@ USTRUCT(meta = (DataflowGeometryCollection))
 struct FCollectionTransformSelectionChildrenDataflowNode : public FDataflowNode
 {
 	GENERATED_USTRUCT_BODY()
-	DATAFLOW_NODE_DEFINE_INTERNAL(FCollectionTransformSelectionChildrenDataflowNode, "CollectionTransformSelectChildren", "GeometryCollection|Selection", "")
+	DATAFLOW_NODE_DEFINE_INTERNAL(FCollectionTransformSelectionChildrenDataflowNode, "CollectionTransformSelectChildren", "GeometryCollection|Selection|Transform", "")
 
 public:
 	/** Array of the selected bone indicies */
@@ -440,7 +440,7 @@ USTRUCT(meta = (DataflowGeometryCollection))
 struct FCollectionTransformSelectionSiblingsDataflowNode : public FDataflowNode
 {
 	GENERATED_USTRUCT_BODY()
-	DATAFLOW_NODE_DEFINE_INTERNAL(FCollectionTransformSelectionSiblingsDataflowNode, "CollectionTransformSelectSiblings", "GeometryCollection|Selection", "")
+	DATAFLOW_NODE_DEFINE_INTERNAL(FCollectionTransformSelectionSiblingsDataflowNode, "CollectionTransformSelectSiblings", "GeometryCollection|Selection|Transform", "")
 
 public:
 	/** Array of the selected bone indicies */
@@ -474,7 +474,7 @@ USTRUCT(meta = (DataflowGeometryCollection))
 struct FCollectionTransformSelectionLevelDataflowNode : public FDataflowNode
 {
 	GENERATED_USTRUCT_BODY()
-	DATAFLOW_NODE_DEFINE_INTERNAL(FCollectionTransformSelectionLevelDataflowNode, "CollectionTransformSelectLevel", "GeometryCollection|Selection", "")
+	DATAFLOW_NODE_DEFINE_INTERNAL(FCollectionTransformSelectionLevelDataflowNode, "CollectionTransformSelectLevel", "GeometryCollection|Selection|Transform", "")
 
 public:
 	/** Array of the selected bone indicies */
@@ -508,7 +508,7 @@ USTRUCT(meta = (DataflowGeometryCollection))
 struct FCollectionTransformSelectionContactDataflowNode : public FDataflowNode
 {
 	GENERATED_USTRUCT_BODY()
-	DATAFLOW_NODE_DEFINE_INTERNAL(FCollectionTransformSelectionContactDataflowNode, "CollectionTransformSelectContact", "GeometryCollection|Selection", "")
+	DATAFLOW_NODE_DEFINE_INTERNAL(FCollectionTransformSelectionContactDataflowNode, "CollectionTransformSelectContact", "GeometryCollection|Selection|Transform", "")
 
 public:
 	/** Array of the selected bone indicies */
@@ -542,7 +542,7 @@ USTRUCT(meta = (DataflowGeometryCollection))
 struct FCollectionTransformSelectionLeafDataflowNode : public FDataflowNode
 {
 	GENERATED_USTRUCT_BODY()
-	DATAFLOW_NODE_DEFINE_INTERNAL(FCollectionTransformSelectionLeafDataflowNode, "CollectionTransformSelectLeaf", "GeometryCollection|Selection", "")
+	DATAFLOW_NODE_DEFINE_INTERNAL(FCollectionTransformSelectionLeafDataflowNode, "CollectionTransformSelectLeaf", "GeometryCollection|Selection|Transform", "")
 
 public:
 	/** GeometryCollection for the selection */
@@ -575,7 +575,7 @@ USTRUCT(meta = (DataflowGeometryCollection))
 struct FCollectionTransformSelectionClusterDataflowNode : public FDataflowNode
 {
 	GENERATED_USTRUCT_BODY()
-	DATAFLOW_NODE_DEFINE_INTERNAL(FCollectionTransformSelectionClusterDataflowNode, "CollectionTransformSelectCluster", "GeometryCollection|Selection", "")
+	DATAFLOW_NODE_DEFINE_INTERNAL(FCollectionTransformSelectionClusterDataflowNode, "CollectionTransformSelectCluster", "GeometryCollection|Selection|Transform", "")
 
 public:
 	/** GeometryCollection for the selection */
@@ -620,7 +620,7 @@ USTRUCT(meta = (DataflowGeometryCollection))
 struct FCollectionTransformSelectionBySizeDataflowNode : public FDataflowNode
 {
 	GENERATED_USTRUCT_BODY()
-	DATAFLOW_NODE_DEFINE_INTERNAL(FCollectionTransformSelectionBySizeDataflowNode, "CollectionTransformSelectBySize", "GeometryCollection|Selection", "")
+	DATAFLOW_NODE_DEFINE_INTERNAL(FCollectionTransformSelectionBySizeDataflowNode, "CollectionTransformSelectBySize", "GeometryCollection|Selection|Transform", "")
 
 public:
 	/** GeometryCollection for the selection */
@@ -671,7 +671,7 @@ USTRUCT(meta = (DataflowGeometryCollection))
 struct FCollectionTransformSelectionByVolumeDataflowNode : public FDataflowNode
 {
 	GENERATED_USTRUCT_BODY()
-	DATAFLOW_NODE_DEFINE_INTERNAL(FCollectionTransformSelectionByVolumeDataflowNode, "CollectionTransformSelectByVolume", "GeometryCollection|Selection", "")
+	DATAFLOW_NODE_DEFINE_INTERNAL(FCollectionTransformSelectionByVolumeDataflowNode, "CollectionTransformSelectByVolume", "GeometryCollection|Selection|Transform", "")
 
 public:
 	/** GeometryCollection for the selection */
@@ -736,7 +736,7 @@ USTRUCT(meta = (DataflowGeometryCollection))
 struct FCollectionTransformSelectionInBoxDataflowNode : public FDataflowNode
 {
 	GENERATED_USTRUCT_BODY()
-	DATAFLOW_NODE_DEFINE_INTERNAL(FCollectionTransformSelectionInBoxDataflowNode, "CollectionTransformSelectInBox", "GeometryCollection|Selection", "")
+	DATAFLOW_NODE_DEFINE_INTERNAL(FCollectionTransformSelectionInBoxDataflowNode, "CollectionTransformSelectInBox", "GeometryCollection|Selection|Transform", "")
 
 public:
 	/** GeometryCollection for the selection */
@@ -745,7 +745,7 @@ public:
 
 	/** Box to contain Vertices/BoundingBox/Centroid */
 	UPROPERTY(meta = (DataflowInput, DataflowIntrinsic))
-	FBox Box;
+	FBox Box = FBox(ForceInit);
 
 	/** Transform for the box */
 	UPROPERTY(EditAnywhere, Category = "Select")
@@ -787,7 +787,7 @@ USTRUCT(meta = (DataflowGeometryCollection))
 struct FCollectionTransformSelectionInSphereDataflowNode : public FDataflowNode
 {
 	GENERATED_USTRUCT_BODY()
-	DATAFLOW_NODE_DEFINE_INTERNAL(FCollectionTransformSelectionInSphereDataflowNode, "CollectionTransformSelectInSphere", "GeometryCollection|Selection", "")
+	DATAFLOW_NODE_DEFINE_INTERNAL(FCollectionTransformSelectionInSphereDataflowNode, "CollectionTransformSelectInSphere", "GeometryCollection|Selection|Transform", "")
 
 public:
 	/** GeometryCollection for the selection */
@@ -796,7 +796,7 @@ public:
 
 	/** Sphere to contain Vertices/BoundingBox/Centroid */
 	UPROPERTY(meta = (DataflowInput, DataflowIntrinsic))
-	FSphere Sphere;
+	FSphere Sphere = FSphere(ForceInit);
 
 	/** Transform for the sphere */
 	UPROPERTY(EditAnywhere, Category = "Select")
@@ -838,7 +838,7 @@ USTRUCT(meta = (DataflowGeometryCollection))
 struct FCollectionTransformSelectionByFloatAttrDataflowNode : public FDataflowNode
 {
 	GENERATED_USTRUCT_BODY()
-	DATAFLOW_NODE_DEFINE_INTERNAL(FCollectionTransformSelectionByFloatAttrDataflowNode, "CollectionTransformSelectByFloatAttribute", "GeometryCollection|Selection", "")
+	DATAFLOW_NODE_DEFINE_INTERNAL(FCollectionTransformSelectionByFloatAttrDataflowNode, "CollectionTransformSelectByFloatAttribute", "GeometryCollection|Selection|Transform", "")
 
 public:
 	/** GeometryCollection for the selection */
@@ -897,7 +897,7 @@ USTRUCT(meta = (DataflowGeometryCollection))
 struct FCollectionTransformSelectionByIntAttrDataflowNode : public FDataflowNode
 {
 	GENERATED_USTRUCT_BODY()
-	DATAFLOW_NODE_DEFINE_INTERNAL(FCollectionTransformSelectionByIntAttrDataflowNode, "CollectionTransformSelectByIntAttribute", "GeometryCollection|Selection", "")
+	DATAFLOW_NODE_DEFINE_INTERNAL(FCollectionTransformSelectionByIntAttrDataflowNode, "CollectionTransformSelectByIntAttribute", "GeometryCollection|Selection|Transform", "")
 
 public:
 	/** GeometryCollection for the selection */
@@ -957,7 +957,7 @@ USTRUCT(meta = (DataflowGeometryCollection))
 struct FCollectionVertexSelectionCustomDataflowNode : public FDataflowNode
 {
 	GENERATED_USTRUCT_BODY()
-	DATAFLOW_NODE_DEFINE_INTERNAL(FCollectionVertexSelectionCustomDataflowNode, "CollectionVertexSelectCustom", "GeometryCollection|Selection", "")
+	DATAFLOW_NODE_DEFINE_INTERNAL(FCollectionVertexSelectionCustomDataflowNode, "CollectionVertexSelectCustom", "GeometryCollection|Selection|Vertex", "")
 
 public:
 	/** GeometryCollection for the selection */
@@ -996,7 +996,7 @@ USTRUCT(meta = (DataflowGeometryCollection))
 struct FCollectionFaceSelectionCustomDataflowNode : public FDataflowNode
 {
 	GENERATED_USTRUCT_BODY()
-	DATAFLOW_NODE_DEFINE_INTERNAL(FCollectionFaceSelectionCustomDataflowNode, "CollectionFaceSelectCustom", "GeometryCollection|Selection", "")
+	DATAFLOW_NODE_DEFINE_INTERNAL(FCollectionFaceSelectionCustomDataflowNode, "CollectionFaceSelectCustom", "GeometryCollection|Selection|Face", "")
 
 public:
 	/** GeometryCollection for the selection */
@@ -1075,6 +1075,34 @@ public:
 };
 
 
+/**
+ *
+ * Inverts selection of faces
+ *
+ */
+USTRUCT(meta = (DataflowGeometryCollection))
+struct FCollectionFaceSelectionInvertDataflowNode : public FDataflowNode
+{
+	GENERATED_USTRUCT_BODY()
+	DATAFLOW_NODE_DEFINE_INTERNAL(FCollectionFaceSelectionInvertDataflowNode, "CollectionFaceSelectInvert", "GeometryCollection|Selection|Face", "")
+
+public:
+	/** Array of the selected bone indicies */
+	UPROPERTY(meta = (DataflowInput, DataflowOutput, DisplayName = "FaceSelection", DataflowPassthrough = "FaceSelection", DataflowIntrinsic))
+	FDataflowFaceSelection FaceSelection;
+
+	FCollectionFaceSelectionInvertDataflowNode(const Dataflow::FNodeParameters& InParam, FGuid InGuid = FGuid::NewGuid())
+		: FDataflowNode(InParam, InGuid)
+	{
+		RegisterInputConnection(&FaceSelection);
+		RegisterOutputConnection(&FaceSelection, &FaceSelection);
+	}
+
+	virtual void Evaluate(Dataflow::FContext& Context, const FDataflowOutput* Out) const override;
+
+};
+
+
 namespace Dataflow
 {
 	void GeometryCollectionSelectionNodes();
diff --git a/Engine/Plugins/Experimental/GeometryCollectionPlugin/Source/GeometryCollectionNodes/Public/Dataflow/GeometryCollectionUtilityNodes.h b/Engine/Plugins/Experimental/GeometryCollectionPlugin/Source/GeometryCollectionNodes/Public/Dataflow/GeometryCollectionUtilityNodes.h
new file mode 100644
--- /dev/null
+++ b/Engine/Plugins/Experimental/GeometryCollectionPlugin/Source/GeometryCollectionNodes/Public/Dataflow/GeometryCollectionUtilityNodes.h
@@ -0,0 +1,74 @@
+// Copyright Epic Games, Inc. All Rights Reserved.
+
+#pragma once 
+
+#include "CoreMinimal.h"
+#include "Dataflow/DataflowEngine.h"
+#include "GeometryCollection/ManagedArrayCollection.h"
+
+#include "GeometryCollectionUtilityNodes.generated.h"
+
+
+UENUM(BlueprintType)
+enum class EConvexOverlapRemovalMethodEnum : uint8
+{
+	Dataflow_EConvexOverlapRemovalMethod_None UMETA(DisplayName = "None"),
+	Dataflow_EConvexOverlapRemovalMethod_All UMETA(DisplayName = "All"),
+	Dataflow_EConvexOverlapRemovalMethod_OnlyClusters UMETA(DisplayName = "Only Clusters"),
+	Dataflow_EConvexOverlapRemovalMethod_OnlyClustersVsClusters UMETA(DisplayName = "Only Clusters vs Clusters"),
+	//~~~
+	//256th entry
+	Dataflow_Max                UMETA(Hidden)
+};
+
+/**
+ *
+ * Generates convex hull representation for the bones for simulation
+ *
+ */
+USTRUCT(meta = (DataflowGeometryCollection))
+struct FCreateNonOverlappingConvexHullsDataflowNode : public FDataflowNode
+{
+	GENERATED_USTRUCT_BODY()
+	DATAFLOW_NODE_DEFINE_INTERNAL(FCreateNonOverlappingConvexHullsDataflowNode, "CreateNonOverlappingConvexHulls", "GeometryCollection|Utilities", "")
+
+public:
+	UPROPERTY(meta = (DataflowInput, DataflowOutput))
+	FManagedArrayCollection Collection;
+
+	UPROPERTY(EditAnywhere, Category = "Convex", meta = (DataflowInput, UIMin = 0.01f, UIMax = 1.f))
+	float CanRemoveFraction = 0.3f;
+
+	UPROPERTY(EditAnywhere, Category = "Convex", meta = (DataflowInput, UIMin = 0.f))
+	float SimplificationDistanceThreshold = 0.f;
+
+	UPROPERTY(EditAnywhere, Category = "Convex", meta = (DataflowInput, UIMin = 0.f))
+	float CanExceedFraction = 0.5f;
+
+	UPROPERTY(EditAnywhere, Category = "Convex")
+	EConvexOverlapRemovalMethodEnum OverlapRemovalMethod = EConvexOverlapRemovalMethodEnum::Dataflow_EConvexOverlapRemovalMethod_All;
+
+	UPROPERTY(EditAnywhere, Category = "Convex", meta = (DataflowInput, UIMin = 0.f, UIMax = 100.f))
+	float OverlapRemovalShrinkPercent = 0.f;
+
+	FCreateNonOverlappingConvexHullsDataflowNode(const Dataflow::FNodeParameters& InParam, FGuid InGuid = FGuid::NewGuid())
+		: FDataflowNode(InParam, InGuid)
+	{
+		RegisterInputConnection(&Collection);
+		RegisterInputConnection(&CanRemoveFraction);
+		RegisterInputConnection(&SimplificationDistanceThreshold);
+		RegisterInputConnection(&CanExceedFraction);
+		RegisterInputConnection(&OverlapRemovalShrinkPercent);
+		RegisterOutputConnection(&Collection);
+	}
+
+	virtual void Evaluate(Dataflow::FContext& Context, const FDataflowOutput* Out) const override;
+
+};
+
+
+namespace Dataflow
+{
+	void GeometryCollectionUtilityNodes();
+}
+
diff --git a/Engine/Plugins/Experimental/GeometryCollectionPlugin/Source/GeometryCollectionNodes/Public/Dataflow/Nodes/GeometryCollectionAssetNodes.h b/Engine/Plugins/Experimental/GeometryCollectionPlugin/Source/GeometryCollectionNodes/Public/Dataflow/Nodes/GeometryCollectionAssetNodes.h
--- a/Engine/Plugins/Experimental/GeometryCollectionPlugin/Source/GeometryCollectionNodes/Public/Dataflow/Nodes/GeometryCollectionAssetNodes.h
+++ b/Engine/Plugins/Experimental/GeometryCollectionPlugin/Source/GeometryCollectionNodes/Public/Dataflow/Nodes/GeometryCollectionAssetNodes.h
@@ -47,7 +47,6 @@ public:
 };
 
 
-
 /**
  * Get Current geometry collection asset 
  * Note : Use with caution as this may get replaced in a near future for a more generic getAsset node
@@ -67,6 +66,7 @@ public:
 	TObjectPtr<UGeometryCollection> Asset;
 };
 
+
 /**
  * Get the list of the original mesh information used to create a specific geometryc collection asset
  * each entry contains a mesh, a transform and a list of override materials
@@ -75,7 +75,7 @@ USTRUCT(meta = (DataflowGeometryCollection))
 struct FGetGeometryCollectionSourcesDataflowNode : public FDataflowNode
 {
 	GENERATED_USTRUCT_BODY()
-	DATAFLOW_NODE_DEFINE_INTERNAL(FGetGeometryCollectionSourcesDataflowNode, "GetGeometryCollectionSources", "GeometryCollection|Utilities", "")
+	DATAFLOW_NODE_DEFINE_INTERNAL(FGetGeometryCollectionSourcesDataflowNode, "GetGeometryCollectionSources", "GeometryCollection|Asset", "")
 
 public:
 	FGetGeometryCollectionSourcesDataflowNode(const Dataflow::FNodeParameters& InParam, FGuid InGuid = FGuid::NewGuid());
@@ -90,6 +90,7 @@ public:
 	TArray<FGeometryCollectionSource> Sources;
 };
 
+
 /**
  * create a geometry collection from a set of geometry sources    
  */
@@ -97,7 +98,7 @@ USTRUCT(meta = (DataflowGeometryCollection))
 struct FCreateGeometryCollectionFromSourcesDataflowNode : public FDataflowNode
 {
 	GENERATED_USTRUCT_BODY()
-	DATAFLOW_NODE_DEFINE_INTERNAL(FCreateGeometryCollectionFromSourcesDataflowNode, "CreateGeometryCollectionFromSources", "GeometryCollection|Utilities", "")
+	DATAFLOW_NODE_DEFINE_INTERNAL(FCreateGeometryCollectionFromSourcesDataflowNode, "CreateGeometryCollectionFromSources", "GeometryCollection|Asset", "")
 
 public:
 	FCreateGeometryCollectionFromSourcesDataflowNode(const Dataflow::FNodeParameters& InParam, FGuid InGuid = FGuid::NewGuid());
diff --git a/Engine/Source/Runtime/Experimental/Chaos/Private/GeometryCollection/Facades/CollectionExplodedVectorFacade.cpp b/Engine/Source/Runtime/Experimental/Chaos/Private/GeometryCollection/Facades/CollectionExplodedVectorFacade.cpp
new file mode 100644
--- /dev/null
+++ b/Engine/Source/Runtime/Experimental/Chaos/Private/GeometryCollection/Facades/CollectionExplodedVectorFacade.cpp
@@ -0,0 +1,70 @@
+// Copyright Epic Games, Inc. All Rights Reserved.
+
+#include "GeometryCollection/Facades/CollectionExplodedVectorFacade.h"
+#include "GeometryCollection/GeometryCollection.h"
+#include "GeometryCollection/TransformCollection.h"
+
+namespace GeometryCollection::Facades
+{
+	FCollectionExplodedVectorFacade::FCollectionExplodedVectorFacade(FManagedArrayCollection& InCollection)
+		: ExplodedVectorAttribute(InCollection, "ExplodedVector", FGeometryCollection::TransformGroup)
+	{
+	}
+
+	FCollectionExplodedVectorFacade::FCollectionExplodedVectorFacade(const FManagedArrayCollection& InCollection)
+		: ExplodedVectorAttribute(InCollection, "ExplodedVector", FGeometryCollection::TransformGroup)
+	{
+	}
+
+	bool FCollectionExplodedVectorFacade::IsValid() const
+	{
+		return ExplodedVectorAttribute.IsValid();
+	}
+
+	void FCollectionExplodedVectorFacade::DefineSchema()
+	{
+		ExplodedVectorAttribute.Add();
+	}
+
+	void FCollectionExplodedVectorFacade::UpdateGlobalMatricesWithExplodedVectors(TArray<FMatrix>& InOutGlobalMatrices) const
+	{
+		if (IsValid())
+		{
+			const int32 NumMatrices = InOutGlobalMatrices.Num();
+			if (NumMatrices > 0)
+			{
+				const TManagedArray<FVector3f>& ExplodedVectors = ExplodedVectorAttribute.Get();
+
+				if (NumMatrices == ExplodedVectors.Num())
+				{
+					for (int32 Idx = 0; Idx < NumMatrices; ++Idx)
+					{
+						InOutGlobalMatrices[Idx] = InOutGlobalMatrices[Idx].ConcatTranslation((FVector)ExplodedVectors[Idx]);
+					}
+				}
+			}
+		}
+	}
+
+	void FCollectionExplodedVectorFacade::UpdateGlobalMatricesWithExplodedVectors(TArray<FTransform>& InOutGlobalTransforms) const
+	{
+		if (IsValid())
+		{
+			const int32 NumMatrices = InOutGlobalTransforms.Num();
+			if (NumMatrices > 0)
+			{
+				const TManagedArray<FVector3f>& ExplodedVectors = ExplodedVectorAttribute.Get();
+
+				if (NumMatrices == ExplodedVectors.Num())
+				{
+					for (int32 Idx = 0; Idx < NumMatrices; ++Idx)
+					{
+						InOutGlobalTransforms[Idx].AddToTranslation((FVector)ExplodedVectors[Idx]);
+					}
+				}
+			}
+		}
+	}
+};
+
+
diff --git a/Engine/Source/Runtime/Experimental/Chaos/Private/GeometryCollection/Facades/CollectionMeshFacade.cpp b/Engine/Source/Runtime/Experimental/Chaos/Private/GeometryCollection/Facades/CollectionMeshFacade.cpp
--- a/Engine/Source/Runtime/Experimental/Chaos/Private/GeometryCollection/Facades/CollectionMeshFacade.cpp
+++ b/Engine/Source/Runtime/Experimental/Chaos/Private/GeometryCollection/Facades/CollectionMeshFacade.cpp
@@ -135,6 +135,34 @@ namespace GeometryCollection::Facades
 		return FaceIndicies;
 	}
 
+	//
+	// This is a temporary solution to decide inside/outside faces based on MaterialIDs
+	// It has to be created and set in the geometry library during fracturing operations
+	//
+	void FCollectionMeshFacade::AddInternalAttribute(FGeometryCollection& InGeometryCollection, const TArray<int32>& InMaterialID)
+	{
+		if (!InGeometryCollection.HasAttribute("Internal", FGeometryCollection::FacesGroup))
+		{
+			InGeometryCollection.AddAttribute<bool>("Internal", FGeometryCollection::FacesGroup);
+
+			TManagedArray<bool>& Internals = InGeometryCollection.ModifyAttribute<bool>("Internal", FGeometryCollection::FacesGroup);
+			const TManagedArray<int32>& MaterialIDs = InGeometryCollection.GetAttribute<int32>("MaterialID", FGeometryCollection::FacesGroup);
+
+			int32 NumFaces = InGeometryCollection.NumElements(FGeometryCollection::FacesGroup);
+			for (int32 FaceIdx = 0; FaceIdx < NumFaces; ++FaceIdx)
+			{
+				if (MaterialIDs[FaceIdx] % 2 == 1)
+				{
+					Internals[FaceIdx] = true;
+				}
+				else
+				{
+					Internals[FaceIdx] = false;
+				}
+			}
+		}
+	}
+
 };
 
 
diff --git a/Engine/Source/Runtime/Experimental/Chaos/Private/GeometryCollection/GeometryCollectionClusteringUtility.cpp b/Engine/Source/Runtime/Experimental/Chaos/Private/GeometryCollection/GeometryCollectionClusteringUtility.cpp
--- a/Engine/Source/Runtime/Experimental/Chaos/Private/GeometryCollection/GeometryCollectionClusteringUtility.cpp
+++ b/Engine/Source/Runtime/Experimental/Chaos/Private/GeometryCollection/GeometryCollectionClusteringUtility.cpp
@@ -5,6 +5,7 @@
 #include "GeometryCollection/GeometryCollectionAlgo.h"
 #include "Containers/Set.h"
 #include "Async/ParallelFor.h"
+#include "GeometryCollection/Facades/CollectionHierarchyFacade.h"
 
 static int32 ChaosValidateResultsOfEditOperations = 0;
 static FAutoConsoleVariableRef CVarChaosStillCheckDistanceThreshold(TEXT("p.fracture.ValidateResultsOfEditOperations"), ChaosValidateResultsOfEditOperations, TEXT("When on this will enable result validation for fracture tool edit operations (can be slow for large geometry collection) [def:0]"));
@@ -683,6 +684,13 @@ void FGeometryCollectionClusteringUtility::UpdateHierarchyLevelOfChildren(FGeome
 	}
 }
 
+void FGeometryCollectionClusteringUtility::UpdateHierarchyLevelOfChildren(FManagedArrayCollection& InCollection, int32 ParentElement)
+{
+	Chaos::Facades::FCollectionHierarchyFacade HierarchyFacade(InCollection);
+
+	HierarchyFacade.GenerateLevelAttribute();
+}
+
 void FGeometryCollectionClusteringUtility::RecursivelyUpdateHierarchyLevelOfChildren(TManagedArray<int32>& Levels, const TManagedArray<TSet<int32>>& Children, int32 ParentElement)
 {
 	check(ParentElement < Levels.Num());
diff --git a/Engine/Source/Runtime/Experimental/Chaos/Public/GeometryCollection/Facades/CollectionExplodedVectorFacade.h b/Engine/Source/Runtime/Experimental/Chaos/Public/GeometryCollection/Facades/CollectionExplodedVectorFacade.h
new file mode 100644
--- /dev/null
+++ b/Engine/Source/Runtime/Experimental/Chaos/Public/GeometryCollection/Facades/CollectionExplodedVectorFacade.h
@@ -0,0 +1,40 @@
+// Copyright Epic Games, Inc. All Rights Reserved.
+
+#pragma once
+
+#include "GeometryCollection/ManagedArrayAccessor.h"
+#include "GeometryCollection/ManagedArrayCollection.h"
+
+
+namespace GeometryCollection::Facades
+{
+
+	class CHAOS_API FCollectionExplodedVectorFacade
+	{
+	public:
+		FCollectionExplodedVectorFacade(FManagedArrayCollection& InCollection);
+		FCollectionExplodedVectorFacade(const FManagedArrayCollection& InCollection);
+
+		/**
+		 * returns true if all the necessary attributes are present
+		 * if not then the API can be used to create
+		 */
+		bool IsValid() const;
+
+		/**
+		 * Add the necessary attributes if they are missing
+		 */
+		void DefineSchema();
+
+		/**
+		 * Adds exploded vector to global matrices translation 
+		 */
+		void UpdateGlobalMatricesWithExplodedVectors(TArray<FMatrix>& InOutGlobalMatrices) const;
+		void UpdateGlobalMatricesWithExplodedVectors(TArray<FTransform>& InOutGlobalTransforms) const;
+
+	private:
+
+		TManagedArrayAccessor<FVector3f> ExplodedVectorAttribute;
+	};
+
+}
\ No newline at end of file
diff --git a/Engine/Source/Runtime/Experimental/Chaos/Public/GeometryCollection/Facades/CollectionMeshFacade.h b/Engine/Source/Runtime/Experimental/Chaos/Public/GeometryCollection/Facades/CollectionMeshFacade.h
--- a/Engine/Source/Runtime/Experimental/Chaos/Public/GeometryCollection/Facades/CollectionMeshFacade.h
+++ b/Engine/Source/Runtime/Experimental/Chaos/Public/GeometryCollection/Facades/CollectionMeshFacade.h
@@ -4,6 +4,7 @@
 
 #include "GeometryCollection/ManagedArrayAccessor.h"
 #include "GeometryCollection/ManagedArrayCollection.h"
+#include "GeometryCollection/GeometryCollection.h"
 
 
 namespace GeometryCollection::Facades
@@ -41,6 +42,11 @@ namespace GeometryCollection::Facades
 		 */
 		const TArray<int32> GetFaceIndices(int32 BoneIdx) const;
 
+		/**
+		 * Adds and sets "Internal" attribute in FGeometryCollection::FacesGroup to designate internal faces
+		 */
+		static void AddInternalAttribute(FGeometryCollection& InGeometryCollection, const TArray<int32>& InMaterialID);
+
 		/**
 		 * Returns the vertex indices of the face for the bone
 		 */
diff --git a/Engine/Source/Runtime/Experimental/Chaos/Public/GeometryCollection/GeometryCollectionClusteringUtility.h b/Engine/Source/Runtime/Experimental/Chaos/Public/GeometryCollection/GeometryCollectionClusteringUtility.h
--- a/Engine/Source/Runtime/Experimental/Chaos/Public/GeometryCollection/GeometryCollectionClusteringUtility.h
+++ b/Engine/Source/Runtime/Experimental/Chaos/Public/GeometryCollection/GeometryCollectionClusteringUtility.h
@@ -93,6 +93,7 @@ public:
 	
 	/** Recursively update the hierarchy level of all the children below this bone */
 	static void UpdateHierarchyLevelOfChildren(FGeometryCollection* GeometryCollection, int32 ParentElement);
+	static void UpdateHierarchyLevelOfChildren(FManagedArrayCollection& InCollection, int32 ParentElement);
 
 	/** Collapse hierarchy at specified level */
 	static void CollapseLevelHierarchy(int8 Level, FGeometryCollection* GeometryCollection);
diff --git a/Engine/Source/Runtime/Experimental/Dataflow/Core/Private/Dataflow/DataflowSelection.cpp b/Engine/Source/Runtime/Experimental/Dataflow/Core/Private/Dataflow/DataflowSelection.cpp
--- a/Engine/Source/Runtime/Experimental/Dataflow/Core/Private/Dataflow/DataflowSelection.cpp
+++ b/Engine/Source/Runtime/Experimental/Dataflow/Core/Private/Dataflow/DataflowSelection.cpp
@@ -63,6 +63,19 @@ void FDataflowSelection::SetFromArray(const TArray<int32>& SelectionArr)
 	}
 }
 
+void FDataflowSelection::SetFromArray(const TArray<bool>& SelectionArr)
+{
+	SelectionArray.Init(false, SelectionArray.Num());
+
+	for (int32 Idx = 0; Idx < SelectionArr.Num(); ++Idx)
+	{
+		if (SelectionArr[Idx])
+		{
+			SetSelected(Idx);
+		}
+	}
+}
+
 void FDataflowSelection::AND(const FDataflowSelection& Other, FDataflowSelection& Result) const
 { 
 	Result.SelectionArray = TBitArray<>::BitwiseAND(SelectionArray, Other.SelectionArray, EBitwiseOperatorFlags::MaxSize);
diff --git a/Engine/Source/Runtime/Experimental/Dataflow/Core/Public/Dataflow/DataflowSelection.h b/Engine/Source/Runtime/Experimental/Dataflow/Core/Public/Dataflow/DataflowSelection.h
--- a/Engine/Source/Runtime/Experimental/Dataflow/Core/Public/Dataflow/DataflowSelection.h
+++ b/Engine/Source/Runtime/Experimental/Dataflow/Core/Public/Dataflow/DataflowSelection.h
@@ -19,7 +19,17 @@ struct DATAFLOWCORE_API FDataflowSelection
 	void SetNotSelected(int32 Idx) { SelectionArray[Idx] = false; }
 	void AsArray(TArray<int32>& SelectionArr) const;
 	TArray<int32> AsArray() const;
+	
+	/**
+	* Sets the selection from a sparse array (it only contains the indices of the selected items)
+	*/
 	void SetFromArray(const TArray<int32>& SelectionArr);
+	
+	/**
+	* Sets the selection from a dense array (it contains element for every item)
+	* (for example from the "Internal" attr from FacesGroup)
+	*/
+	void SetFromArray(const TArray<bool>& SelectionArr);
 	void AND(const FDataflowSelection& Other, FDataflowSelection& Result) const;
 	void OR(const FDataflowSelection& Other, FDataflowSelection& Result) const;
 	void XOR(const FDataflowSelection& Other, FDataflowSelection& Result) const;
