Dataflow:
- Added TransformCollection, TransformMesh, BakeTransformsInCollection, CollectionToMesh nodes
- Added rendering for FBox, FDynamicMesh3
- Bug fixes

diff --git a/Engine/Plugins/Experimental/Dataflow/Dataflow.uplugin b/Engine/Plugins/Experimental/Dataflow/Dataflow.uplugin
--- a/Engine/Plugins/Experimental/Dataflow/Dataflow.uplugin
+++ b/Engine/Plugins/Experimental/Dataflow/Dataflow.uplugin
@@ -36,6 +36,10 @@
 		{
 			"Name": "ProceduralMeshComponent",
 			"Enabled": true
+		},
+		{
+			"Name": "GeometryProcessing",
+			"Enabled": true
 		}
 	]
 }
diff --git a/Engine/Plugins/Experimental/Dataflow/Source/DataflowEnginePlugin/DataflowEnginePlugin.Build.cs b/Engine/Plugins/Experimental/Dataflow/Source/DataflowEnginePlugin/DataflowEnginePlugin.Build.cs
--- a/Engine/Plugins/Experimental/Dataflow/Source/DataflowEnginePlugin/DataflowEnginePlugin.Build.cs
+++ b/Engine/Plugins/Experimental/Dataflow/Source/DataflowEnginePlugin/DataflowEnginePlugin.Build.cs
@@ -12,6 +12,7 @@ namespace UnrealBuildTool.Rules
 				{
 					"Core",
 					"CoreUObject",
+					"GeometryFramework"
 				}
 			);
 			PrivateDependencyModuleNames.AddRange(
@@ -26,6 +27,8 @@ namespace UnrealBuildTool.Rules
 					"DataflowCore",
 					"DataflowEngine",
 					"Chaos",
+					"GeometryCore",
+					"DynamicMesh"
 				}
 			);
 		}
diff --git a/Engine/Plugins/Experimental/Dataflow/Source/DataflowEnginePlugin/Private/Dataflow/DataflowEngineRendering.cpp b/Engine/Plugins/Experimental/Dataflow/Source/DataflowEnginePlugin/Private/Dataflow/DataflowEngineRendering.cpp
--- a/Engine/Plugins/Experimental/Dataflow/Source/DataflowEnginePlugin/Private/Dataflow/DataflowEngineRendering.cpp
+++ b/Engine/Plugins/Experimental/Dataflow/Source/DataflowEnginePlugin/Private/Dataflow/DataflowEngineRendering.cpp
@@ -8,6 +8,8 @@
 #include "GeometryCollection/Facades/CollectionExplodedVectorFacade.h"
 #include "GeometryCollection/ManagedArrayCollection.h"
 #include "GeometryCollection/GeometryCollectionAlgo.h"
+#include "DynamicMesh/DynamicMesh3.h"
+#include "UDynamicMesh.h"
 
 namespace Dataflow
 {
@@ -17,7 +19,7 @@ namespace Dataflow
 		using namespace Dataflow;
 
 		/**
-		* DatflowNode (FGeometryCollection) Rendering
+		* DataflowNode (FGeometryCollection) Rendering
 		*
 		*		@param Type : FGeometryCollection::StaticType()
 
@@ -93,7 +95,105 @@ namespace Dataflow
 				}
 			});
 
-	}
 
+		/**
+		* DataflowNode (FDynamicMesh3) Rendering
+		*
+		*		@param Type : FName("FDynamicMesh3")
+
+		*		@param Outputs : {FDynamicMesh3 : "Mesh"}
+		*/
+		FRenderingFactory::GetInstance()->RegisterOutput(FName("FDynamicMesh3"),
+			[](GeometryCollection::Facades::FRenderingFacade& RenderCollection, const Dataflow::FGraphRenderingState& State)
+			{
+				if (State.GetRenderOutputs().Num())
+				{
+					FName PrimaryOutput = State.GetRenderOutputs()[0]; // "Mesh"
+
+					TObjectPtr<UDynamicMesh> Default;
+					if (const TObjectPtr<UDynamicMesh> Mesh = State.GetValue<TObjectPtr<UDynamicMesh>>(PrimaryOutput, Default))
+					{
+						const UE::Geometry::FDynamicMesh3& DynamicMesh = Mesh->GetMeshRef();
+
+						const int32 NumVertices = DynamicMesh.VertexCount();
+						const int32 NumTriangles = DynamicMesh.TriangleCount();
+
+						if (NumVertices > 0 && NumTriangles > 0)
+						{
+							// This will contain the valid triangles only
+							TArray<FIntVector> Tris; Tris.Reserve(DynamicMesh.TriangleCount());
+														
+							// DynamicMesh.TrianglesItr() returns the valid triangles only
+							for (UE::Geometry::FIndex3i Tri : DynamicMesh.TrianglesItr())
+							{
+								Tris.Add(FIntVector(Tri.A, Tri.B, Tri.C));
+							}
+
+							// This will contain all the vertices (invalid ones too)
+							// Otherwise the IDs need to be remaped
+							TArray<FVector3f> Vertices; Vertices.AddZeroed(DynamicMesh.MaxVertexID());
+
+							// DynamicMesh.VertexIndicesItr() returns the valid vertices only
+							for (int32 VertexID : DynamicMesh.VertexIndicesItr())
+							{
+								Vertices[VertexID] = (FVector3f)DynamicMesh.GetVertex(VertexID);
+							}
+
+							RenderCollection.AddSurface(MoveTemp(Vertices), MoveTemp(Tris));
+						}
+					}
+				}
+			});
+
+
+		/**
+		* DataflowNode (FBox) Rendering
+		*
+		*		@param Type : FName("FBox")
+
+		*		@param Outputs : {FBox : "Box"}
+		*/
+		FRenderingFactory::GetInstance()->RegisterOutput(FName("FBox"),
+			[](GeometryCollection::Facades::FRenderingFacade& RenderCollection, const Dataflow::FGraphRenderingState& State)
+			{
+				if (State.GetRenderOutputs().Num())
+				{
+					FName PrimaryOutput = State.GetRenderOutputs()[0]; // "Box"
+
+					FBox Default;
+					const FBox& Box = State.GetValue<FBox>(PrimaryOutput, Default);
+
+					const int32 NumVertices = 8;
+					const int32 NumTriangles = 12;
+
+					TArray<FVector3f> Vertices; Vertices.AddUninitialized(NumVertices);
+					TArray<FIntVector> Tris; Tris.AddUninitialized(NumTriangles);
+
+					FVector Min = Box.Min;
+					FVector Max = Box.Max;
+
+					// Add vertices
+					Vertices[0] = FVector3f(Min);
+					Vertices[1] = FVector3f(Max.X, Min.Y, Min.Z);
+					Vertices[2] = FVector3f(Max.X, Max.Y, Min.Z);
+					Vertices[3] = FVector3f(Min.X, Max.Y, Min.Z);
+					Vertices[4] = FVector3f(Min.X, Min.Y, Max.Z);
+					Vertices[5] = FVector3f(Max.X, Min.Y, Max.Z);
+					Vertices[6] = FVector3f(Max);
+					Vertices[7] = FVector3f(Min.X, Max.Y, Max.Z);
+
+					// Add triangles
+					Tris[0] = FIntVector(0, 1, 3); Tris[1] = FIntVector(1, 2, 3);
+					Tris[2] = FIntVector(0, 4, 1); Tris[3] = FIntVector(4, 5, 1);
+					Tris[4] = FIntVector(5, 2, 1); Tris[5] = FIntVector(5, 6, 2);
+					Tris[6] = FIntVector(3, 2, 6); Tris[7] = FIntVector(7, 3, 6);
+					Tris[8] = FIntVector(0, 3, 7); Tris[9] = FIntVector(4, 0, 7);
+					Tris[10] = FIntVector(5, 4, 7); Tris[11] = FIntVector(5, 7, 6);
+
+					RenderCollection.AddSurface(MoveTemp(Vertices), MoveTemp(Tris));
+				}
+			});
+
+	}
 
 }
diff --git a/Engine/Plugins/Experimental/Fracture/Source/FractureEngine/FractureEngine.Build.cs b/Engine/Plugins/Experimental/Fracture/Source/FractureEngine/FractureEngine.Build.cs
--- a/Engine/Plugins/Experimental/Fracture/Source/FractureEngine/FractureEngine.Build.cs
+++ b/Engine/Plugins/Experimental/Fracture/Source/FractureEngine/FractureEngine.Build.cs
@@ -17,6 +17,8 @@ namespace UnrealBuildTool.Rules
 					"Core",
 					"DataflowCore",
 					"PlanarCut",
+					"GeometryCore",
+					"DynamicMesh",
 					"Voronoi"
 				}
 				);
diff --git a/Engine/Plugins/Experimental/Fracture/Source/FractureEngine/Private/FractureEngineUtility.cpp b/Engine/Plugins/Experimental/Fracture/Source/FractureEngine/Private/FractureEngineUtility.cpp
--- a/Engine/Plugins/Experimental/Fracture/Source/FractureEngine/Private/FractureEngineUtility.cpp
+++ b/Engine/Plugins/Experimental/Fracture/Source/FractureEngine/Private/FractureEngineUtility.cpp
@@ -1,10 +1,78 @@
 // Copyright Epic Games, Inc. All Rights Reserved.
 
 #include "FractureEngineUtility.h"
-#include "GeometryCollection/GeometryCollectionClusteringUtility.h"
-#include "GeometryCollection/GeometryCollectionAlgo.h"
-#include "GeometryCollection/GeometryCollectionProximityUtility.h"
 
 
+void FFractureEngineUtility::ConvertBoxToVertexAndTriangleData(const FBox& InBox, TArray<FVector3f>& OutVertices, TArray<FIntVector>& OutTriangles)
+{
+	const int32 NumVertices = 8;
+	const int32 NumTriangles = 12;
 
+	OutVertices.AddUninitialized(NumVertices);
+	OutTriangles.AddUninitialized(NumTriangles);
 
+	FVector Min = InBox.Min;
+	FVector Max = InBox.Max;
+
+	// Add vertices
+	OutVertices[0] = FVector3f(Min);
+	OutVertices[1] = FVector3f(Max.X, Min.Y, Min.Z);
+	OutVertices[2] = FVector3f(Max.X, Max.Y, Min.Z);
+	OutVertices[3] = FVector3f(Min.X, Max.Y, Min.Z);
+	OutVertices[4] = FVector3f(Min.X, Min.Y, Max.Z);
+	OutVertices[5] = FVector3f(Max.X, Min.Y, Max.Z);
+	OutVertices[6] = FVector3f(Max);
+	OutVertices[7] = FVector3f(Min.X, Max.Y, Max.Z);
+
+	// Add triangles
+	OutTriangles[0] = FIntVector(0, 1, 3); OutTriangles[1] = FIntVector(1, 2, 3);
+	OutTriangles[2] = FIntVector(0, 4, 1); OutTriangles[3] = FIntVector(4, 5, 1);
+	OutTriangles[4] = FIntVector(5, 2, 1); OutTriangles[5] = FIntVector(5, 6, 2);
+	OutTriangles[6] = FIntVector(3, 2, 6); OutTriangles[7] = FIntVector(7, 3, 6);
+	OutTriangles[8] = FIntVector(0, 3, 7); OutTriangles[9] = FIntVector(4, 0, 7);
+	OutTriangles[10] = FIntVector(5, 4, 7); OutTriangles[11] = FIntVector(5, 7, 6);
+}
+
+
+void FFractureEngineUtility::ConstructMesh(UE::Geometry::FDynamicMesh3& OutMesh, const TArray<FVector3f>& InVertices, const TArray<FIntVector>& InTriangles)
+{
+	for (int32 VertexIdx = 0; VertexIdx < InVertices.Num(); ++VertexIdx)
+	{
+		OutMesh.AppendVertex(FVector(InVertices[VertexIdx]));
+	}
+
+	int GroupID = 0;
+	for (int32 TriangleIdx = 0; TriangleIdx < InTriangles.Num(); ++TriangleIdx)
+	{
+		OutMesh.AppendTriangle(InTriangles[TriangleIdx].X, InTriangles[TriangleIdx].Y, InTriangles[TriangleIdx].Z, GroupID);
+	}
+}
+
+
+void FFractureEngineUtility::DeconstructMesh(const UE::Geometry::FDynamicMesh3& InMesh, TArray<FVector3f>& OutVertices, TArray<FIntVector>& OutTriangles)
+{
+	const int32 NumVertices = InMesh.VertexCount();
+	const int32 NumTriangles = InMesh.TriangleCount();
+
+	if (NumVertices > 0 && NumTriangles > 0)
+	{
+		// This will contain the valid triangles only
+		OutTriangles.Reserve(InMesh.TriangleCount());
+
+		// DynamicMesh.TrianglesItr() returns the valid triangles only
+		for (UE::Geometry::FIndex3i Tri : InMesh.TrianglesItr())
+		{
+			OutTriangles.Add(FIntVector(Tri.A, Tri.B, Tri.C));
+		}
+
+		// This will contain all the vertices (invalid ones too)
+		// Otherwise the IDs need to be remaped
+		OutVertices.AddZeroed(InMesh.MaxVertexID());
+
+		// DynamicMesh.VertexIndicesItr() returns the valid vertices only
+		for (int32 VertexID : InMesh.VertexIndicesItr())
+		{
+			OutVertices[VertexID] = (FVector3f)InMesh.GetVertex(VertexID);
+		}
+	}
+}
diff --git a/Engine/Plugins/Experimental/Fracture/Source/FractureEngine/Public/FractureEngineUtility.h b/Engine/Plugins/Experimental/Fracture/Source/FractureEngine/Public/FractureEngineUtility.h
--- a/Engine/Plugins/Experimental/Fracture/Source/FractureEngine/Public/FractureEngineUtility.h
+++ b/Engine/Plugins/Experimental/Fracture/Source/FractureEngine/Public/FractureEngineUtility.h
@@ -4,10 +4,26 @@
 
 #include "Algo/Count.h"
 #include "GeometryCollection/GeometryCollection.h"
+#include "DynamicMesh/DynamicMesh3.h"
 
+//class FDynamicMesh3;
 
 class FRACTUREENGINE_API FFractureEngineUtility
 {
 public:
+	/**
+	* Outputs the vertex and triangle data of a FBox into OutVertices and OutTriangles
+	*/
+	static void ConvertBoxToVertexAndTriangleData(const FBox& InBox, TArray<FVector3f>& OutVertices, TArray<FIntVector>& OutTriangles);
+
+	/**
+	* Creates a mesh from vertex and triangle data
+	*/
+	static void ConstructMesh(UE::Geometry::FDynamicMesh3& OutMesh, const TArray<FVector3f>& InVertices, const TArray<FIntVector>& InTriangles);
+
+	/** 
+	* Outputs the vertex and triangle data of a mesh into OutVertices and OutTriangles
+	*/
+	static void DeconstructMesh(const UE::Geometry::FDynamicMesh3& InMesh, TArray<FVector3f>& OutVertices, TArray<FIntVector>& OutTriangles);
 
 };
\ No newline at end of file
diff --git a/Engine/Plugins/Experimental/GeometryCollectionPlugin/GeometryCollectionPlugin.uplugin b/Engine/Plugins/Experimental/GeometryCollectionPlugin/GeometryCollectionPlugin.uplugin
--- a/Engine/Plugins/Experimental/GeometryCollectionPlugin/GeometryCollectionPlugin.uplugin
+++ b/Engine/Plugins/Experimental/GeometryCollectionPlugin/GeometryCollectionPlugin.uplugin
@@ -53,6 +53,10 @@
 			"Name": "PlanarCut",
 			"Enabled": true
 		},
+		{
+			"Name": "GeometryProcessing",
+			"Enabled": true
+		},
 		{
 			"Name": "Fracture",
 			"Enabled": true
diff --git a/Engine/Plugins/Experimental/GeometryCollectionPlugin/Source/GeometryCollectionNodes/Private/Dataflow/GeometryCollectionMeshNodes.cpp b/Engine/Plugins/Experimental/GeometryCollectionPlugin/Source/GeometryCollectionNodes/Private/Dataflow/GeometryCollectionMeshNodes.cpp
--- a/Engine/Plugins/Experimental/GeometryCollectionPlugin/Source/GeometryCollectionNodes/Private/Dataflow/GeometryCollectionMeshNodes.cpp
+++ b/Engine/Plugins/Experimental/GeometryCollectionPlugin/Source/GeometryCollectionNodes/Private/Dataflow/GeometryCollectionMeshNodes.cpp
@@ -27,6 +27,7 @@
 #include "GeometryCollection/GeometryCollectionProximityUtility.h"
 #include "FractureEngineClustering.h"
 #include "FractureEngineSelection.h"
+#include "FractureEngineUtility.h"
 #include "DynamicMesh/DynamicMesh3.h"
 #include "UDynamicMesh.h"
 
@@ -43,6 +44,7 @@ namespace Dataflow
 		DATAFLOW_NODE_REGISTER_CREATION_FACTORY(FBoxToMeshDataflowNode);
 		DATAFLOW_NODE_REGISTER_CREATION_FACTORY(FMeshInfoDataflowNode);
 		DATAFLOW_NODE_REGISTER_CREATION_FACTORY(FMeshToCollectionDataflowNode);
+		DATAFLOW_NODE_REGISTER_CREATION_FACTORY(FCollectionToMeshDataflowNode);
 		DATAFLOW_NODE_REGISTER_CREATION_FACTORY(FStaticMeshToMeshDataflowNode);
 		DATAFLOW_NODE_REGISTER_CREATION_FACTORY(FMeshAppendDataflowNode);
 		DATAFLOW_NODE_REGISTER_CREATION_FACTORY(FMeshBooleanDataflowNode);
@@ -84,105 +86,150 @@ void FPointsToMeshDataflowNode::Evaluate(Dataflow::FContext& Context, const FDat
 	}
 }
 
+
 void FBoxToMeshDataflowNode::Evaluate(Dataflow::FContext& Context, const FDataflowOutput* Out) const
 {
 	if (Out->IsA<TObjectPtr<UDynamicMesh>>(&Mesh) || Out->IsA<int32>(&TriangleCount))
 	{
-		TObjectPtr<UDynamicMesh> DynamicMesh = NewObject<UDynamicMesh>();
-		DynamicMesh->Reset();
+		TObjectPtr<UDynamicMesh> NewMesh = NewObject<UDynamicMesh>();
+		NewMesh->Reset();
 
-		UE::Geometry::FDynamicMesh3& DynMesh = DynamicMesh->GetMeshRef();
+		UE::Geometry::FDynamicMesh3& DynMesh = NewMesh->GetMeshRef();
 
 		FBox InBox = GetValue<FBox>(Context, &Box);
-		FVector Min = InBox.Min;
-		FVector Max = InBox.Max;
-
-		// Add vertices
-		int Vertex0 = DynMesh.AppendVertex(Min);
-		int Vertex1 = DynMesh.AppendVertex(FVector(Min.X, Max.Y, Min.Z));
-		int Vertex2 = DynMesh.AppendVertex(FVector(Max.X, Max.Y, Min.Z));
-		int Vertex3 = DynMesh.AppendVertex(FVector(Max.X, Min.Y, Min.Z));
-		int Vertex4 = DynMesh.AppendVertex(FVector(Min.X, Min.Y, Max.Z));
-		int Vertex5 = DynMesh.AppendVertex(FVector(Min.X, Max.Y, Max.Z));
-		int Vertex6 = DynMesh.AppendVertex(Max);
-		int Vertex7 = DynMesh.AppendVertex(FVector(Max.X, Min.Y, Max.Z));
-
-		// Add triangles
-		int GroupID = 0;
-		DynMesh.AppendTriangle(Vertex0, Vertex1, Vertex3, GroupID);
-		DynMesh.AppendTriangle(Vertex1, Vertex2, Vertex3, GroupID);
-		DynMesh.AppendTriangle(Vertex3, Vertex6, Vertex7, GroupID);
-		DynMesh.AppendTriangle(Vertex3, Vertex2, Vertex6, GroupID);
-		DynMesh.AppendTriangle(Vertex7, Vertex4, Vertex0, GroupID);
-		DynMesh.AppendTriangle(Vertex0, Vertex3, Vertex7, GroupID);
-		DynMesh.AppendTriangle(Vertex0, Vertex4, Vertex5, GroupID);
-		DynMesh.AppendTriangle(Vertex0, Vertex5, Vertex1, GroupID);
-		DynMesh.AppendTriangle(Vertex1, Vertex5, Vertex6, GroupID);
-		DynMesh.AppendTriangle(Vertex6, Vertex2, Vertex1, GroupID);
-		DynMesh.AppendTriangle(Vertex4, Vertex6, Vertex5, GroupID);
-		DynMesh.AppendTriangle(Vertex4, Vertex7, Vertex6, GroupID);
-
-		SetValue<TObjectPtr<UDynamicMesh>>(Context, DynamicMesh, &Mesh);
-		SetValue<int32>(Context, DynamicMesh->GetTriangleCount(), &TriangleCount);
+
+		TArray<FVector3f> Vertices;
+		TArray<FIntVector> Triangles;
+
+		FFractureEngineUtility::ConvertBoxToVertexAndTriangleData(InBox, Vertices, Triangles);
+		FFractureEngineUtility::ConstructMesh(DynMesh, Vertices, Triangles);
+
+		SetValue<TObjectPtr<UDynamicMesh>>(Context, NewMesh, &Mesh);
+		SetValue<int32>(Context, NewMesh->GetTriangleCount(), &TriangleCount);
 	}
 }
 
+
 void FMeshInfoDataflowNode::Evaluate(Dataflow::FContext& Context, const FDataflowOutput* Out) const
 {
 	if (Out->IsA<FString>(&InfoString))
 	{
-		UE::Geometry::FDynamicMesh3& DynMesh = (GetValue<TObjectPtr<UDynamicMesh>>(Context, &Mesh))->GetMeshRef();
+		if (const TObjectPtr<const UDynamicMesh> InMesh = GetValue(Context, &Mesh))
+		{
+			const UE::Geometry::FDynamicMesh3& DynMesh = InMesh->GetMeshRef();
 
-		SetValue<FString>(Context, DynMesh.MeshInfoString(), &InfoString);
+			SetValue<FString>(Context, DynMesh.MeshInfoString(), &InfoString);
+		}
+		else
+		{
+			SetValue<FString>(Context, FString(""), &InfoString);
+		}
 	}
 }
 
+
 void FMeshToCollectionDataflowNode::Evaluate(Dataflow::FContext& Context, const FDataflowOutput* Out) const
 {
 	if (Out->IsA<FManagedArrayCollection>(&Collection))
 	{
-		const UE::Geometry::FDynamicMesh3& DynMesh = (GetValue<TObjectPtr<UDynamicMesh>>(Context, &Mesh))->GetMeshRef();
-		if (DynMesh.VertexCount() > 0)
+		if (const TObjectPtr<const UDynamicMesh> InMesh = GetValue(Context, &Mesh))
 		{
-			FMeshDescription MeshDescription;
-			FStaticMeshAttributes Attributes(MeshDescription);
-			Attributes.Register();
+			const UE::Geometry::FDynamicMesh3& DynMesh = InMesh->GetMeshRef();
 
-			FDynamicMeshToMeshDescription Converter;
-			Converter.Convert(&DynMesh, MeshDescription, true);
+			if (DynMesh.VertexCount() > 0)
+			{
+				FMeshDescription MeshDescription;
+				FStaticMeshAttributes Attributes(MeshDescription);
+				Attributes.Register();
 
-			FGeometryCollection NewGeometryCollection = FGeometryCollection();
-			FGeometryCollectionEngineConversion::AppendMeshDescription(&MeshDescription, FString(TEXT("TEST")), 0, FTransform().Identity, &NewGeometryCollection);
+				FDynamicMeshToMeshDescription Converter;
+				Converter.Convert(&DynMesh, MeshDescription, true);
 
-			FManagedArrayCollection Col = FManagedArrayCollection();
-			NewGeometryCollection.CopyTo(&Col);
+				FGeometryCollection NewGeometryCollection = FGeometryCollection();
+				FGeometryCollectionEngineConversion::AppendMeshDescription(&MeshDescription, FString(TEXT("TEST")), 0, FTransform().Identity, &NewGeometryCollection);
 
-			SetValue<FManagedArrayCollection>(Context, Col, &Collection);
+				FManagedArrayCollection NewCollection = FManagedArrayCollection();
+				NewGeometryCollection.CopyTo(&NewCollection);
+
+				SetValue<FManagedArrayCollection>(Context, NewCollection, &Collection);
+
+				return;
+			}
 		}
-		else
+
+		SetValue<FManagedArrayCollection>(Context, FManagedArrayCollection(), &Collection);
+	}
+}
+
+
+void FCollectionToMeshDataflowNode::Evaluate(Dataflow::FContext& Context, const FDataflowOutput* Out) const
+{
+#if WITH_EDITORONLY_DATA
+	if (Out->IsA<TObjectPtr<UDynamicMesh>>(&Mesh))
+	{
+		const FManagedArrayCollection& InCollection = GetValue<FManagedArrayCollection>(Context, &Collection);
+		
+		if (InCollection.NumElements(FGeometryCollection::TransformGroup) > 0)
 		{
-			SetValue<FManagedArrayCollection>(Context, FManagedArrayCollection(), &Collection);
+			if (TUniquePtr<FGeometryCollection> GeomCollection = TUniquePtr<FGeometryCollection>(InCollection.NewCopy<FGeometryCollection>()))
+			{
+				const TManagedArray<FTransform>& BoneTransforms = InCollection.GetAttribute<FTransform>("Transform", FGeometryCollection::TransformGroup);
+
+				TArray<int32> TransformIndices;
+				TransformIndices.AddUninitialized(BoneTransforms.Num());
+
+				int32 Idx = 0;
+				for (int32& TransformIdx : TransformIndices)
+				{
+					TransformIdx = Idx++;
+				}
+
+				FMeshDescription MeshDescription;
+				FStaticMeshAttributes Attributes(MeshDescription);
+				Attributes.Register();
+
+				FTransform TransformOut;
+
+				ConvertToMeshDescription(MeshDescription, TransformOut, bCenterPivot, *GeomCollection, BoneTransforms, TransformIndices);
+
+				TObjectPtr<UDynamicMesh> NewMesh = NewObject<UDynamicMesh>();
+				NewMesh->Reset();
+
+				UE::Geometry::FDynamicMesh3& DynMesh = NewMesh->GetMeshRef();
+				{
+					FMeshDescriptionToDynamicMesh ConverterToDynamicMesh;
+					ConverterToDynamicMesh.Convert(&MeshDescription, DynMesh);
+				}
+
+				SetValue<TObjectPtr<UDynamicMesh>>(Context, NewMesh, &Mesh);
+
+				return;
+			}
 		}
+
+		SetValue<TObjectPtr<UDynamicMesh>>(Context, NewObject<UDynamicMesh>(), &Mesh);
 	}
+#endif
 }
 
+
 void FStaticMeshToMeshDataflowNode::Evaluate(Dataflow::FContext& Context, const FDataflowOutput* Out) const
 {
 #if WITH_EDITORONLY_DATA
 	if (Out->IsA<TObjectPtr<UDynamicMesh>>(&Mesh))
 	{
-		if (FMeshDescription* MeshDescription = UseHiRes ? StaticMesh->GetHiResMeshDescription() : StaticMesh->GetMeshDescription(LODLevel))
+		if (FMeshDescription* MeshDescription = bUseHiRes ? StaticMesh->GetHiResMeshDescription() : StaticMesh->GetMeshDescription(LODLevel))
 		{
-			TObjectPtr<UDynamicMesh> DynamicMesh = NewObject<UDynamicMesh>();
-			DynamicMesh->Reset();
+			TObjectPtr<UDynamicMesh> NewMesh = NewObject<UDynamicMesh>();
+			NewMesh->Reset();
 
-			UE::Geometry::FDynamicMesh3& DynMesh = DynamicMesh->GetMeshRef();
+			UE::Geometry::FDynamicMesh3& DynMesh = NewMesh->GetMeshRef();
 			{
 				FMeshDescriptionToDynamicMesh ConverterToDynamicMesh;
 				ConverterToDynamicMesh.Convert(MeshDescription, DynMesh);
 			}
 
-			SetValue<TObjectPtr<UDynamicMesh>>(Context, DynamicMesh, &Mesh);
+			SetValue<TObjectPtr<UDynamicMesh>>(Context, NewMesh, &Mesh);
 		}
 		else
 		{
@@ -197,78 +244,89 @@ void FMeshAppendDataflowNode::Evaluate(Dataflow::FContext& Context, const FDataf
 {
 	if (Out->IsA<TObjectPtr<UDynamicMesh>>(&Mesh))
 	{
-		const UE::Geometry::FDynamicMesh3& DynMesh1 = (GetValue<TObjectPtr<UDynamicMesh>>(Context, &Mesh1))->GetMeshRef();
-		const UE::Geometry::FDynamicMesh3& DynMesh2 = (GetValue<TObjectPtr<UDynamicMesh>>(Context, &Mesh2))->GetMeshRef();
-
-		if (DynMesh1.VertexCount() > 0 || DynMesh2.VertexCount() > 0)
+		if (TObjectPtr<UDynamicMesh> InMesh1 = GetValue<TObjectPtr<UDynamicMesh>>(Context, &Mesh1))
 		{
-			TObjectPtr<UDynamicMesh> DynamicMesh = NewObject<UDynamicMesh>();
-			DynamicMesh->Reset();
+			if (TObjectPtr<UDynamicMesh> InMesh2 = GetValue<TObjectPtr<UDynamicMesh>>(Context, &Mesh2))
+			{
+				const UE::Geometry::FDynamicMesh3& DynMesh1 = InMesh1->GetMeshRef();
+				const UE::Geometry::FDynamicMesh3& DynMesh2 = InMesh2->GetMeshRef();
 
-			UE::Geometry::FDynamicMesh3& ResultDynMesh = DynamicMesh->GetMeshRef();
+				if (DynMesh1.VertexCount() > 0 || DynMesh2.VertexCount() > 0)
+				{
+					TObjectPtr<UDynamicMesh> NewMesh = NewObject<UDynamicMesh>();
+					NewMesh->Reset();
 
-			UE::Geometry::FDynamicMeshEditor MeshEditor(&ResultDynMesh);
+					UE::Geometry::FDynamicMesh3& ResultDynMesh = NewMesh->GetMeshRef();
 
-			UE::Geometry::FMeshIndexMappings IndexMaps1;
-			MeshEditor.AppendMesh(&DynMesh1, IndexMaps1);
+					UE::Geometry::FDynamicMeshEditor MeshEditor(&ResultDynMesh);
 
-			UE::Geometry::FMeshIndexMappings IndexMaps2;
-			MeshEditor.AppendMesh(&DynMesh2, IndexMaps2);
+					UE::Geometry::FMeshIndexMappings IndexMaps1;
+					MeshEditor.AppendMesh(&DynMesh1, IndexMaps1);
 
-			SetValue<TObjectPtr<UDynamicMesh>>(Context, DynamicMesh, &Mesh);
-		}
-		else
-		{
-			SetValue<TObjectPtr<UDynamicMesh>>(Context, NewObject<UDynamicMesh>(), &Mesh);
+					UE::Geometry::FMeshIndexMappings IndexMaps2;
+					MeshEditor.AppendMesh(&DynMesh2, IndexMaps2);
+
+					SetValue<TObjectPtr<UDynamicMesh>>(Context, NewMesh, &Mesh);
+
+					return;
+				}
+			}
 		}
+
+		SetValue<TObjectPtr<UDynamicMesh>>(Context, NewObject<UDynamicMesh>(), &Mesh);
 	}
 }
 
+
 void FMeshBooleanDataflowNode::Evaluate(Dataflow::FContext& Context, const FDataflowOutput* Out) const
 {
 	if (Out->IsA<TObjectPtr<UDynamicMesh>>(&Mesh))
 	{
-		const UE::Geometry::FDynamicMesh3& DynMesh1 = (GetValue<TObjectPtr<UDynamicMesh>>(Context, &Mesh1))->GetMeshRef();
-		const UE::Geometry::FDynamicMesh3& DynMesh2 = (GetValue<TObjectPtr<UDynamicMesh>>(Context, &Mesh2))->GetMeshRef();
-
-		if (DynMesh1.VertexCount() > 0 && DynMesh2.VertexCount() > 0)
+		if (TObjectPtr<UDynamicMesh> InMesh1 = GetValue<TObjectPtr<UDynamicMesh>>(Context, &Mesh1))
 		{
-			// Get output
-			TObjectPtr<UDynamicMesh> DynamicMesh = NewObject<UDynamicMesh>();
-			DynamicMesh->Reset();
-			UE::Geometry::FDynamicMesh3& ResultDynMesh = DynamicMesh->GetMeshRef();
-
-			UE::Geometry::FMeshBoolean::EBooleanOp BoolOp = UE::Geometry::FMeshBoolean::EBooleanOp::Intersect;
-			if (Operation == EMeshBooleanOperationEnum::Dataflow_MeshBoolean_Intersect)
-			{
-				BoolOp = UE::Geometry::FMeshBoolean::EBooleanOp::Intersect;
-			}
-			else if (Operation == EMeshBooleanOperationEnum::Dataflow_MeshBoolean_Union)
+			if (TObjectPtr<UDynamicMesh> InMesh2 = GetValue<TObjectPtr<UDynamicMesh>>(Context, &Mesh2))
 			{
-				BoolOp = UE::Geometry::FMeshBoolean::EBooleanOp::Union;
-			}
-			else if (Operation == EMeshBooleanOperationEnum::Dataflow_MeshBoolean_Difference)
-			{
-				BoolOp = UE::Geometry::FMeshBoolean::EBooleanOp::Difference;
-			}
+				const UE::Geometry::FDynamicMesh3& DynMesh1 = InMesh1->GetMeshRef();
+				const UE::Geometry::FDynamicMesh3& DynMesh2 = InMesh2->GetMeshRef();
 
-			UE::Geometry::FMeshBoolean Boolean(&DynMesh1, &DynMesh2, &ResultDynMesh, BoolOp);
-			Boolean.bSimplifyAlongNewEdges = true;
-			Boolean.PreserveUVsOnlyForMesh = 0; // slight warping of the autogenerated cell UVs generally doesn't matter
-			Boolean.bWeldSharedEdges = false;
-			Boolean.bTrackAllNewEdges = true;
-			if (!Boolean.Compute())
-			{
-				SetValue<TObjectPtr<UDynamicMesh>>(Context, NewObject<UDynamicMesh>(), &Mesh);
-				return;
+				if (DynMesh1.VertexCount() > 0 && DynMesh2.VertexCount() > 0)
+				{
+					// Get output
+					TObjectPtr<UDynamicMesh> NewMesh = NewObject<UDynamicMesh>();
+					NewMesh->Reset();
+					UE::Geometry::FDynamicMesh3& ResultDynMesh = NewMesh->GetMeshRef();
+
+					UE::Geometry::FMeshBoolean::EBooleanOp BoolOp = UE::Geometry::FMeshBoolean::EBooleanOp::Intersect;
+					if (Operation == EMeshBooleanOperationEnum::Dataflow_MeshBoolean_Intersect)
+					{
+						BoolOp = UE::Geometry::FMeshBoolean::EBooleanOp::Intersect;
+					}
+					else if (Operation == EMeshBooleanOperationEnum::Dataflow_MeshBoolean_Union)
+					{
+						BoolOp = UE::Geometry::FMeshBoolean::EBooleanOp::Union;
+					}
+					else if (Operation == EMeshBooleanOperationEnum::Dataflow_MeshBoolean_Difference)
+					{
+						BoolOp = UE::Geometry::FMeshBoolean::EBooleanOp::Difference;
+					}
+
+					UE::Geometry::FMeshBoolean Boolean(&DynMesh1, &DynMesh2, &ResultDynMesh, BoolOp);
+					Boolean.bSimplifyAlongNewEdges = true;
+					Boolean.PreserveUVsOnlyForMesh = 0; // slight warping of the autogenerated cell UVs generally doesn't matter
+					Boolean.bWeldSharedEdges = false;
+					Boolean.bTrackAllNewEdges = true;
+
+					if (Boolean.Compute())
+					{
+						SetValue<TObjectPtr<UDynamicMesh>>(Context, NewMesh, &Mesh);
+
+						return;
+					}
+				}
 			}
-
-			SetValue<TObjectPtr<UDynamicMesh>>(Context, DynamicMesh, &Mesh);
-		}
-		else
-		{
-			SetValue<TObjectPtr<UDynamicMesh>>(Context, NewObject<UDynamicMesh>(), &Mesh);
 		}
+
+		SetValue<TObjectPtr<UDynamicMesh>>(Context, NewObject<UDynamicMesh>(), &Mesh);
 	}
 }
 
@@ -276,41 +334,45 @@ void FMeshCopyToPointsDataflowNode::Evaluate(Dataflow::FContext& Context, const
 {
 	if (Out->IsA<TObjectPtr<UDynamicMesh>>(&Mesh))
 	{
-		const TArray<FVector>& InPoints = GetValue<TArray<FVector>>(Context, &Points);
-		const UE::Geometry::FDynamicMesh3& InDynMeshToCopy = (GetValue<TObjectPtr<UDynamicMesh>>(Context, &MeshToCopy))->GetMeshRef();
-
-		if (InPoints.Num() > 0 && InDynMeshToCopy.VertexCount() > 0)
+		if (TObjectPtr<UDynamicMesh> InMeshToCopy = GetValue<TObjectPtr<UDynamicMesh>>(Context, &MeshToCopy))
 		{
-			TObjectPtr<UDynamicMesh> DynamicMesh = NewObject<UDynamicMesh>();
-			DynamicMesh->Reset();
+			const UE::Geometry::FDynamicMesh3& InDynMeshToCopy = InMeshToCopy->GetMeshRef();
 
-			UE::Geometry::FDynamicMesh3& ResultDynMesh = DynamicMesh->GetMeshRef();
-			UE::Geometry::FDynamicMeshEditor MeshEditor(&ResultDynMesh);
+			const TArray<FVector>& InPoints = GetValue<TArray<FVector>>(Context, &Points);
 
-			for (auto& Point : InPoints)
+			if (InPoints.Num() > 0 && InDynMeshToCopy.VertexCount() > 0)
 			{
-				UE::Geometry::FDynamicMesh3 DynMeshTemp(InDynMeshToCopy);
-				UE::Geometry::FRefCountVector VertexRefCounts = DynMeshTemp.GetVerticesRefCounts();
+				TObjectPtr<UDynamicMesh> NewMesh = NewObject<UDynamicMesh>();
+				NewMesh->Reset();
 
-				UE::Geometry::FRefCountVector::IndexIterator ItVertexID = VertexRefCounts.BeginIndices();
-				const UE::Geometry::FRefCountVector::IndexIterator ItEndVertexID = VertexRefCounts.EndIndices();
+				UE::Geometry::FDynamicMesh3& ResultDynMesh = NewMesh->GetMeshRef();
+				UE::Geometry::FDynamicMeshEditor MeshEditor(&ResultDynMesh);
 
-				while (ItVertexID != ItEndVertexID)
+				for (auto& Point : InPoints)
 				{
-					DynMeshTemp.SetVertex(*ItVertexID, Scale * DynMeshTemp.GetVertex(*ItVertexID) + Point);
-					++ItVertexID;
+					UE::Geometry::FDynamicMesh3 DynMeshTemp(InDynMeshToCopy);
+					UE::Geometry::FRefCountVector VertexRefCounts = DynMeshTemp.GetVerticesRefCounts();
+
+					UE::Geometry::FRefCountVector::IndexIterator ItVertexID = VertexRefCounts.BeginIndices();
+					const UE::Geometry::FRefCountVector::IndexIterator ItEndVertexID = VertexRefCounts.EndIndices();
+
+					while (ItVertexID != ItEndVertexID)
+					{
+						DynMeshTemp.SetVertex(*ItVertexID, Scale * DynMeshTemp.GetVertex(*ItVertexID) + Point);
+						++ItVertexID;
+					}
+
+					UE::Geometry::FMeshIndexMappings IndexMaps;
+					MeshEditor.AppendMesh(&DynMeshTemp, IndexMaps);
 				}
 
-				UE::Geometry::FMeshIndexMappings IndexMaps;
-				MeshEditor.AppendMesh(&DynMeshTemp, IndexMaps);
-			}
+				SetValue<TObjectPtr<UDynamicMesh>>(Context, NewMesh, &Mesh);
 
-			SetValue<TObjectPtr<UDynamicMesh>>(Context, DynamicMesh, &Mesh);
-		}
-		else
-		{
-			SetValue<TObjectPtr<UDynamicMesh>>(Context, NewObject<UDynamicMesh>(), &Mesh);
+				return;
+			}
 		}
+
+		SetValue<TObjectPtr<UDynamicMesh>>(Context, NewObject<UDynamicMesh>(), &Mesh);
 	}
 }
 
@@ -319,24 +381,36 @@ void FGetMeshDataDataflowNode::Evaluate(Dataflow::FContext& Context, const FData
 {
 	if (Out->IsA<int32>(&VertexCount))
 	{
-		TObjectPtr<UDynamicMesh> DynamicMesh = GetValue<TObjectPtr<UDynamicMesh>>(Context, &Mesh);
-		const UE::Geometry::FDynamicMesh3& DynMesh = DynamicMesh->GetMeshRef();
-
-		SetValue<int32>(Context, DynMesh.VertexCount(), &VertexCount);
+		if (const TObjectPtr<const UDynamicMesh> InMesh = GetValue<TObjectPtr<UDynamicMesh>>(Context, &Mesh))
+		{
+			SetValue<int32>(Context, InMesh->GetMeshRef().VertexCount(), &VertexCount);
+		}
+		else
+		{
+			SetValue<int32>(Context, 0, &VertexCount);
+		}
 	}
 	else if (Out->IsA<int32>(&EdgeCount))
 	{
-		TObjectPtr<UDynamicMesh> DynamicMesh = GetValue<TObjectPtr<UDynamicMesh>>(Context, &Mesh);
-		const UE::Geometry::FDynamicMesh3& DynMesh = DynamicMesh->GetMeshRef();
-
-		SetValue<int32>(Context, DynMesh.EdgeCount(), &EdgeCount);
+		if (const TObjectPtr<const UDynamicMesh> InMesh = GetValue<TObjectPtr<UDynamicMesh>>(Context, &Mesh))
+		{
+			SetValue<int32>(Context, InMesh->GetMeshRef().EdgeCount(), &EdgeCount);
+		}
+		else
+		{
+			SetValue<int32>(Context, 0, &EdgeCount);
+		}
 	}
 	else if (Out->IsA<int32>(&TriangleCount))
 	{
-		TObjectPtr<UDynamicMesh> DynamicMesh = GetValue<TObjectPtr<UDynamicMesh>>(Context, &Mesh);
-		const UE::Geometry::FDynamicMesh3& DynMesh = DynamicMesh->GetMeshRef();
-
-		SetValue<int32>(Context, DynMesh.TriangleCount(), &TriangleCount);
+		if (const TObjectPtr<const UDynamicMesh> InMesh = GetValue<TObjectPtr<UDynamicMesh>>(Context, &Mesh))
+		{
+			SetValue<int32>(Context, InMesh->GetMeshRef().TriangleCount(), &TriangleCount);
+		}
+		else
+		{
+			SetValue<int32>(Context, 0, &TriangleCount);
+		}
 	}
 }
 
diff --git a/Engine/Plugins/Experimental/GeometryCollectionPlugin/Source/GeometryCollectionNodes/Private/Dataflow/GeometryCollectionNodes.cpp b/Engine/Plugins/Experimental/GeometryCollectionPlugin/Source/GeometryCollectionNodes/Private/Dataflow/GeometryCollectionNodes.cpp
--- a/Engine/Plugins/Experimental/GeometryCollectionPlugin/Source/GeometryCollectionNodes/Private/Dataflow/GeometryCollectionNodes.cpp
+++ b/Engine/Plugins/Experimental/GeometryCollectionPlugin/Source/GeometryCollectionNodes/Private/Dataflow/GeometryCollectionNodes.cpp
@@ -33,6 +33,8 @@
 #include "GeometryCollection/Facades/CollectionAnchoringFacade.h"
 #include "GeometryCollection/Facades/CollectionRemoveOnBreakFacade.h"
 #include "GeometryCollection/Facades/CollectionTransformFacade.h"
+#include "DynamicMesh/MeshTransforms.h"
+#include "DynamicMesh/DynamicMesh3.h"
 
 #include UE_INLINE_GENERATED_CPP_BY_NAME(GeometryCollectionNodes)
 
@@ -77,7 +79,9 @@ namespace Dataflow
 		DATAFLOW_NODE_REGISTER_CREATION_FACTORY(FGetNumArrayElementsDataflowNode);
 		DATAFLOW_NODE_REGISTER_CREATION_FACTORY(FGetBoundingBoxesFromCollectionDataflowNode);
 		DATAFLOW_NODE_REGISTER_CREATION_FACTORY(FGetCentroidsFromCollectionDataflowNode);
-		DATAFLOW_NODE_REGISTER_CREATION_FACTORY(FTransformDataflowNode);
+		DATAFLOW_NODE_REGISTER_CREATION_FACTORY(FTransformCollectionDataflowNode);
+		DATAFLOW_NODE_REGISTER_CREATION_FACTORY(FBakeTransformsInCollectionDataflowNode);
+		DATAFLOW_NODE_REGISTER_CREATION_FACTORY(FTransformMeshDataflowNode);
 		DATAFLOW_NODE_REGISTER_CREATION_FACTORY(FCompareIntDataflowNode);
 		DATAFLOW_NODE_REGISTER_CREATION_FACTORY(FBranchDataflowNode);
 		DATAFLOW_NODE_REGISTER_CREATION_FACTORY(FGetSchemaDataflowNode);
@@ -659,12 +663,82 @@ void FGetCentroidsFromCollectionDataflowNode::Evaluate(Dataflow::FContext& Conte
 }
 
 
-void FTransformDataflowNode::Evaluate(Dataflow::FContext& Context, const FDataflowOutput* Out) const
+void FTransformCollectionDataflowNode::Evaluate(Dataflow::FContext& Context, const FDataflowOutput* Out) const
+{
+	if (Out->IsA<FManagedArrayCollection>(&Collection))
+	{
+		FManagedArrayCollection InCollection = GetValue<FManagedArrayCollection>(Context, &Collection);
+
+		FTransform NewTransform = GeometryCollection::Facades::FCollectionTransformFacade::BuildTransform(Translate,
+			(uint8)RotationOrder,
+			Rotate,
+			Scale,
+			UniformScale,
+			RotatePivot,
+			ScalePivot,
+			InvertTransformation);
+
+		GeometryCollection::Facades::FCollectionTransformFacade TransformFacade(InCollection);
+		TransformFacade.Transform(NewTransform);
+
+		SetValue<FManagedArrayCollection>(Context, InCollection, &Collection);
+	}
+}
+
+
+void FBakeTransformsInCollectionDataflowNode::Evaluate(Dataflow::FContext& Context, const FDataflowOutput* Out) const
+{
+	if (Out->IsA<FManagedArrayCollection>(&Collection))
+	{
+		FManagedArrayCollection InCollection = GetValue<FManagedArrayCollection>(Context, &Collection);
+
+		GeometryCollection::Facades::FCollectionTransformFacade TransformFacade(InCollection);
+		const TArray<FTransform>& CollectionSpaceTransforms = TransformFacade.ComputeCollectionSpaceTransforms();
+
+		GeometryCollection::Facades::FCollectionMeshFacade MeshFacade(InCollection);
+
+		const int32 NumTransforms = InCollection.NumElements(FGeometryCollection::TransformGroup);
+
+		for (int32 TransformIdx = 0; TransformIdx < NumTransforms; ++TransformIdx)
+		{
+			MeshFacade.BakeTransform(TransformIdx, CollectionSpaceTransforms[TransformIdx]);
+			TransformFacade.SetBoneTransformToIdentity(TransformIdx);
+		}
+
+		SetValue<FManagedArrayCollection>(Context, InCollection, &Collection);
+	}
+}
+
+
+void FTransformMeshDataflowNode::Evaluate(Dataflow::FContext& Context, const FDataflowOutput* Out) const
 {
 	if (Out->IsA<TObjectPtr<UDynamicMesh>>(&Mesh))
 	{
-		// TODO: Implement this
-		SetValue<TObjectPtr<UDynamicMesh>>(Context, NewObject<UDynamicMesh>(), &Mesh);
+		if (TObjectPtr<const UDynamicMesh> InMesh = GetValue<TObjectPtr<UDynamicMesh>>(Context, &Mesh))
+		{
+			// Creating a new mesh object from InMesh
+			TObjectPtr<UDynamicMesh> NewMesh = NewObject<UDynamicMesh>();
+			NewMesh->SetMesh(InMesh->GetMeshRef());
+
+			FTransform NewTransform = GeometryCollection::Facades::FCollectionTransformFacade::BuildTransform(Translate,
+				(uint8)RotationOrder,
+				Rotate,
+				Scale,
+				UniformScale,
+				RotatePivot,
+				ScalePivot,
+				InvertTransformation);
+
+			UE::Geometry::FDynamicMesh3& DynamicMesh = NewMesh->GetMeshRef();
+
+			MeshTransforms::ApplyTransform(DynamicMesh, UE::Geometry::FTransformSRT3d(NewTransform), true);
+
+			SetValue<TObjectPtr<UDynamicMesh>>(Context, NewMesh, &Mesh);
+		}
+		else
+		{
+			SetValue<TObjectPtr<UDynamicMesh>>(Context, NewObject<UDynamicMesh>(), &Mesh);
+		}
 	}
 }
 
@@ -706,14 +780,28 @@ void FBranchDataflowNode::Evaluate(Dataflow::FContext& Context, const FDataflowO
 {
 	if (Out->IsA<TObjectPtr<UDynamicMesh>>(&Mesh))
 	{
-		if (GetValue<bool>(Context, &Condition))
+		bool InCondition = GetValue<bool>(Context, &Condition);
+
+		if (InCondition)
 		{
-			SetValue<TObjectPtr<UDynamicMesh>>(Context, GetValue<TObjectPtr<UDynamicMesh>>(Context, &MeshA), &Mesh);
+			if (TObjectPtr<UDynamicMesh> InMeshA = GetValue<TObjectPtr<UDynamicMesh>>(Context, &MeshA))
+			{
+				SetValue<TObjectPtr<UDynamicMesh>>(Context, InMeshA, &Mesh);
+
+				return;
+			}
 		}
 		else
 		{
-			SetValue<TObjectPtr<UDynamicMesh>>(Context, GetValue<TObjectPtr<UDynamicMesh>>(Context, &MeshB), &Mesh);
+			if (TObjectPtr<UDynamicMesh> InMeshB = GetValue<TObjectPtr<UDynamicMesh>>(Context, &MeshB))
+			{
+				SetValue<TObjectPtr<UDynamicMesh>>(Context, InMeshB, &Mesh);
+
+				return;
+			}
 		}
+
+		SetValue<TObjectPtr<UDynamicMesh>>(Context, NewObject<UDynamicMesh>(), &Mesh);
 	}
 }
 
diff --git a/Engine/Plugins/Experimental/GeometryCollectionPlugin/Source/GeometryCollectionNodes/Private/Dataflow/Nodes/GeometryCollectionAssetNodes.cpp b/Engine/Plugins/Experimental/GeometryCollectionPlugin/Source/GeometryCollectionNodes/Private/Dataflow/Nodes/GeometryCollectionAssetNodes.cpp
--- a/Engine/Plugins/Experimental/GeometryCollectionPlugin/Source/GeometryCollectionNodes/Private/Dataflow/Nodes/GeometryCollectionAssetNodes.cpp
+++ b/Engine/Plugins/Experimental/GeometryCollectionPlugin/Source/GeometryCollectionNodes/Private/Dataflow/Nodes/GeometryCollectionAssetNodes.cpp
@@ -99,7 +99,7 @@ void FGetGeometryCollectionSourcesDataflowNode::Evaluate(Dataflow::FContext& Con
 
 	TArray<FGeometryCollectionSource> OutSources;
 	
-	if (const TObjectPtr<UGeometryCollection> InAsset = GetValue(Context, &Asset))
+	if (const TObjectPtr<const UGeometryCollection> InAsset = GetValue(Context, &Asset))
 	{
 #if WITH_EDITORONLY_DATA
 		OutSources = InAsset->GeometrySource; 
diff --git a/Engine/Plugins/Experimental/GeometryCollectionPlugin/Source/GeometryCollectionNodes/Public/Dataflow/GeometryCollectionMeshNodes.h b/Engine/Plugins/Experimental/GeometryCollectionPlugin/Source/GeometryCollectionNodes/Public/Dataflow/GeometryCollectionMeshNodes.h
--- a/Engine/Plugins/Experimental/GeometryCollectionPlugin/Source/GeometryCollectionNodes/Public/Dataflow/GeometryCollectionMeshNodes.h
+++ b/Engine/Plugins/Experimental/GeometryCollectionPlugin/Source/GeometryCollectionNodes/Public/Dataflow/GeometryCollectionMeshNodes.h
@@ -6,6 +6,7 @@
 #include "Dataflow/DataflowEngine.h"
 #include "GeometryCollection/ManagedArrayCollection.h"
 #include "Dataflow/DataflowSelection.h"
+#include "DynamicMesh/DynamicMesh3.h"
 
 #include "GeometryCollectionMeshNodes.generated.h"
 
@@ -13,7 +14,6 @@ class FGeometryCollection;
 class UStaticMesh;
 class UDynamicMesh;
 
-
 /**
  *
  * Converts points into a DynamicMesh
@@ -61,6 +61,7 @@ struct FBoxToMeshDataflowNode : public FDataflowNode
 {
 	GENERATED_USTRUCT_BODY()
 	DATAFLOW_NODE_DEFINE_INTERNAL(FBoxToMeshDataflowNode, "BoxToMesh", "Mesh|Utilities", "")
+	DATAFLOW_NODE_RENDER_TYPE(FName("FDynamicMesh3"), "Mesh")
 
 public:
 	/** BoundingBox input */
@@ -153,6 +154,42 @@ public:
 };
 
 
+/**
+ *
+ * Converts a Collection to a DynamicMesh
+ *
+ */
+USTRUCT(meta = (DataflowGeometryCollection))
+struct FCollectionToMeshDataflowNode : public FDataflowNode
+{
+	GENERATED_USTRUCT_BODY()
+	DATAFLOW_NODE_DEFINE_INTERNAL(FCollectionToMeshDataflowNode, "CollectionToMesh", "Mesh|Utilities", "")
+	DATAFLOW_NODE_RENDER_TYPE(FName("FDynamicMesh3"), "Mesh")
+
+public:
+	/** Collection to convert*/
+	UPROPERTY(meta = (DataflowInput, DataflowIntrinsic))
+	FManagedArrayCollection Collection;
+
+	UPROPERTY(EditAnywhere, Category = "General", meta = (DisplayName = "Center Pivot"));
+	bool bCenterPivot = false;
+
+	/** Output DynamicMesh */
+	UPROPERTY(meta = (DataflowOutput))
+	TObjectPtr<UDynamicMesh> Mesh;
+
+	FCollectionToMeshDataflowNode(const Dataflow::FNodeParameters& InParam, FGuid InGuid = FGuid::NewGuid())
+		: FDataflowNode(InParam, InGuid)
+	{
+		RegisterInputConnection(&Collection);
+		RegisterOutputConnection(&Mesh);
+	}
+
+	virtual void Evaluate(Dataflow::FContext& Context, const FDataflowOutput* Out) const override;
+
+};
+
+
 /**
  *
  * Converts a StaticMesh into a DynamicMesh
@@ -163,6 +200,7 @@ struct FStaticMeshToMeshDataflowNode : public FDataflowNode
 {
 	GENERATED_USTRUCT_BODY()
 	DATAFLOW_NODE_DEFINE_INTERNAL(FStaticMeshToMeshDataflowNode, "StaticMeshToMesh", "Mesh|Utilities", "")
+	DATAFLOW_NODE_RENDER_TYPE(FName("FDynamicMesh3"), "Mesh")
 
 public:
 	/** StaticMesh to convert */
@@ -171,7 +209,7 @@ public:
 
 	/** Output the HiRes representation, if set to true and HiRes doesn't exist it will output empty mesh */
 	UPROPERTY(EditAnywhere, Category = "StaticMesh", meta = (DisplayName = "Use HiRes"));
-	bool UseHiRes = true;
+	bool bUseHiRes = false;
 
 	/** Specifies the LOD level to use */
 	UPROPERTY(EditAnywhere, Category = "StaticMesh", meta = (DisplayName = "LOD Level"));
@@ -202,6 +240,7 @@ struct FMeshAppendDataflowNode : public FDataflowNode
 {
 	GENERATED_USTRUCT_BODY()
 	DATAFLOW_NODE_DEFINE_INTERNAL(FMeshAppendDataflowNode, "MeshAppend", "Mesh|Utilities", "")
+	DATAFLOW_NODE_RENDER_TYPE(FName("FDynamicMesh3"), "Mesh")
 
 public:
 	/** Mesh input */
@@ -250,6 +289,7 @@ struct FMeshBooleanDataflowNode : public FDataflowNode
 {
 	GENERATED_USTRUCT_BODY()
 	DATAFLOW_NODE_DEFINE_INTERNAL(FMeshBooleanDataflowNode, "MeshBoolean", "Mesh|Utilities", "")
+	DATAFLOW_NODE_RENDER_TYPE(FName("FDynamicMesh3"), "Mesh")
 
 public:
 	/** Boolean operation */
@@ -290,7 +330,8 @@ USTRUCT(meta = (DataflowGeometryCollection))
 struct FMeshCopyToPointsDataflowNode : public FDataflowNode
 {
 	GENERATED_USTRUCT_BODY()
-		DATAFLOW_NODE_DEFINE_INTERNAL(FMeshCopyToPointsDataflowNode, "MeshCopyToPoints", "Mesh|Utilities", "")
+	DATAFLOW_NODE_DEFINE_INTERNAL(FMeshCopyToPointsDataflowNode, "MeshCopyToPoints", "Mesh|Utilities", "")
+	DATAFLOW_NODE_RENDER_TYPE(FName("FDynamicMesh3"), "Mesh")
 
 public:
 	/** Points to copy meshes onto */
diff --git a/Engine/Plugins/Experimental/GeometryCollectionPlugin/Source/GeometryCollectionNodes/Public/Dataflow/GeometryCollectionNodes.h b/Engine/Plugins/Experimental/GeometryCollectionPlugin/Source/GeometryCollectionNodes/Public/Dataflow/GeometryCollectionNodes.h
--- a/Engine/Plugins/Experimental/GeometryCollectionPlugin/Source/GeometryCollectionNodes/Public/Dataflow/GeometryCollectionNodes.h
+++ b/Engine/Plugins/Experimental/GeometryCollectionPlugin/Source/GeometryCollectionNodes/Public/Dataflow/GeometryCollectionNodes.h
@@ -9,6 +9,7 @@
 #include "UDynamicMesh.h"
 #include "Dataflow/DataflowSelection.h"
 #include "Math/MathFwd.h"
+#include "DynamicMesh/DynamicMesh3.h"
 
 #include "GeometryCollectionNodes.generated.h"
 
@@ -17,6 +18,7 @@ class FGeometryCollection;
 class UGeometryCollection;
 class UStaticMesh;
 
+
 /**
  *
  * Description for this node
@@ -187,6 +189,7 @@ struct FBoundingBoxDataflowNode : public FDataflowNode
 {
 	GENERATED_USTRUCT_BODY()
 	DATAFLOW_NODE_DEFINE_INTERNAL(FBoundingBoxDataflowNode, "BoundingBox", "Utilities|Box", "")
+	DATAFLOW_NODE_RENDER_TYPE(FName("FBox"), "BoundingBox")
 
 public:
 	UPROPERTY(meta = (DataflowInput))
@@ -362,6 +365,7 @@ struct FMakeBoxDataflowNode : public FDataflowNode
 {
 	GENERATED_USTRUCT_BODY()
 	DATAFLOW_NODE_DEFINE_INTERNAL(FMakeBoxDataflowNode, "MakeBox", "Generators|Box", "")
+	DATAFLOW_NODE_RENDER_TYPE(FName("FBox"), "Box")
 
 public:
 	UPROPERTY(EditAnywhere, Category = "Box", meta = (DisplayName = "Input Data Type"));
@@ -371,13 +375,13 @@ public:
 	FVector Min = FVector(0.0);
 
 	UPROPERTY(EditAnywhere, Category = "Box", meta = (DataflowInput, EditCondition = "DataType == EMakeBoxDataTypeEnum::Dataflow_MakeBox_DataType_MinMax", EditConditionHides));
-	FVector Max = FVector(0.0);
+	FVector Max = FVector(10.0);
 
 	UPROPERTY(EditAnywhere, Category = "Box", meta = (DataflowInput, EditCondition = "DataType == EMakeBoxDataTypeEnum::Dataflow_MakeBox_DataType_CenterSize", EditConditionHides));
 	FVector Center = FVector(0.0);
 
 	UPROPERTY(EditAnywhere, Category = "Box", meta = (DataflowInput, EditCondition = "DataType == EMakeBoxDataTypeEnum::Dataflow_MakeBox_DataType_CenterSize", EditConditionHides));
-	FVector Size = FVector(0.0);
+	FVector Size = FVector(10.0);
 
 	UPROPERTY(meta = (DataflowOutput));
 	FBox Box = FBox(ForceInit);
@@ -413,7 +417,7 @@ public:
 	FVector Center = FVector(0.f);
 
 	UPROPERTY(EditAnywhere, Category = "Sphere", meta = (DataflowInput));
-	float Radius = 0.f;
+	float Radius = 10.f;
 
 	UPROPERTY(meta = (DataflowOutput));
 	FSphere Sphere = FSphere(ForceInit);
@@ -1056,7 +1060,6 @@ enum class EMathConstantsEnum : uint8
 	Dataflow_Max                UMETA(Hidden)
 };
 
-
 /**
  *
  * Offers a selection of Math constants
@@ -1231,36 +1234,59 @@ public:
 };
 
 
+UENUM(BlueprintType)
+enum class ERotationOrderEnum : uint8
+{
+	Dataflow_RotationOrder_XYZ UMETA(DisplayName = "XYZ"),
+	Dataflow_RotationOrder_YZX UMETA(DisplayName = "YZX"),
+	Dataflow_RotationOrder_ZXY UMETA(DisplayName = "ZXY"),
+	Dataflow_RotationOrder_XZY UMETA(DisplayName = "XZY"),
+	Dataflow_RotationOrder_YXZ UMETA(DisplayName = "YXZ"),
+	Dataflow_RotationOrder_ZYX UMETA(DisplayName = "ZYX"),
+	//~~~
+	//256th entry
+	Dataflow_Max                UMETA(Hidden)
+};
+
 /**
  *
- * Transforms a mesh
+ * Transforms a Collection
  *
  */
 USTRUCT()
-struct FTransformDataflowNode : public FDataflowNode
+struct FTransformCollectionDataflowNode : public FDataflowNode
 {
 	GENERATED_USTRUCT_BODY()
-	DATAFLOW_NODE_DEFINE_INTERNAL(FTransformDataflowNode, "Transform", "Math", "")
+	DATAFLOW_NODE_DEFINE_INTERNAL(FTransformCollectionDataflowNode, "TransformCollection", "Math|Transform", "")
+	DATAFLOW_NODE_RENDER_TYPE(FGeometryCollection::StaticType(), "Collection")
 
 public:
+	/** Output mesh */
+	UPROPERTY(meta = (DataflowInput, DataflowOutput, DataflowPassthrough = "Collection", DataflowIntrinsic))
+	FManagedArrayCollection Collection;
+
 	/** Translation */
 	UPROPERTY(EditAnywhere, Category = "Transform");
 	FVector Translate = FVector(0.0);
 
+	/** Rotation Order */
+	UPROPERTY(EditAnywhere, Category = "Transform");
+	ERotationOrderEnum RotationOrder = ERotationOrderEnum::Dataflow_RotationOrder_XYZ;
+
 	/** Rotation */
 	UPROPERTY(EditAnywhere, Category = "Transform");
 	FVector Rotate = FVector(0.0);
 
 	/** Scale */
-	UPROPERTY(EditAnywhere, Category = "Transform");
+	UPROPERTY(EditAnywhere, Category = "Transform", meta = (UIMin = 0.001f));
 	FVector Scale = FVector(1.0);
 
 	/** Shear */
-	UPROPERTY(EditAnywhere, Category = "Transform");
-	FVector Shear = FVector(0.0);
+//	UPROPERTY(EditAnywhere, Category = "Transform");
+//	FVector Shear = FVector(0.0);
 
 	/** Uniform scale */
-	UPROPERTY(EditAnywhere, Category = "Transform");
+	UPROPERTY(EditAnywhere, Category = "Transform", meta = (UIMin = 0.001f));
 	float UniformScale = 1.f;
 
 	/** Pivot for the rotation */
@@ -1275,15 +1301,105 @@ public:
 	UPROPERTY(EditAnywhere, Category = "General");
 	bool InvertTransformation = false;
 
+	FTransformCollectionDataflowNode(const Dataflow::FNodeParameters& InParam, FGuid InGuid = FGuid::NewGuid())
+		: FDataflowNode(InParam, InGuid)
+	{
+		RegisterInputConnection(&Collection);
+		RegisterOutputConnection(&Collection, &Collection);
+	}
+
+	virtual void Evaluate(Dataflow::FContext& Context, const FDataflowOutput* Out) const override;
+
+};
+
+
+/**
+ *
+ * Bake transforms in Collection
+ *
+ */
+USTRUCT()
+struct FBakeTransformsInCollectionDataflowNode : public FDataflowNode
+{
+	GENERATED_USTRUCT_BODY()
+	DATAFLOW_NODE_DEFINE_INTERNAL(FBakeTransformsInCollectionDataflowNode, "BakeTransformsInCollection", "Math|Transform", "")
+	DATAFLOW_NODE_RENDER_TYPE(FGeometryCollection::StaticType(), "Collection")
+
+public:
+	/** Collection to bake transforms in */
+	UPROPERTY(meta = (DataflowInput, DataflowOutput, DataflowPassthrough = "Collection", DataflowIntrinsic))
+	FManagedArrayCollection Collection;
+
+	FBakeTransformsInCollectionDataflowNode(const Dataflow::FNodeParameters& InParam, FGuid InGuid = FGuid::NewGuid())
+		: FDataflowNode(InParam, InGuid)
+	{
+		RegisterInputConnection(&Collection);
+		RegisterOutputConnection(&Collection, &Collection);
+	}
+
+	virtual void Evaluate(Dataflow::FContext& Context, const FDataflowOutput* Out) const override;
+
+};
+
+
+/**
+ *
+ * Transforms a mesh
+ *
+ */
+USTRUCT()
+struct FTransformMeshDataflowNode : public FDataflowNode
+{
+	GENERATED_USTRUCT_BODY()
+	DATAFLOW_NODE_DEFINE_INTERNAL(FTransformMeshDataflowNode, "TransformMesh", "Math|Transform", "")
+	DATAFLOW_NODE_RENDER_TYPE(FName("FDynamicMesh3"), "Mesh")
+
+public:
 	/** Output mesh */
-	UPROPERTY(meta = (DataflowInput, DataflowOutput))
+	UPROPERTY(meta = (DataflowInput, DataflowOutput, DataflowPassthrough = "Mesh", DataflowIntrinsic))
 	TObjectPtr<UDynamicMesh> Mesh;
 
-	FTransformDataflowNode(const Dataflow::FNodeParameters& InParam, FGuid InGuid = FGuid::NewGuid())
+	/** Translation */
+	UPROPERTY(EditAnywhere, Category = "Transform");
+	FVector Translate = FVector(0.0);
+
+	/** Rotation Order */
+	UPROPERTY(EditAnywhere, Category = "Transform");
+	ERotationOrderEnum RotationOrder = ERotationOrderEnum::Dataflow_RotationOrder_XYZ;
+
+	/** Rotation */
+	UPROPERTY(EditAnywhere, Category = "Transform");
+	FVector Rotate = FVector(0.0);
+
+	/** Scale */
+	UPROPERTY(EditAnywhere, Category = "Transform", meta = (UIMin = 0.001f));
+	FVector Scale = FVector(1.0);
+
+	/** Shear */
+//	UPROPERTY(EditAnywhere, Category = "Transform");
+//	FVector Shear = FVector(0.0);
+
+	/** Uniform scale */
+	UPROPERTY(EditAnywhere, Category = "Transform", meta = (UIMin = 0.001f));
+	float UniformScale = 1.f;
+
+	/** Pivot for the rotation */
+	UPROPERTY(EditAnywhere, Category = "Pivot");
+	FVector RotatePivot = FVector(0.0);
+
+	/** Pivot for the scale */
+	UPROPERTY(EditAnywhere, Category = "Pivot");
+	FVector ScalePivot = FVector(0.0);
+
+	/** Invert the transformation */
+	UPROPERTY(EditAnywhere, Category = "General");
+	bool InvertTransformation = false;
+
+	FTransformMeshDataflowNode(const Dataflow::FNodeParameters& InParam, FGuid InGuid = FGuid::NewGuid())
 		: FDataflowNode(InParam, InGuid)
 	{
 		RegisterInputConnection(&Mesh);
-		RegisterOutputConnection(&Mesh);
+		RegisterOutputConnection(&Mesh, &Mesh);
 	}
 
 	virtual void Evaluate(Dataflow::FContext& Context, const FDataflowOutput* Out) const override;
diff --git a/Engine/Source/Runtime/Experimental/Chaos/Private/GeometryCollection/Facades/CollectionMeshFacade.cpp b/Engine/Source/Runtime/Experimental/Chaos/Private/GeometryCollection/Facades/CollectionMeshFacade.cpp
--- a/Engine/Source/Runtime/Experimental/Chaos/Private/GeometryCollection/Facades/CollectionMeshFacade.cpp
+++ b/Engine/Source/Runtime/Experimental/Chaos/Private/GeometryCollection/Facades/CollectionMeshFacade.cpp
@@ -163,6 +163,29 @@ namespace GeometryCollection::Facades
 		}
 	}
 
+	void FCollectionMeshFacade::BakeTransform(int32 TransformIdx, const FTransform& InTransform)
+	{
+		const TManagedArray<int32>& TransformToGeometryIndices = TransformToGeometryIndexAttribute.Get();
+		TManagedArray<FVector3f>& Vertices = VertexAttribute.Modify();
+		const TManagedArray<int32>& VertexStarts = VertexStartAttribute.Get();
+		const TManagedArray<int32>& VertexCounts = VertexCountAttribute.Get();
+		TManagedArray<FVector3f>& Normals = NormalAttribute.Modify();
+		TManagedArray<FVector3f>& TangentUs = TangentUAttribute.Modify();
+		TManagedArray<FVector3f>& TangentVs = TangentVAttribute.Modify();
+
+		const int32 VertexIndexStart = VertexStarts[TransformToGeometryIndices[TransformIdx]];
+
+		for (int32 Offset = 0; Offset < VertexCounts[TransformToGeometryIndices[TransformIdx]]; ++Offset)
+		{
+			const int32 VertexIdx = VertexIndexStart + Offset;
+
+			Vertices[VertexIdx] = (FVector3f)InTransform.TransformPosition((FVector)Vertices[VertexIdx]);
+			Normals[VertexIdx] = (FVector3f)InTransform.TransformVector((FVector)Normals[VertexIdx]);
+			TangentUs[VertexIdx] = (FVector3f)InTransform.TransformVector((FVector)TangentUs[VertexIdx]);
+			TangentVs[VertexIdx] = (FVector3f)InTransform.TransformVector((FVector)TangentVs[VertexIdx]);
+		}
+	}
+
 };
 
 
diff --git a/Engine/Source/Runtime/Experimental/Chaos/Private/GeometryCollection/Facades/CollectionTransformFacade.cpp b/Engine/Source/Runtime/Experimental/Chaos/Private/GeometryCollection/Facades/CollectionTransformFacade.cpp
--- a/Engine/Source/Runtime/Experimental/Chaos/Private/GeometryCollection/Facades/CollectionTransformFacade.cpp
+++ b/Engine/Source/Runtime/Experimental/Chaos/Private/GeometryCollection/Facades/CollectionTransformFacade.cpp
@@ -72,4 +72,149 @@ namespace GeometryCollection::Facades
 		ensureAlwaysMsgf(false, TEXT("NOT YET IMPLEMENTED"));
 	}
 
+	/**
+	* Builds an FMatrix from FVector4 column vectors
+	*/
+	static FMatrix SetMatrix(const FVector4& Column0, const FVector4& Column1, const FVector4& Column2, const FVector4& Column3)
+	{
+		//
+		// Matrix elements are accessed with M[RowIndex][ColumnumnIndex]
+		//
+		FMatrix Matrix;
+
+		Matrix.M[0][0] = Column0[0]; Matrix.M[1][0] = Column0[1]; Matrix.M[2][0] = Column0[2]; Matrix.M[3][0] = Column0[3];
+		Matrix.M[0][1] = Column1[0]; Matrix.M[1][1] = Column1[1]; Matrix.M[2][1] = Column1[2]; Matrix.M[3][1] = Column1[3];
+		Matrix.M[0][2] = Column2[0]; Matrix.M[1][2] = Column2[1]; Matrix.M[2][2] = Column2[2]; Matrix.M[3][2] = Column2[3];
+		Matrix.M[0][3] = Column3[0]; Matrix.M[1][3] = Column3[1]; Matrix.M[2][3] = Column3[2]; Matrix.M[3][3] = Column3[3];
+
+		return Matrix;
+	}
+
+	FMatrix FCollectionTransformFacade::BuildMatrix(const FVector& Translate,
+		const uint8 RotationOrder,
+		const FVector& Rotate,
+		const FVector& InScale,
+		const FVector& Shear,
+		const float UniformScale,
+		const FVector& RotatePivot,
+		const FVector& ScalePivot,
+		const bool InvertTransformation)
+	{
+// 
+//   M = [Sp-1]x[S]x[Sh]x[Sp]x[St]x[Rp-1]x[Ro]x[R]x[Rp]x[Rt]x[T]
+//
+					
+		// Scale Pivot point
+		FMatrix Sp = SetMatrix(FVector4(1.f, 0.f, 0.f, ScalePivot.X), FVector4(0.f, 1.f, 0.f, ScalePivot.Y), FVector4(0.f, 0.f, 1.f, ScalePivot.Z), FVector4(0.f, 0.f, 0.f, 1.f));
+		FMatrix SpInv = Sp.Inverse();
+
+		// Scale
+		FVector Scale = InScale * UniformScale;
+		FMatrix S = SetMatrix(FVector4(Scale.X, 0.f, 0.f, 0.f), FVector4(0.f, Scale.Y, 0.f, 0.f), FVector4(0.f, 0.f, Scale.Z, 0.f), FVector4(0.f, 0.f, 0.f, 1.f));
+
+		// Shear
+		FMatrix Sh = SetMatrix(FVector4(1.f, Shear.X, Shear.Y, 0.f), FVector4(0.f, 1.f, Shear.Z, 0.f), FVector4(0.f, 0.f, 1.f, 0.f), FVector4(0.f, 0.f, 0.f, 1.f));
+	
+		// Scale Pivot translation
+		FMatrix St = FMatrix::Identity;
+
+		// Rotate Pivot point
+		FMatrix Rp = SetMatrix(FVector4(1.f, 0.f, 0.f, RotatePivot.X), FVector4(0.f, 1.f, 0.f, RotatePivot.Y), FVector4(0.f, 0.f, 1.f, RotatePivot.Z), FVector4(0.f, 0.f, 0.f, 1.f));
+		FMatrix RpInv = Rp.Inverse();
+
+		// Rotation Orientation
+		FMatrix Ro = FMatrix::Identity;
+
+		// Rotation
+		double Sx = sin(FMath::DegreesToRadians(Rotate.X)); double Sy = sin(FMath::DegreesToRadians(Rotate.Y)); double Sz = sin(FMath::DegreesToRadians(Rotate.Z));
+		double Cx = cos(FMath::DegreesToRadians(Rotate.X)); double Cy = cos(FMath::DegreesToRadians(Rotate.Y)); double Cz = cos(FMath::DegreesToRadians(Rotate.Z));
+
+		FMatrix Rx = SetMatrix(FVector4(1.f, 0.f, 0.f, 0.f), FVector4(0.f, Cx, -1.f * Sx, 0.f), FVector4(0.f, Sx, Cx, 0.f), FVector4(0.f, 0.f, 0.f, 1.f));
+		FMatrix Ry = SetMatrix(FVector4(Cy, 0.f, Sy, 0.f), FVector4(0.f, 1.f, 0.f, 0.f), FVector4(-1.f * Sy, 0.f, Cy, 0.f), FVector4(0.f, 0.f, 0.f, 1.f));
+		FMatrix Rz = SetMatrix(FVector4(Cz, -1.f * Sz, 0.f, 0.f), FVector4(Sz, Cz, 0.f, 0.f), FVector4(0.f, 0.f, 1.f, 0.f), FVector4(0.f, 0.f, 0.f, 1.f));
+
+		FMatrix R;
+		if (RotationOrder == 0) // XYZ
+		{
+			R = Rx * Ry; R = R * Rz;
+		}
+		else if (RotationOrder == 1) // YZX
+		{
+			R = Ry * Rz; R = R * Rx;
+		}
+		else if (RotationOrder == 2) // ZXY
+		{
+			R = Rz * Rx; R = R * Ry;
+		}
+		else if (RotationOrder == 3) // XZY
+		{
+			R = Rx * Rz; R = R * Ry;
+		}
+		else if (RotationOrder == 4) // YXZ
+		{
+			R = Ry * Rx; R = R * Rz;
+		}
+		else if (RotationOrder == 5) // ZYX
+		{
+			R = Rz * Ry; R = R * Rx;
+		}
+
+		// Rotate pivot translation
+		FMatrix Rt = FMatrix::Identity;
+
+		// Translate
+		FMatrix T = SetMatrix(FVector4(1.f, 0.f, 0.f, Translate.X), FVector4(0.f, 1.f, 0.f, Translate.Y), FVector4(0.f, 0.f, 1.f, Translate.Z), FVector4(0.f, 0.f, 0.f, 1.f));
+
+		FMatrix MatrixResult = SpInv * S;
+		MatrixResult = MatrixResult * Sh;
+		MatrixResult = MatrixResult * Sp;
+		MatrixResult = MatrixResult * St;
+
+		MatrixResult = MatrixResult * RpInv;
+		MatrixResult = MatrixResult * Ro;
+		MatrixResult = MatrixResult * R;
+		MatrixResult = MatrixResult * Rp;
+		MatrixResult = MatrixResult * Rt;
+
+		MatrixResult = MatrixResult * T;
+
+		if (InvertTransformation)
+		{
+			MatrixResult = MatrixResult.Inverse();
+		}
+
+		return MatrixResult;
+	}
+
+
+	FTransform FCollectionTransformFacade::BuildTransform(const FVector& Translate,
+		const uint8 RotationOrder,
+		const FVector& Rotate,
+		const FVector& InScale,
+		const float UniformScale,
+		const FVector& RotatePivot,
+		const FVector& ScalePivot,
+		const bool InvertTransformation)
+	{
+		return FTransform(BuildMatrix(Translate,
+			RotationOrder,
+			Rotate,
+			InScale,
+			FVector(0.0),
+			UniformScale,
+			RotatePivot,
+			ScalePivot,
+			InvertTransformation));
+	}
+
+
+	void FCollectionTransformFacade::SetBoneTransformToIdentity(int32 BoneIdx)
+	{
+		TManagedArray<FTransform>& BoneTransforms = TransformAttribute.Modify();
+
+		BoneTransforms[BoneIdx] = FTransform::Identity;
+	}
+
+
+
 }
\ No newline at end of file
diff --git a/Engine/Source/Runtime/Experimental/Chaos/Private/GeometryCollection/Facades/CollectionTransformSelectionFacade.cpp b/Engine/Source/Runtime/Experimental/Chaos/Private/GeometryCollection/Facades/CollectionTransformSelectionFacade.cpp
--- a/Engine/Source/Runtime/Experimental/Chaos/Private/GeometryCollection/Facades/CollectionTransformSelectionFacade.cpp
+++ b/Engine/Source/Runtime/Experimental/Chaos/Private/GeometryCollection/Facades/CollectionTransformSelectionFacade.cpp
@@ -999,7 +999,7 @@ namespace GeometryCollection::Facades
 
 		if (TransformIndexAttribute.IsValid() && VertexStartAttribute.IsValid() && VertexCountAttribute.IsValid())
 		{
-			const TManagedArray<int32>& TransformIndicies = TransformIndexAttribute.Get();
+			const TManagedArray<int32>& TransformIndices = TransformIndexAttribute.Get();
 			const TManagedArray<int32>& VertexStarts = VertexStartAttribute.Get();
 			const TManagedArray<int32>& VertexCounts = VertexCountAttribute.Get();
 
@@ -1014,7 +1014,7 @@ namespace GeometryCollection::Facades
 
 					if (VertexIdx >= VertexStart && VertexIdx < VertexEnd)
 					{
-						OutSelection.AddUnique(TransformIndicies[GeometryIdx]);
+						OutSelection.AddUnique(TransformIndices[GeometryIdx]);
 						break;
 					}
 				}
@@ -1030,7 +1030,7 @@ namespace GeometryCollection::Facades
 
 		if (TransformIndexAttribute.IsValid() && FaceStartAttribute.IsValid() && FaceCountAttribute.IsValid())
 		{
-			const TManagedArray<int32>& TransformIndicies = TransformIndexAttribute.Get();
+			const TManagedArray<int32>& TransformIndices = TransformIndexAttribute.Get();
 			const TManagedArray<int32>& FaceStarts = FaceStartAttribute.Get();
 			const TManagedArray<int32>& FaceCounts = FaceCountAttribute.Get();
 
@@ -1045,7 +1045,7 @@ namespace GeometryCollection::Facades
 
 					if (FaceIdx >= FaceStart && FaceIdx < FaceEnd)
 					{
-						OutSelection.AddUnique(TransformIndicies[GeometryIdx]);
+						OutSelection.AddUnique(TransformIndices[GeometryIdx]);
 						break;
 					}
 				}
@@ -1061,31 +1061,31 @@ namespace GeometryCollection::Facades
 
 		if (FaceStartAttribute.IsValid() && FaceCountAttribute.IsValid() && IndicesAttribute.IsValid())
 		{
-			// Build VertexToFaceIndicies inverse lookup
+			// Build VertexToFaceIndices inverse lookup
 			const TManagedArray<int32>& FaceStarts = FaceStartAttribute.Get();
 			const TManagedArray<int32>& FaceCounts = FaceCountAttribute.Get();
 
 			const int32 NumFaces = ConstCollection.NumElements(FGeometryCollection::FacesGroup);
 			const int32 NumVertices = ConstCollection.NumElements(FGeometryCollection::VerticesGroup);
 
-			TArray<TArray<int32>> VertexToFaceIndicies;
-			VertexToFaceIndicies.Reset(NumVertices);
+			TArray<TArray<int32>> VertexToFaceIndices;
+			VertexToFaceIndices.Reset(NumVertices);
 
-			const TManagedArray<FIntVector>& Indicies = IndicesAttribute.Get();
+			const TManagedArray<FIntVector>& Indices = IndicesAttribute.Get();
 
 			for (int32 FaceIdx = 0; FaceIdx < NumFaces; ++FaceIdx)
 			{
-				VertexToFaceIndicies[Indicies[FaceIdx].X].AddUnique(FaceIdx);
-				VertexToFaceIndicies[Indicies[FaceIdx].Y].AddUnique(FaceIdx);
-				VertexToFaceIndicies[Indicies[FaceIdx].Z].AddUnique(FaceIdx);
+				VertexToFaceIndices[Indices[FaceIdx].X].AddUnique(FaceIdx);
+				VertexToFaceIndices[Indices[FaceIdx].Y].AddUnique(FaceIdx);
+				VertexToFaceIndices[Indices[FaceIdx].Z].AddUnique(FaceIdx);
 			}
 
 			// Convert VertexSelection to FaceSelection
 			for (int32 VertexIdx : InVertexSelection)
 			{
-				for (int32 FaceIdx = 0; FaceIdx < VertexToFaceIndicies[VertexIdx].Num(); ++FaceIdx)
+				for (int32 FaceIdx = 0; FaceIdx < VertexToFaceIndices[VertexIdx].Num(); ++FaceIdx)
 				{
-					OutSelection.AddUnique(VertexToFaceIndicies[VertexIdx][FaceIdx]);
+					OutSelection.AddUnique(VertexToFaceIndices[VertexIdx][FaceIdx]);
 				}
 			}
 		}
@@ -1099,15 +1099,15 @@ namespace GeometryCollection::Facades
 
 		if (TransformIndexAttribute.IsValid() && FaceStartAttribute.IsValid() && FaceCountAttribute.IsValid())
 		{
-			const TManagedArray<int32>& TransformToGeometryIndicies = TransformToGeometryIndexAttribute.Get();
+			const TManagedArray<int32>& TransformToGeometryIndices = TransformToGeometryIndexAttribute.Get();
 			const TManagedArray<int32>& FaceStarts = FaceStartAttribute.Get();
 			const TManagedArray<int32>& FaceCounts = FaceCountAttribute.Get();
 
 			for (int32 TransformIdx : InTransformSelection)
 			{
-				const int32 FaceIdxStart = FaceStarts[TransformToGeometryIndicies[TransformIdx]];
+				const int32 FaceIdxStart = FaceStarts[TransformToGeometryIndices[TransformIdx]];
 
-				for (int32 Offset = 0; Offset < FaceCounts[TransformToGeometryIndicies[TransformIdx]]; ++Offset)
+				for (int32 Offset = 0; Offset < FaceCounts[TransformToGeometryIndices[TransformIdx]]; ++Offset)
 				{
 					const int32 FaceIdx = FaceIdxStart + Offset;
 					OutSelection.Add(FaceIdx);
@@ -1124,13 +1124,13 @@ namespace GeometryCollection::Facades
 
 		if (IndicesAttribute.IsValid())
 		{
-			const TManagedArray<FIntVector>& Indicies = IndicesAttribute.Get();
+			const TManagedArray<FIntVector>& Indices = IndicesAttribute.Get();
 
 			for (int32 FaceIdx : InFaceSelection)
 			{
-				OutSelection.AddUnique(Indicies[FaceIdx].X);
-				OutSelection.AddUnique(Indicies[FaceIdx].Y);
-				OutSelection.AddUnique(Indicies[FaceIdx].Z);
+				OutSelection.AddUnique(Indices[FaceIdx].X);
+				OutSelection.AddUnique(Indices[FaceIdx].Y);
+				OutSelection.AddUnique(Indices[FaceIdx].Z);
 			}
 		}
 
@@ -1143,15 +1143,15 @@ namespace GeometryCollection::Facades
 
 		if (TransformToGeometryIndexAttribute.IsValid() && VertexStartAttribute.IsValid() && VertexCountAttribute.IsValid())
 		{
-			const TManagedArray<int32>& TransformToGeometryIndicies = TransformToGeometryIndexAttribute.Get();
+			const TManagedArray<int32>& TransformToGeometryIndices = TransformToGeometryIndexAttribute.Get();
 			const TManagedArray<int32>& VertexStarts = VertexStartAttribute.Get();
 			const TManagedArray<int32>& VertexCounts = VertexCountAttribute.Get();
 
 			for (int32 TransformIdx : InTransformSelection)
 			{
-				const int32 VertexIndexStart = VertexStarts[TransformToGeometryIndicies[TransformIdx]];
+				const int32 VertexIndexStart = VertexStarts[TransformToGeometryIndices[TransformIdx]];
 
-				for (int32 Offset = 0; Offset < VertexCounts[TransformToGeometryIndicies[TransformIdx]]; ++Offset)
+				for (int32 Offset = 0; Offset < VertexCounts[TransformToGeometryIndices[TransformIdx]]; ++Offset)
 				{
 					const int32 VertexIdx = VertexIndexStart + Offset;
 					OutSelection.Add(VertexIdx);
diff --git a/Engine/Source/Runtime/Experimental/Chaos/Public/GeometryCollection/Facades/CollectionMeshFacade.h b/Engine/Source/Runtime/Experimental/Chaos/Public/GeometryCollection/Facades/CollectionMeshFacade.h
--- a/Engine/Source/Runtime/Experimental/Chaos/Public/GeometryCollection/Facades/CollectionMeshFacade.h
+++ b/Engine/Source/Runtime/Experimental/Chaos/Public/GeometryCollection/Facades/CollectionMeshFacade.h
@@ -52,6 +52,11 @@ namespace GeometryCollection::Facades
 		 */
 //		const TArrayView<FIntVector> GetFaceVertexIndices(int32 BoneIdx) const;
 
+		/**
+		 * Bakes the transforms into the vertex positions and sets the bone transforms to identity
+		 */
+		 void BakeTransform(int32 TransformIdx, const FTransform& InTransform);
+
 		TManagedArrayAccessor<int32> TransformToGeometryIndexAttribute;
 		TManagedArrayAccessor<FVector3f> VertexAttribute;
 		TManagedArrayAccessor<FVector3f> TangentUAttribute;
diff --git a/Engine/Source/Runtime/Experimental/Chaos/Public/GeometryCollection/Facades/CollectionTransformFacade.h b/Engine/Source/Runtime/Experimental/Chaos/Public/GeometryCollection/Facades/CollectionTransformFacade.h
--- a/Engine/Source/Runtime/Experimental/Chaos/Public/GeometryCollection/Facades/CollectionTransformFacade.h
+++ b/Engine/Source/Runtime/Experimental/Chaos/Public/GeometryCollection/Facades/CollectionTransformFacade.h
@@ -17,7 +17,7 @@ namespace GeometryCollection::Facades
 		FCollectionTransformFacade(FManagedArrayCollection& InCollection);
 		FCollectionTransformFacade(const FManagedArrayCollection& InCollection);
 
-		/** Create the facade attributes. */
+		/** Creates the facade attributes. */
 		void DefineSchema() {}
 
 		/** Valid if parent and children arrays are available */
@@ -26,32 +26,33 @@ namespace GeometryCollection::Facades
 		/** Is the facade defined constant. */
 		bool IsConst() const { return ParentAttribute.IsConst(); }
 
-		/** Get the root index */
+		/** Gets the root index */
 		TArray<int32> GetRootIndices() const;
 
 		/**
-		* Return the parent indices from the collection. Null if not initialized.
+		* Returns the parent indices from the collection. Null if not initialized.
 		*/
 		const TManagedArray< int32 >* GetParents() const { return ParentAttribute.Find(); }
 
 		/**
-		* Return the child indicesfrom the collection. Null if not initialized.
+		* Returns the child indicesfrom the collection. Null if not initialized.
 		*/
 		const TManagedArray< TSet<int32> >* FindChildren() const { return ChildrenAttribute.Find(); }
 
 		/**
-		* Return the child indicesfrom the collection. Null if not initialized.
+		* Returns the child indicesfrom the collection. Null if not initialized.
 		*/
 		const TManagedArray< FTransform >* FindTransforms() const { return TransformAttribute.Find(); }
 
 		/**
-		* Return array of transforms for transforming from bone space to collection space
-		* Vertex(inCollectionSpace) = TransformComputed.TransformPosition(Vertex(inBoneSpace));
+		* Returns array of transforms for transforming from bone space to collection space
+		* Vertex(inCollectionSpace) = TransformComputed.TransformPosition(Vertex(inBoneSpace))
 		*/
 		TArray<FTransform> ComputeCollectionSpaceTransforms() const;
 
 		/**
-		* Return the child indicesfrom the collection. Null if not initialized.
+		* Returns the transform for transforming from bone space to collection space for specified bone
+		* Vertex(inCollectionSpace) = TransformComputed.TransformPosition(Vertex(inBoneSpace))
 		*/
 		FTransform ComputeCollectionSpaceTransform(int32 BoneIdx) const;
 
@@ -64,6 +65,30 @@ namespace GeometryCollection::Facades
 		/** Transforms selected bones in the collection */
 		void Transform(const FTransform& InTransform, const TArray<int32>& InSelection);
 
+		/** Builds a FMatrix from all the components */
+		static FMatrix BuildMatrix(const FVector& Translate,
+			const uint8 RotationOrder,
+			const FVector& Rotate,
+			const FVector& Scale,
+			const FVector& Shear,
+			const float UniformScale,
+			const FVector& RotatePivot,
+			const FVector& ScalePivot,
+			const bool InvertTransformation);
+
+		/** Builds a FTransform from all the components */
+		static FTransform BuildTransform(const FVector& Translate,
+			const uint8 RotationOrder,
+			const FVector& Rotate,
+			const FVector& Scale,
+			const float UniformScale,
+			const FVector& RotatePivot,
+			const FVector& ScalePivot,
+			const bool InvertTransformation);
+
+		/** Sets the selected bone's transform to identity */
+		void SetBoneTransformToIdentity(int32 BoneIdx);
+
 	private:
 		TManagedArrayAccessor<int32>		ParentAttribute;
 		TManagedArrayAccessor<TSet<int32>>	ChildrenAttribute;
