Geometry collection : Fix thumbnaiol not properly rendering because of out of date bounds

diff --git a/Engine/Plugins/Experimental/GeometryCollectionPlugin/Source/GeometryCollectionEditor/Private/GeometryCollection/GeometryCollectionThumbnailScene.cpp b/Engine/Plugins/Experimental/GeometryCollectionPlugin/Source/GeometryCollectionEditor/Private/GeometryCollection/GeometryCollectionThumbnailScene.cpp
--- a/Engine/Plugins/Experimental/GeometryCollectionPlugin/Source/GeometryCollectionEditor/Private/GeometryCollection/GeometryCollectionThumbnailScene.cpp
+++ b/Engine/Plugins/Experimental/GeometryCollectionPlugin/Source/GeometryCollectionEditor/Private/GeometryCollection/GeometryCollectionThumbnailScene.cpp
@@ -25,21 +25,25 @@ FGeometryCollectionThumbnailScene::FGeometryCollectionThumbnailScene()
 
 void FGeometryCollectionThumbnailScene::SetGeometryCollection(UGeometryCollection* GeometryCollection)
 {
-	if (PreviewActor->GetGeometryCollectionComponent())
+	if (UGeometryCollectionComponent* Component = PreviewActor->GetGeometryCollectionComponent())
 	{
-		PreviewActor->GetGeometryCollectionComponent()->SetRestCollection(GeometryCollection);
+		Component->SetRestCollection(GeometryCollection);
 
 		if (GeometryCollection)
 		{
 			FTransform MeshTransform = FTransform::Identity;
 
 			PreviewActor->SetActorLocation(FVector(0, 0, 0), false);
-			PreviewActor->GetGeometryCollectionComponent()->UpdateBounds();
+			// use UpdateCachedBounds because UpdateBounds() may not properly refresh the bounds for geometry collection as it will not detect the move
+			Component->UpdateCachedBounds();
 
 			// Center the mesh at the world origin then offset to put it on top of the plane
-			const float BoundsZOffset = GetBoundsZOffset(PreviewActor->GetGeometryCollectionComponent()->Bounds);
-			PreviewActor->SetActorLocation(-PreviewActor->GetGeometryCollectionComponent()->Bounds.Origin + FVector(0, 0, BoundsZOffset), false);
-			PreviewActor->GetGeometryCollectionComponent()->RecreateRenderState_Concurrent();
+			const FVector::FReal BoundsZOffset = GetBoundsZOffset(Component->Bounds);
+			const FVector PlacedLocation = -Component->Bounds.Origin + FVector(0, 0, BoundsZOffset);
+
+			PreviewActor->SetActorLocation(PlacedLocation, false);
+			Component->UpdateCachedBounds();
+			Component->RecreateRenderState_Concurrent();
 		}
 	}
 }
@@ -50,10 +54,12 @@ void FGeometryCollectionThumbnailScene::GetViewMatrixParameters(const float InFO
 	check(PreviewActor->GetGeometryCollectionComponent());
 	check(PreviewActor->GetGeometryCollectionComponent()->GetRestCollection());
 
+	UGeometryCollectionComponent* Component = PreviewActor->GetGeometryCollectionComponent();
+
 	const float HalfFOVRadians = FMath::DegreesToRadians<float>(InFOVDegrees) * 0.5f;
 	// Add extra size to view slightly outside of the sphere to compensate for perspective
-	const float HalfMeshSize = PreviewActor->GetGeometryCollectionComponent()->Bounds.SphereRadius * 1.15;
-	const float BoundsZOffset = GetBoundsZOffset(PreviewActor->GetGeometryCollectionComponent()->Bounds);
+	const float HalfMeshSize = Component->Bounds.SphereRadius * 1.15;
+	const float BoundsZOffset = GetBoundsZOffset(Component->Bounds);
 	const float TargetDistance = HalfMeshSize / FMath::Tan(HalfFOVRadians);
 
 	USceneThumbnailInfo* ThumbnailInfo = Cast<USceneThumbnailInfo>(PreviewActor->GetGeometryCollectionComponent()->GetRestCollection()->ThumbnailInfo);
@@ -74,4 +80,3 @@ void FGeometryCollectionThumbnailScene::GetViewMatrixParameters(const float InFO
 	OutOrbitYaw = ThumbnailInfo->OrbitYaw;
 	OutOrbitZoom = TargetDistance + ThumbnailInfo->OrbitZoom;
 }
-
