Fix potential crash when adding a GC to a sequecner and adding component as the track foollowed by adding a Geometry Collection track to it

diff --git a/Engine/Plugins/Experimental/GeometryCollectionPlugin/Source/GeometryCollectionTracks/Private/MovieSceneGeometryCollectionTemplate.cpp b/Engine/Plugins/Experimental/GeometryCollectionPlugin/Source/GeometryCollectionTracks/Private/MovieSceneGeometryCollectionTemplate.cpp
--- a/Engine/Plugins/Experimental/GeometryCollectionPlugin/Source/GeometryCollectionTracks/Private/MovieSceneGeometryCollectionTemplate.cpp
+++ b/Engine/Plugins/Experimental/GeometryCollectionPlugin/Source/GeometryCollectionTracks/Private/MovieSceneGeometryCollectionTemplate.cpp
@@ -77,30 +77,31 @@ struct FGeometryCollectionExecutionToken
 			UGeometryCollectionComponent* GeometryCollectionComponent = Cast<UGeometryCollectionComponent>(WeakObject.Get());
 			if (GeometryCollectionComponent)
 			{
-				UGeometryCollectionCache* GeometryCollectionCache = Cast<UGeometryCollectionCache>(Params.GeometryCollectionCache.ResolveObject());
-				
-				// Validate that the Cache is compatible with the Collection before attempting to play back.
-				const bool bIsValidCache = GeometryCollectionCache->CompatibleWithForPlayback(GeometryCollectionComponent->GetRestCollection());
-				if (!bIsValidCache)
+				if (UGeometryCollectionCache* GeometryCollectionCache = Cast<UGeometryCollectionCache>(Params.GeometryCollectionCache.ResolveObject()))
 				{
-					UE_LOG(LogMovieScene, Warning, TEXT("Unsupported Geometry Collection Cache (%s) for Component (%s), ignoring playback!"), *GeometryCollectionCache->GetFullName(), *GeometryCollectionComponent->GetFullName());
-					continue;
-				}
+					// Validate that the Cache is compatible with the Collection before attempting to play back.
+					const bool bIsValidCache = GeometryCollectionCache->CompatibleWithForPlayback(GeometryCollectionComponent->GetRestCollection());
+					if (!bIsValidCache)
+					{
+						UE_LOG(LogMovieScene, Warning, TEXT("Unsupported Geometry Collection Cache (%s) for Component (%s), ignoring playback!"), *GeometryCollectionCache->GetFullName(), *GeometryCollectionComponent->GetFullName());
+						continue;
+					}
 
-				Player.SavePreAnimatedState(*GeometryCollectionComponent, FPreAnimatedGeometryCollectionTokenProducer::GetAnimTypeID(), FPreAnimatedGeometryCollectionTokenProducer());
+					Player.SavePreAnimatedState(*GeometryCollectionComponent, FPreAnimatedGeometryCollectionTokenProducer::GetAnimTypeID(), FPreAnimatedGeometryCollectionTokenProducer());
 
-				// Now that we've stored the PreAnimated State we change whatever settings are needed for Sequencer to actually control this. This is done every frame in the event
-				// that Gameplay Code tries to fight Sequencer and resets a setting.
-				GeometryCollectionComponent->ObjectType = EObjectStateTypeEnum::Chaos_Object_Kinematic;
-				GeometryCollectionComponent->CachePlayback = true;
-				GeometryCollectionComponent->CacheParameters.CacheMode = EGeometryCollectionCacheType::None;
+					// Now that we've stored the PreAnimated State we change whatever settings are needed for Sequencer to actually control this. This is done every frame in the event
+					// that Gameplay Code tries to fight Sequencer and resets a setting.
+					GeometryCollectionComponent->ObjectType = EObjectStateTypeEnum::Chaos_Object_Kinematic;
+					GeometryCollectionComponent->CachePlayback = true;
+					GeometryCollectionComponent->CacheParameters.CacheMode = EGeometryCollectionCacheType::None;
 
-				GeometryCollectionComponent->CacheParameters.TargetCache = GeometryCollectionCache;
+					GeometryCollectionComponent->CacheParameters.TargetCache = GeometryCollectionCache;
 
-				// Finding out what time (relative to the start of the section)
-				const float TimeInSeconds = Params.MapTimeToAnimation(Context.GetTime(), Context.GetFrameRate());
+					// Finding out what time (relative to the start of the section)
+					const float TimeInSeconds = Params.MapTimeToAnimation(Context.GetTime(), Context.GetFrameRate());
 
-				GeometryCollectionComponent->DesiredCacheTime = TimeInSeconds;
+					GeometryCollectionComponent->DesiredCacheTime = TimeInSeconds;
+				}
 			}
 		}
 	}
