Dataflow:
- Added AddMaterialToCollection,  ReAssignMaterialInCollection, MaterialsInfo, GetMaterialFromMaterialsArray, SetMaterialInMaterialsArray, MakeMaterial, MakeMaterialsArray nodes

diff --git a/Engine/Plugins/Experimental/GeometryCollectionPlugin/Source/GeometryCollectionNodes/Private/Dataflow/GeometryCollectionNodes.cpp b/Engine/Plugins/Experimental/GeometryCollectionPlugin/Source/GeometryCollectionNodes/Private/Dataflow/GeometryCollectionNodes.cpp
--- a/Engine/Plugins/Experimental/GeometryCollectionPlugin/Source/GeometryCollectionNodes/Private/Dataflow/GeometryCollectionNodes.cpp
+++ b/Engine/Plugins/Experimental/GeometryCollectionPlugin/Source/GeometryCollectionNodes/Private/Dataflow/GeometryCollectionNodes.cpp
@@ -36,6 +36,7 @@
 #include "GeometryCollection/Facades/CollectionTransformFacade.h"
 #include "DynamicMesh/MeshTransforms.h"
 #include "DynamicMesh/DynamicMesh3.h"
+#include "GeometryCollection/Facades/CollectionTransformSelectionFacade.h"
 
 #include UE_INLINE_GENERATED_CPP_BY_NAME(GeometryCollectionNodes)
 
@@ -95,6 +96,13 @@ namespace Dataflow
 		DATAFLOW_NODE_REGISTER_CREATION_FACTORY(FGetCollectionAttributeDataTypedDataflowNode);
 		DATAFLOW_NODE_REGISTER_CREATION_FACTORY(FSetCollectionAttributeDataTypedDataflowNode);
 		DATAFLOW_NODE_REGISTER_CREATION_FACTORY(FBoolArrayToFaceSelectionDataflowNode);
+		DATAFLOW_NODE_REGISTER_CREATION_FACTORY(FAddMaterialToCollectionDataflowNode);
+		DATAFLOW_NODE_REGISTER_CREATION_FACTORY(FReAssignMaterialInCollectionDataflowNode);
+		DATAFLOW_NODE_REGISTER_CREATION_FACTORY(FMaterialsInfoDataflowNode);
+		DATAFLOW_NODE_REGISTER_CREATION_FACTORY(FGetMaterialFromMaterialsArrayDataflowNode);
+		DATAFLOW_NODE_REGISTER_CREATION_FACTORY(FSetMaterialInMaterialsArrayDataflowNode);
+		DATAFLOW_NODE_REGISTER_CREATION_FACTORY(FMakeMaterialDataflowNode);
+		DATAFLOW_NODE_REGISTER_CREATION_FACTORY(FMakeMaterialsArrayDataflowNode);
 
 		// GeometryCollection
 		DATAFLOW_NODE_REGISTER_CREATION_FACTORY_NODE_COLORS_BY_CATEGORY("GeometryCollection", FLinearColor(0.55f, 0.45f, 1.0f), CDefaultNodeBodyTintColor);
@@ -861,7 +869,7 @@ void FRemoveOnBreakDataflowNode::Evaluate(Dataflow::FContext& Context, const FDa
 {
 	if (Out->IsA<FManagedArrayCollection>(&Collection))
 	{
-		const bool& InEnableRemoval = GetValue(Context, &EnabledRemoval, true);
+		const bool& InEnableRemoval = GetValue(Context, &bEnabledRemoval, true);
 		const FVector2f& InPostBreakTimer = GetValue(Context, &PostBreakTimer);
 		const FVector2f& InRemovalTimer = GetValue(Context, &RemovalTimer);
 		const bool& InClusterCrumbling = GetValue(Context, &ClusterCrumbling);
@@ -1317,3 +1325,217 @@ void FBoolArrayToFaceSelectionDataflowNode::Evaluate(Dataflow::FContext& Context
 	}
 }
 
+
+void FAddMaterialToCollectionDataflowNode::Evaluate(Dataflow::FContext& Context, const FDataflowOutput* Out) const
+{
+	if (Out->IsA<FManagedArrayCollection>(&Collection) || Out->IsA<TArray<TObjectPtr<UMaterial>>>(&Materials))
+	{
+		FManagedArrayCollection InCollection = GetValue<FManagedArrayCollection>(Context, &Collection);
+		TArray<TObjectPtr<UMaterial>> InMaterials = GetValue<TArray<TObjectPtr<UMaterial>>>(Context, &Materials);
+		const FDataflowFaceSelection& InFaceSelection = GetValue<FDataflowFaceSelection>(Context, &FaceSelection);
+		TObjectPtr<UMaterial> InOutsideMaterial = GetValue<TObjectPtr<UMaterial>>(Context, &OutsideMaterial);
+		TObjectPtr<UMaterial> InInsideMaterial = GetValue<TObjectPtr<UMaterial>>(Context, &InsideMaterial);
+
+		int32 NewIndex;
+
+		if (InCollection.HasAttribute("Internal", FGeometryCollection::FacesGroup))
+		{
+			if (bAssignOutsideMaterial || bAssignInsideMaterial)
+			{
+				TManagedArray<int32>& MaterialIDs = InCollection.ModifyAttribute<int32>("MaterialID", FGeometryCollection::FacesGroup);
+
+				const TManagedArray<bool>& Internals = InCollection.GetAttribute<bool>("Internal", FGeometryCollection::FacesGroup);
+
+				if (bAssignOutsideMaterial)
+				{
+					NewIndex = InMaterials.Add(InOutsideMaterial);
+
+					// Update MaterialIdx for selected outside faces
+					for (int32 FaceIdx = 0; FaceIdx < InFaceSelection.Num(); ++FaceIdx)
+					{
+						if (!Internals[FaceIdx] && InFaceSelection.IsSelected(FaceIdx))
+						{
+							MaterialIDs[FaceIdx] = NewIndex;
+						}
+					}
+				}
+
+				if (bAssignInsideMaterial)
+				{
+					NewIndex = InMaterials.Add(InInsideMaterial);
+
+					// Update MaterialIdx for selected inside faces
+					for (int32 FaceIdx = 0; FaceIdx < InFaceSelection.Num(); ++FaceIdx)
+					{
+						if (Internals[FaceIdx] && InFaceSelection.IsSelected(FaceIdx))
+						{
+							MaterialIDs[FaceIdx] = NewIndex;
+						}
+					}
+				}
+			}
+		}
+
+		SetValue<TArray<TObjectPtr<UMaterial>>>(Context, MoveTemp(InMaterials), &Materials);
+		SetValue<FManagedArrayCollection>(Context, MoveTemp(InCollection), &Collection);
+	}
+}
+
+
+void FReAssignMaterialInCollectionDataflowNode::Evaluate(Dataflow::FContext& Context, const FDataflowOutput* Out) const
+{
+	if (Out->IsA<FManagedArrayCollection>(&Collection))
+	{
+		FManagedArrayCollection InCollection = GetValue<FManagedArrayCollection>(Context, &Collection);
+		const FDataflowFaceSelection& InFaceSelection = GetValue<FDataflowFaceSelection>(Context, &FaceSelection);
+		const TArray<TObjectPtr<UMaterial>>& InMaterials = GetValue<TArray<TObjectPtr<UMaterial>>>(Context, &Materials);
+		int32 InOutsideMaterialIdx = GetValue<int32>(Context, &OutsideMaterialIdx);
+		int32 InInsideMaterialIdx = GetValue<int32>(Context, &InsideMaterialIdx);
+
+		if (InMaterials.Num() > 0)
+		{
+			if (InCollection.HasAttribute("Internal", FGeometryCollection::FacesGroup))
+			{
+				if (bAssignOutsideMaterial || bAssignInsideMaterial)
+				{
+					TManagedArray<int32>& MaterialIDs = InCollection.ModifyAttribute<int32>("MaterialID", FGeometryCollection::FacesGroup);
+
+					const TManagedArray<bool>& Internals = InCollection.GetAttribute<bool>("Internal", FGeometryCollection::FacesGroup);
+
+					if (bAssignOutsideMaterial)
+					{
+						if (InOutsideMaterialIdx >= 0 && InOutsideMaterialIdx < InMaterials.Num())
+						{
+							// Update MaterialIdx for selected outside faces
+							for (int32 FaceIdx = 0; FaceIdx < InFaceSelection.Num(); ++FaceIdx)
+							{
+								if (!Internals[FaceIdx] && InFaceSelection.IsSelected(FaceIdx))
+								{
+									MaterialIDs[FaceIdx] = InOutsideMaterialIdx;
+								}
+							}
+						}
+					}
+
+					if (bAssignInsideMaterial)
+					{
+						if (InInsideMaterialIdx >= 0 && InInsideMaterialIdx < InMaterials.Num())
+						{
+							// Update MaterialIdx for selected inside faces
+							for (int32 FaceIdx = 0; FaceIdx < InFaceSelection.Num(); ++FaceIdx)
+							{
+								if (Internals[FaceIdx] && InFaceSelection.IsSelected(FaceIdx))
+								{
+									MaterialIDs[FaceIdx] = InInsideMaterialIdx;
+								}
+							}
+						}
+					}
+				}
+			}
+		}
+
+		// move the collection to the output to avoid making another copy
+		SetValue<FManagedArrayCollection>(Context, MoveTemp(InCollection), &Collection);
+		SetValue<TArray<TObjectPtr<UMaterial>>>(Context, InMaterials, &Materials);
+	}
+}
+
+
+void FMaterialsInfoDataflowNode::Evaluate(Dataflow::FContext& Context, const FDataflowOutput* Out) const
+{
+	if (Out->IsA<FString>(&String))
+	{
+		const TArray<TObjectPtr<UMaterial>>& InMaterials = GetValue<TArray<TObjectPtr<UMaterial>>>(Context, &Materials);
+
+		FString OutputStr;
+
+		OutputStr.Appendf(TEXT("\n----------------------------------------\n"));
+		OutputStr.Appendf(TEXT("Number of Materials: %d\n"), InMaterials.Num());
+
+		int32 Idx = 0;
+		for (TObjectPtr<UMaterial> Material : InMaterials)
+		{
+			OutputStr.Appendf(TEXT("%4d: %s\n"), Idx, *(Material->GetFullName()));
+
+			Idx++;
+		}
+
+		OutputStr.Appendf(TEXT("----------------------------------------\n"));
+
+		SetValue<FString>(Context, OutputStr, &String);
+	}
+}
+
+
+void FGetMaterialFromMaterialsArrayDataflowNode::Evaluate(Dataflow::FContext& Context, const FDataflowOutput* Out) const
+{
+	if (Out->IsA<TObjectPtr<UMaterial>>(&Material))
+	{
+		const TArray<TObjectPtr<UMaterial>>& InMaterials = GetValue<TArray<TObjectPtr<UMaterial>>>(Context, &Materials);
+		int32 InMaterialIdx = GetValue<int32>(Context, &MaterialIdx);
+
+		if (InMaterials.Num() > 0)
+		{
+			if (InMaterialIdx >= 0 && InMaterialIdx < InMaterials.Num())
+			{
+				SetValue<TObjectPtr<UMaterial>>(Context, InMaterials[InMaterialIdx], &Material);
+
+				return;
+			}
+		}
+
+		SetValue<TObjectPtr<UMaterial>>(Context, TObjectPtr<UMaterial>(), &Material);
+	}
+}
+
+
+void FSetMaterialInMaterialsArrayDataflowNode::Evaluate(Dataflow::FContext& Context, const FDataflowOutput* Out) const
+{
+	if (Out->IsA<TArray<TObjectPtr<UMaterial>>>(&Materials))
+	{
+		TArray<TObjectPtr<UMaterial>> InMaterials = GetValue<TArray<TObjectPtr<UMaterial>>>(Context, &Materials);
+		TObjectPtr<UMaterial> InMaterial = GetValue<TObjectPtr<UMaterial>>(Context, &Material);
+		int32 InMaterialIdx = GetValue<int32>(Context, &MaterialIdx);
+
+		if (Operation == ESetMaterialOperationTypeEnum::Dataflow_SetMaterialOperationType_Add)
+		{
+			InMaterials.Add(InMaterial);
+		}
+		else
+		{
+			if (InMaterialIdx >= 0 && InMaterialIdx < InMaterials.Num())
+			{
+				InMaterials.Insert(InMaterial, InMaterialIdx);
+			}
+		}
+
+		SetValue<TArray<TObjectPtr<UMaterial>>>(Context, MoveTemp(InMaterials), &Materials);
+	}
+}
+
+
+void FMakeMaterialDataflowNode::Evaluate(Dataflow::FContext& Context, const FDataflowOutput* Out) const
+{
+	if (Out->IsA<TObjectPtr<UMaterial>>(&Material))
+	{
+		SetValue<TObjectPtr<UMaterial>>(Context, InMaterial, &Material);
+	}
+}
+
+
+void FMakeMaterialsArrayDataflowNode::Evaluate(Dataflow::FContext& Context, const FDataflowOutput* Out) const
+{
+	if (Out->IsA<TArray<TObjectPtr<UMaterial>>>(&Materials))
+	{
+		SetValue<TArray<TObjectPtr<UMaterial>>>(Context, TArray<TObjectPtr<UMaterial>>(), &Materials);
+	}
+}
+
+
+
+
+// ----------------------------------------------------------------------------
+
+
+
diff --git a/Engine/Plugins/Experimental/GeometryCollectionPlugin/Source/GeometryCollectionNodes/Public/Dataflow/GeometryCollectionNodes.h b/Engine/Plugins/Experimental/GeometryCollectionPlugin/Source/GeometryCollectionNodes/Public/Dataflow/GeometryCollectionNodes.h
--- a/Engine/Plugins/Experimental/GeometryCollectionPlugin/Source/GeometryCollectionNodes/Public/Dataflow/GeometryCollectionNodes.h
+++ b/Engine/Plugins/Experimental/GeometryCollectionPlugin/Source/GeometryCollectionNodes/Public/Dataflow/GeometryCollectionNodes.h
@@ -1552,7 +1552,7 @@ public:
 
 	/** Whether or not to enable the removal on the selection */
 	UPROPERTY(EditAnywhere, Category = "Removal", meta = (DataflowInput))
-	bool EnabledRemoval = true;
+	bool bEnabledRemoval = true;
 
 	/** How long after the break the removal will start ( Min / Max ) */
 	UPROPERTY(EditAnywhere, Category = "Removal", meta = (DataflowInput))
@@ -1571,7 +1571,7 @@ public:
 	{
 		RegisterInputConnection(&Collection);
 		RegisterInputConnection(&TransformSelection);
-		RegisterInputConnection(&EnabledRemoval);
+		RegisterInputConnection(&bEnabledRemoval);
 		RegisterInputConnection(&PostBreakTimer);
 		RegisterInputConnection(&RemovalTimer);
 		RegisterInputConnection(&ClusterCrumbling);
@@ -2000,6 +2000,299 @@ public:
 };
 
 
+/**
+ *
+ * Adds Outside/Inside Materials to Outside/Inside faces
+ *
+ */
+USTRUCT(meta = (DataflowGeometryCollection))
+struct FAddMaterialToCollectionDataflowNode : public FDataflowNode
+{
+	GENERATED_USTRUCT_BODY()
+	DATAFLOW_NODE_DEFINE_INTERNAL(FAddMaterialToCollectionDataflowNode, "AddMaterialToCollection", "GeometryCollection|Materials", "")
+
+public:
+	/** Collection to add material(s) to */
+	UPROPERTY(meta = (DataflowInput, DataflowOutput, DataflowPassthrough = "Collection", DataflowIntrinsic))
+	FManagedArrayCollection Collection;
+
+	/** Face selection, the material(s) will be added to the selected faces */
+	UPROPERTY(meta = (DataflowInput, DisplayName = "FaceSelection", DataflowIntrinsic))
+	FDataflowFaceSelection FaceSelection;
+
+	/** Materials array storing the materials */
+	UPROPERTY(meta = (DataflowInput, DataflowOutput, DataflowPassthrough = "Materials", DataflowIntrinsic))
+	TArray<TObjectPtr<UMaterial>> Materials;
+
+	/** Outside material to assign to the outside faces from the face selection */
+	UPROPERTY(EditAnywhere, Category = "Material", meta = (DataflowInput))
+	TObjectPtr<UMaterial> OutsideMaterial;
+
+	/** Inside material to assign to the inside faces from the face selection */
+	UPROPERTY(EditAnywhere, Category = "Material", meta = (DataflowInput))
+	TObjectPtr<UMaterial> InsideMaterial;
+
+	/** If true, the outside material will be assigned to the outside faces from the face selection */
+	UPROPERTY(EditAnywhere, Category = "Material", meta = (DisplayName = "Assign Outside Material to Outside Faces"))
+	bool bAssignOutsideMaterial = true;
+
+	/** If true, the inside material will be assigned to the inside faces from the face selection */
+	UPROPERTY(EditAnywhere, Category = "Material", meta = (DisplayName = "Assign Inside Material to Inside Faces"))
+	bool bAssignInsideMaterial = false;
+
+	FAddMaterialToCollectionDataflowNode(const Dataflow::FNodeParameters& InParam, FGuid InGuid = FGuid::NewGuid())
+		: FDataflowNode(InParam, InGuid)
+	{
+		RegisterInputConnection(&Collection);
+		RegisterInputConnection(&Materials);
+		RegisterInputConnection(&FaceSelection);
+		RegisterInputConnection(&OutsideMaterial);
+		RegisterInputConnection(&InsideMaterial);
+		RegisterOutputConnection(&Collection, &Collection);
+		RegisterOutputConnection(&Materials, &Materials);
+	}
+
+	virtual void Evaluate(Dataflow::FContext& Context, const FDataflowOutput* Out) const override;
+
+};
+
+
+/**
+ *
+ * Reassign existing material(s) to Outside/Inside faces
+ *
+ */
+USTRUCT(meta = (DataflowGeometryCollection))
+struct FReAssignMaterialInCollectionDataflowNode : public FDataflowNode
+{
+	GENERATED_USTRUCT_BODY()
+	DATAFLOW_NODE_DEFINE_INTERNAL(FReAssignMaterialInCollectionDataflowNode, "ReAssignMaterialInCollection", "GeometryCollection|Materials", "")
+
+public:
+	/** Collection for reassign the material(s) */
+	UPROPERTY(meta = (DataflowInput, DataflowOutput, DataflowPassthrough = "Collection", DataflowIntrinsic))
+	FManagedArrayCollection Collection;
+
+	/** Face selection, the material(s) will be assigned to the selected faces */
+	UPROPERTY(meta = (DataflowInput, DisplayName = "FaceSelection", DataflowIntrinsic))
+	FDataflowFaceSelection FaceSelection;
+
+	/** Materials array storing the materials */
+	UPROPERTY(EditAnywhere, Category = "Material", meta = (DataflowInput, DataflowOutput, DataflowPassthrough = "Materials", DataflowIntrinsic))
+	TArray<TObjectPtr<UMaterial>> Materials;
+
+	/** Index of the material in the Materials array to assign to the outside faces from the face selection */
+	UPROPERTY(EditAnywhere, Category = "Material", meta = (DataflowInput, DisplayName = "Outside Material Index", ArrayClamp = "Materials"))
+	int32 OutsideMaterialIdx = -1;
+
+	/** Index of the material in the Materials array to assign to the inside faces from the face selection */
+	UPROPERTY(EditAnywhere, Category = "Material", meta = (DataflowInput, DisplayName = "Inside Material Index", ArrayClamp = "Materials"))
+	int32 InsideMaterialIdx = -1;
+
+	/** If true, the selected material from the Materials array will be assigned to the outside faces from the face selection */
+	UPROPERTY(EditAnywhere, Category = "Material", meta = (DisplayName = "Assign Outside Material to Outside Faces"))
+	bool bAssignOutsideMaterial = false;
+
+	/** If true, the selected material from the Materials array will be assigned to the inside faces from the face selection */
+	UPROPERTY(EditAnywhere, Category = "Material", meta = (DisplayName = "Assign Inside Material to Inside Faces"))
+	bool bAssignInsideMaterial = false;
+
+	FReAssignMaterialInCollectionDataflowNode(const Dataflow::FNodeParameters& InParam, FGuid InGuid = FGuid::NewGuid())
+		: FDataflowNode(InParam, InGuid)
+	{
+		RegisterInputConnection(&Collection);
+		RegisterInputConnection(&Materials);
+		RegisterInputConnection(&FaceSelection);
+		RegisterInputConnection(&OutsideMaterialIdx);
+		RegisterInputConnection(&InsideMaterialIdx);
+		RegisterOutputConnection(&Collection, &Collection);
+		RegisterOutputConnection(&Materials, &Materials);
+	}
+
+	virtual void Evaluate(Dataflow::FContext& Context, const FDataflowOutput* Out) const override;
+
+};
+
+
+/**
+ *
+ * Generates a formatted string of materials from the Materials array
+ *
+ */
+USTRUCT(meta = (DataflowGeometryCollection))
+struct FMaterialsInfoDataflowNode : public FDataflowNode
+{
+	GENERATED_USTRUCT_BODY()
+	DATAFLOW_NODE_DEFINE_INTERNAL(FMaterialsInfoDataflowNode, "MaterialsInfo", "GeometryCollection|Materials", "")
+
+public:
+	/** Materials array storing the materials */
+	UPROPERTY(meta = (DataflowInput, DataflowIntrinsic))
+	TArray<TObjectPtr<UMaterial>> Materials;
+
+	/** Formatted string of the materials */
+	UPROPERTY(meta = (DataflowOutput))
+	FString String;
+
+	FMaterialsInfoDataflowNode(const Dataflow::FNodeParameters& InParam, FGuid InGuid = FGuid::NewGuid())
+		: FDataflowNode(InParam, InGuid)
+	{
+		RegisterInputConnection(&Materials);
+		RegisterOutputConnection(&String);
+	}
+
+	virtual void Evaluate(Dataflow::FContext& Context, const FDataflowOutput* Out) const override;
+
+};
+
+
+/**
+ *
+ * Get a Material from a Materials array
+ *
+ */
+USTRUCT(meta = (DataflowGeometryCollection))
+struct FGetMaterialFromMaterialsArrayDataflowNode : public FDataflowNode
+{
+	GENERATED_USTRUCT_BODY()
+	DATAFLOW_NODE_DEFINE_INTERNAL(FGetMaterialFromMaterialsArrayDataflowNode, "GetMaterialFromMaterialsArray", "Utilities|Materials", "")
+
+public:
+	/** Materials array storing the materials */
+	UPROPERTY(meta = (DataflowInput, DataflowIntrinsic))
+	TArray<TObjectPtr<UMaterial>> Materials;
+
+	/** Selected material from the Materials array */
+	UPROPERTY(meta = (DataflowOutput))
+	TObjectPtr<UMaterial> Material;
+
+	/** Index in the Materials array for the selected material */
+	UPROPERTY(EditAnywhere, Category = "Material", meta = (DataflowInput, DisplayName = "Material Index"))
+	int32 MaterialIdx = 0;
+
+	FGetMaterialFromMaterialsArrayDataflowNode(const Dataflow::FNodeParameters& InParam, FGuid InGuid = FGuid::NewGuid())
+		: FDataflowNode(InParam, InGuid)
+	{
+		RegisterInputConnection(&Materials);
+		RegisterInputConnection(&MaterialIdx);
+		RegisterOutputConnection(&Material);
+	}
+
+	virtual void Evaluate(Dataflow::FContext& Context, const FDataflowOutput* Out) const override;
+
+};
+
+
+UENUM(BlueprintType)
+enum class ESetMaterialOperationTypeEnum : uint8
+{
+	Dataflow_SetMaterialOperationType_Add UMETA(DisplayName = "Add"),
+	Dataflow_SetMaterialOperationType_Insert UMETA(DisplayName = "Insert"),
+	//~~~
+	//256th entry
+	Dataflow_Max                UMETA(Hidden)
+};
+
+/**
+ *
+ * Set a Material in a Materials array
+ *
+ */
+USTRUCT(meta = (DataflowGeometryCollection))
+struct FSetMaterialInMaterialsArrayDataflowNode : public FDataflowNode
+{
+	GENERATED_USTRUCT_BODY()
+	DATAFLOW_NODE_DEFINE_INTERNAL(FSetMaterialInMaterialsArrayDataflowNode, "SetMaterialInMaterialsArray", "Utilities|Materials", "")
+
+public:
+	/** Materials array storing the materials */
+	UPROPERTY(meta = (DataflowInput, DataflowOutput, DataflowPassthrough = "Materials", DataflowIntrinsic))
+	TArray<TObjectPtr<UMaterial>> Materials;
+
+	/** Material to add/insert to/in Materials array */
+	UPROPERTY(EditAnywhere, Category = "Material", meta = (DataflowInput))
+	TObjectPtr<UMaterial> Material;
+
+	/** Operation type for setting the material, add will add the new material to the end off Materials array, insert will insert the 
+	new material into Materials array at the index specified by MaterialIdx	*/
+	UPROPERTY(EditAnywhere, Category = "Material")
+	ESetMaterialOperationTypeEnum Operation = ESetMaterialOperationTypeEnum::Dataflow_SetMaterialOperationType_Add;
+
+	/** Index for inserting a nem material into the Materials array */
+	UPROPERTY(EditAnywhere, Category = "Material", meta = (DataflowInput, DisplayName = "Material Index", EditCondition = "Operation == ESetMaterialOperationTypeEnum::Dataflow_SetMaterialOperationType_Insert", EditConditionHides))
+	int32 MaterialIdx = 0;
+
+	FSetMaterialInMaterialsArrayDataflowNode(const Dataflow::FNodeParameters& InParam, FGuid InGuid = FGuid::NewGuid())
+		: FDataflowNode(InParam, InGuid)
+	{
+		RegisterInputConnection(&Materials);
+		RegisterInputConnection(&Material);
+		RegisterInputConnection(&MaterialIdx);
+		RegisterOutputConnection(&Materials, &Materials);
+	}
+
+	virtual void Evaluate(Dataflow::FContext& Context, const FDataflowOutput* Out) const override;
+
+};
+
+
+/**
+ *
+ * Makes a material
+ *
+ */
+USTRUCT()
+struct FMakeMaterialDataflowNode : public FDataflowNode
+{
+	GENERATED_USTRUCT_BODY()
+	DATAFLOW_NODE_DEFINE_INTERNAL(FMakeMaterialDataflowNode, "MakeMaterial", "Generators|Material", "")
+
+public:
+	/** Material which will be outputed */
+	UPROPERTY(EditAnywhere, Category = "Material", meta = (DisplayName = "Material"))
+	TObjectPtr<UMaterial> InMaterial;
+
+	/** Output material */
+	UPROPERTY(meta = (DataflowOutput));
+	TObjectPtr<UMaterial> Material;
+
+	FMakeMaterialDataflowNode(const Dataflow::FNodeParameters& InParam, FGuid InGuid = FGuid::NewGuid())
+		: FDataflowNode(InParam, InGuid)
+	{
+		RegisterOutputConnection(&Material);
+	}
+
+	virtual void Evaluate(Dataflow::FContext& Context, const FDataflowOutput* Out) const override;
+
+};
+
+/**
+ *
+ * Makes an empty Materials array
+ *
+ */
+USTRUCT()
+struct FMakeMaterialsArrayDataflowNode : public FDataflowNode
+{
+	GENERATED_USTRUCT_BODY()
+	DATAFLOW_NODE_DEFINE_INTERNAL(FMakeMaterialsArrayDataflowNode, "MakeMaterialsArray", "Generators|Material", "")
+
+public:
+	/** Output Matarials array */
+	UPROPERTY(meta = (DataflowOutput));
+	TArray<TObjectPtr<UMaterial>> Materials;
+
+	FMakeMaterialsArrayDataflowNode(const Dataflow::FNodeParameters& InParam, FGuid InGuid = FGuid::NewGuid())
+		: FDataflowNode(InParam, InGuid)
+	{
+		RegisterOutputConnection(&Materials);
+	}
+
+	virtual void Evaluate(Dataflow::FContext& Context, const FDataflowOutput* Out) const override;
+
+};
+
+
 namespace Dataflow
 {
 	void GeometryCollectionEngineNodes();
