Dataflow:
- Added FRadialFalloffFieldDataflowNode, FPlaneFalloffFieldDataflowNode, FRadialIntMaskFieldDataflowNode, FUniformScalarFieldDataflowNode, FUniformVectorFieldDataflowNode, FRadialVectorFieldDataflowNode, FRandomVectorFieldDataflowNode, FNoiseFieldDataflowNode, FUniformIntegerFieldDataflowNode, FBoxFalloffFieldDataflowNode field nodes
- Added vertex color rendering to Dataflow Editor Viewport
- Moved material nodes to separate file

diff --git a/Engine/Plugins/Experimental/Dataflow/Source/DataflowEnginePlugin/Private/Dataflow/DataflowEngineRendering.cpp b/Engine/Plugins/Experimental/Dataflow/Source/DataflowEnginePlugin/Private/Dataflow/DataflowEngineRendering.cpp
--- a/Engine/Plugins/Experimental/Dataflow/Source/DataflowEnginePlugin/Private/Dataflow/DataflowEngineRendering.cpp
+++ b/Engine/Plugins/Experimental/Dataflow/Source/DataflowEnginePlugin/Private/Dataflow/DataflowEngineRendering.cpp
@@ -10,6 +10,7 @@
 #include "GeometryCollection/GeometryCollectionAlgo.h"
 #include "DynamicMesh/DynamicMesh3.h"
 #include "UDynamicMesh.h"
+#include "Dataflow/DataflowEnginePlugin.h"
 
 namespace Dataflow
 {
@@ -90,6 +91,19 @@ namespace Dataflow
 							for (int i = 0; i < Visited.Num(); i++) if (!Visited[i]) Vertices[i] = FVector3f(0);
 
 							RenderCollection.AddSurface(MoveTemp(Vertices), MoveTemp(Tris));
+							if (const TManagedArray<FLinearColor>* VertexColor = Collection.FindAttribute<FLinearColor>("Color", FGeometryCollection::VerticesGroup))
+							{
+								TManagedArray<FLinearColor>& RenderCollectionVertexColor = RenderCollection.ModifyVertexColor();
+								for (int32 Idx = 0; Idx < Vertex.Num(); ++Idx)
+								{
+									RenderCollectionVertexColor[Idx] = (*VertexColor)[Idx];
+								}
+							}
+							else
+							{
+								TManagedArray<FLinearColor>& RenderCollectionVertexColor = RenderCollection.ModifyVertexColor();
+								RenderCollectionVertexColor.Fill(IDataflowEnginePlugin::SurfaceColor);
+							}
 						}
 					}
 				}
diff --git a/Engine/Plugins/Experimental/Dataflow/Source/DataflowEnginePlugin/Private/Dataflow/DataflowEngineSceneProxy.cpp b/Engine/Plugins/Experimental/Dataflow/Source/DataflowEnginePlugin/Private/Dataflow/DataflowEngineSceneProxy.cpp
--- a/Engine/Plugins/Experimental/Dataflow/Source/DataflowEnginePlugin/Private/Dataflow/DataflowEngineSceneProxy.cpp
+++ b/Engine/Plugins/Experimental/Dataflow/Source/DataflowEnginePlugin/Private/Dataflow/DataflowEngineSceneProxy.cpp
@@ -186,6 +186,7 @@ void FDataflowEngineSceneProxy::CreateMeshRenderThreadResources()
 		const TManagedArray<FVector3f>& Vertex = Facade.GetVertices();
 		const TManagedArray<int32>& SelectionArray = Facade.GetSelectionState();
 		const TManagedArray<int32>& GeomIndex = Facade.GetVertexToGeometryIndex();
+		const TManagedArray<FLinearColor>& VertexColor = Facade.GetVertexColor();
 
 		// The color stored in the vertices actually gets interpreted as a linear
 		// color by the material, whereas it is more convenient for the user of the
@@ -218,10 +219,9 @@ void FDataflowEngineSceneProxy::CreateMeshRenderThreadResources()
 		VertexBuffers.StaticMeshVertexBuffer.SetVertexUV(VertexBufferIndex + 2, 0, FVector2f(0, 0));
 
 		FColor SelectionColor = (State.Mode == FDataflowSelectionState::EMode::DSS_Dataflow_Object) ? IDataflowEnginePlugin::SelectionPrimaryColor : IDataflowEnginePlugin::SelectionLockedPrimaryColor;
-		FColor FaceColor = SelectionArray[GeomIndex[Indices[i][0]]] ? SelectionColor : IDataflowEnginePlugin::SurfaceColor;
-		VertexBuffers.ColorVertexBuffer.VertexColor(VertexBufferIndex + 0) = FaceColor;
-		VertexBuffers.ColorVertexBuffer.VertexColor(VertexBufferIndex + 1) = FaceColor;
-		VertexBuffers.ColorVertexBuffer.VertexColor(VertexBufferIndex + 2) = FaceColor;
+		VertexBuffers.ColorVertexBuffer.VertexColor(VertexBufferIndex + 0) = SelectionArray[GeomIndex[Indices[i][0]]] ? SelectionColor : VertexColor[Indices[i][0]].ToFColor(true);
+		VertexBuffers.ColorVertexBuffer.VertexColor(VertexBufferIndex + 1) = SelectionArray[GeomIndex[Indices[i][1]]] ? SelectionColor : VertexColor[Indices[i][1]].ToFColor(true);
+		VertexBuffers.ColorVertexBuffer.VertexColor(VertexBufferIndex + 2) = SelectionArray[GeomIndex[Indices[i][2]]] ? SelectionColor : VertexColor[Indices[i][2]].ToFColor(true);
 
 		IndexBuffer.Indices[IndexBufferIndex + 0] = VertexBufferIndex + 0;
 		IndexBuffer.Indices[IndexBufferIndex + 1] = VertexBufferIndex + 1;
diff --git a/Engine/Plugins/Experimental/GeometryCollectionPlugin/Source/GeometryCollectionNodes/Private/Dataflow/GeometryCollectionFieldNodes.cpp b/Engine/Plugins/Experimental/GeometryCollectionPlugin/Source/GeometryCollectionNodes/Private/Dataflow/GeometryCollectionFieldNodes.cpp
new file mode 100644
--- /dev/null
+++ b/Engine/Plugins/Experimental/GeometryCollectionPlugin/Source/GeometryCollectionNodes/Private/Dataflow/GeometryCollectionFieldNodes.cpp
@@ -0,0 +1,815 @@
+// Copyright Epic Games, Inc. All Rights Reserved.
+
+#include "Dataflow/GeometryCollectionFieldNodes.h"
+#include "Dataflow/DataflowCore.h"
+
+#include "Engine/Engine.h"
+#include "Engine/StaticMesh.h"
+#include "GeometryCollection/Facades/CollectionMeshFacade.h"
+#include "GeometryCollection/GeometryCollectionObject.h"
+#include "GeometryCollection/ManagedArrayCollection.h"
+#include "GeometryCollection/GeometryCollection.h"
+#include "GeometryCollection/GeometryCollectionEngineUtility.h"
+#include "GeometryCollection/GeometryCollectionEngineRemoval.h"
+#include "GeometryCollection/GeometryCollectionEngineConversion.h"
+#include "Logging/LogMacros.h"
+#include "Templates/SharedPointer.h"
+#include "UObject/UnrealTypePrivate.h"
+#include "DynamicMeshToMeshDescription.h"
+#include "MeshDescriptionToDynamicMesh.h"
+#include "StaticMeshAttributes.h"
+#include "DynamicMeshEditor.h"
+#include "Operations/MeshBoolean.h"
+#include "Materials/Material.h"
+
+#include "EngineGlobals.h"
+#include "GeometryCollection/GeometryCollectionAlgo.h"
+#include "GeometryCollection/GeometryCollectionClusteringUtility.h"
+#include "GeometryCollection/GeometryCollectionConvexUtility.h"
+#include "Voronoi/Voronoi.h"
+#include "PlanarCut.h"
+#include "GeometryCollection/GeometryCollectionProximityUtility.h"
+#include "FractureEngineClustering.h"
+#include "FractureEngineSelection.h"
+#include "GeometryCollection/Facades/CollectionBoundsFacade.h"
+#include "GeometryCollection/Facades/CollectionAnchoringFacade.h"
+#include "GeometryCollection/Facades/CollectionRemoveOnBreakFacade.h"
+#include "GeometryCollection/Facades/CollectionTransformFacade.h"
+#include "DynamicMesh/MeshTransforms.h"
+#include "DynamicMesh/DynamicMesh3.h"
+
+#include "Field/FieldSystem.h"
+#include "Field/FieldSystemNodes.h"
+
+#include "GeometryCollection/Facades/CollectionMeshFacade.h"
+#include "GeometryCollection/Facades/CollectionTransformFacade.h"
+
+
+//#include UE_INLINE_GENERATED_CPP_BY_NAME(GeometryCollectionFieldNodes)
+
+namespace Dataflow
+{
+	void GeometryCollectionFieldNodes()
+	{
+		static const FLinearColor CDefaultNodeBodyTintColor = FLinearColor(0.f, 0.f, 0.f, 0.5f);
+
+		DATAFLOW_NODE_REGISTER_CREATION_FACTORY(FRadialFalloffFieldDataflowNode);
+		DATAFLOW_NODE_REGISTER_CREATION_FACTORY(FPlaneFalloffFieldDataflowNode);
+		DATAFLOW_NODE_REGISTER_CREATION_FACTORY(FRadialIntMaskFieldDataflowNode);
+		DATAFLOW_NODE_REGISTER_CREATION_FACTORY(FUniformScalarFieldDataflowNode);
+		DATAFLOW_NODE_REGISTER_CREATION_FACTORY(FUniformVectorFieldDataflowNode);
+		DATAFLOW_NODE_REGISTER_CREATION_FACTORY(FRadialVectorFieldDataflowNode);
+		DATAFLOW_NODE_REGISTER_CREATION_FACTORY(FRandomVectorFieldDataflowNode);
+		DATAFLOW_NODE_REGISTER_CREATION_FACTORY(FNoiseFieldDataflowNode);
+		DATAFLOW_NODE_REGISTER_CREATION_FACTORY(FUniformIntegerFieldDataflowNode);
+		DATAFLOW_NODE_REGISTER_CREATION_FACTORY(FBoxFalloffFieldDataflowNode);
+
+		// Field nodes
+		DATAFLOW_NODE_REGISTER_CREATION_FACTORY_NODE_COLORS_BY_CATEGORY("Fields", FLinearColor(.0f, 0.8f, 1.f), CDefaultNodeBodyTintColor);
+	}
+}
+
+
+void FRadialFalloffFieldDataflowNode::Evaluate(Dataflow::FContext& Context, const FDataflowOutput* Out) const
+{
+	if (Out->IsA<TArray<float>>(&WeightArray))
+	{
+		const TArray<FVector3f>& InVertexArray = GetValue<TArray<FVector3f>>(Context, &VertexArray);
+		const FSphere& InSphere = GetValue<FSphere>(Context, &Sphere);
+		const FVector& InTranslation = GetValue<FVector>(Context, &Translation);
+
+		const int32 NumVertices = InVertexArray.Num();
+
+		FFieldExecutionDatas ExecutionDatas;
+
+		FFieldContextIndex::ContiguousIndices(ExecutionDatas.SampleIndices, NumVertices);
+		ExecutionDatas.SamplePositions.Init(FVector(0.f), NumVertices);
+
+		for (int32 Index = 0; Index < NumVertices; Index++)
+		{
+			ExecutionDatas.SamplePositions[Index] = (FVector)InVertexArray[Index];
+		}
+
+		FRadialFalloff* RadialFalloffField = new FRadialFalloff();
+		RadialFalloffField->Position = InSphere.Center + InTranslation;
+		RadialFalloffField->Radius = InSphere.W;
+		RadialFalloffField->Magnitude = Magnitude;
+		RadialFalloffField->MinRange = MinRange;
+		RadialFalloffField->MaxRange = MaxRange;
+		RadialFalloffField->Default = Default;
+		RadialFalloffField->Falloff = (EFieldFalloffType)FalloffType;
+
+		FFieldContext FieldContext{
+			ExecutionDatas,
+			FFieldContext::UniquePointerMap(),
+			0.0
+		};
+
+		TArray<float> ResultsArray;
+		ResultsArray.Init(false, NumVertices);
+		TFieldArrayView<float> ResultsView(ResultsArray, 0, ResultsArray.Num());
+		RadialFalloffField->Evaluate(FieldContext, ResultsView);
+
+		TArray<float> NewWeightArray;
+		NewWeightArray.Init(0.f, NumVertices);
+
+		if (IsConnected<FDataflowVertexSelection>(&VertexSelection))
+		{
+			const FDataflowVertexSelection& InVertexSelection = GetValue<FDataflowVertexSelection>(Context, &VertexSelection);
+
+			if (InVertexArray.Num() == InVertexSelection.Num())
+			{
+				for (int32 Idx = 0; Idx < NumVertices; ++Idx)
+				{
+					if (InVertexSelection.IsSelected(Idx))
+					{
+						NewWeightArray[Idx] = ResultsView[Idx];
+					}
+				}
+			}
+		}
+		else
+		{
+			for (int32 Idx = 0; Idx < NumVertices; ++Idx)
+			{
+				NewWeightArray[Idx] = ResultsView[Idx];
+			}
+		}
+
+		SetValue<TArray<float>>(Context, NewWeightArray, &WeightArray);
+	}
+}
+
+
+void FPlaneFalloffFieldDataflowNode::Evaluate(Dataflow::FContext& Context, const FDataflowOutput* Out) const
+{
+	if (Out->IsA<TArray<float>>(&WeightArray))
+	{
+		const TArray<FVector3f>& InVertexArray = GetValue<TArray<FVector3f>>(Context, &VertexArray);
+		const FVector& InTranslation = GetValue<FVector>(Context, &Translation);
+		const FVector& InPosition = GetValue<FVector>(Context, &Position);
+		const FVector& InNormal = GetValue<FVector>(Context, &Normal);
+
+		const int32 NumVertices = InVertexArray.Num();
+
+		FFieldExecutionDatas ExecutionDatas;
+
+		FFieldContextIndex::ContiguousIndices(ExecutionDatas.SampleIndices, NumVertices);
+		ExecutionDatas.SamplePositions.Init(FVector(0.f), NumVertices);
+
+		for (int32 Index = 0; Index < NumVertices; Index++)
+		{
+			ExecutionDatas.SamplePositions[Index] = (FVector)InVertexArray[Index];
+		}
+
+		FPlaneFalloff* PlaneFalloffField = new FPlaneFalloff();
+		PlaneFalloffField->Position = InPosition + InTranslation;
+		PlaneFalloffField->Normal = InNormal;
+		PlaneFalloffField->Magnitude = Magnitude;
+		PlaneFalloffField->MinRange = MinRange;
+		PlaneFalloffField->MaxRange = MaxRange;
+		PlaneFalloffField->Distance = Distance;
+		PlaneFalloffField->Default = Default;
+		PlaneFalloffField->Falloff = (EFieldFalloffType)FalloffType;
+
+		FFieldContext FieldContext{
+			ExecutionDatas,
+			FFieldContext::UniquePointerMap(),
+			0.0
+		};
+
+		TArray<float> ResultsArray;
+		ResultsArray.Init(false, NumVertices);
+		TFieldArrayView<float> ResultsView(ResultsArray, 0, ResultsArray.Num());
+		PlaneFalloffField->Evaluate(FieldContext, ResultsView);
+
+		TArray<float> NewWeightArray;
+		NewWeightArray.Init(0.f, NumVertices);
+
+		if (IsConnected<FDataflowVertexSelection>(&VertexSelection))
+		{
+			const FDataflowVertexSelection& InVertexSelection = GetValue<FDataflowVertexSelection>(Context, &VertexSelection);
+
+			if (InVertexArray.Num() == InVertexSelection.Num())
+			{
+				for (int32 Idx = 0; Idx < NumVertices; ++Idx)
+				{
+					if (InVertexSelection.IsSelected(Idx))
+					{
+						NewWeightArray[Idx] = ResultsView[Idx];
+					}
+				}
+			}
+		}
+		else
+		{
+			for (int32 Idx = 0; Idx < NumVertices; ++Idx)
+			{
+				NewWeightArray[Idx] = ResultsView[Idx];
+			}
+		}
+
+		SetValue<TArray<float>>(Context, NewWeightArray, &WeightArray);
+	}
+}
+
+
+void FRadialIntMaskFieldDataflowNode::Evaluate(Dataflow::FContext& Context, const FDataflowOutput* Out) const
+{
+	if (Out->IsA<TArray<int32>>(&MaskArray))
+	{
+		const TArray<FVector3f>& InVertexArray = GetValue<TArray<FVector3f>>(Context, &VertexArray);
+		const FSphere& InSphere = GetValue<FSphere>(Context, &Sphere);
+		const FVector& InTranslation = GetValue<FVector>(Context, &Translation);
+
+		const int32 NumVertices = InVertexArray.Num();
+
+		FFieldExecutionDatas ExecutionDatas;
+
+		FFieldContextIndex::ContiguousIndices(ExecutionDatas.SampleIndices, NumVertices);
+		ExecutionDatas.SamplePositions.Init(FVector(0.f), NumVertices);
+
+		for (int32 Index = 0; Index < NumVertices; Index++)
+		{
+			ExecutionDatas.SamplePositions[Index] = (FVector)InVertexArray[Index];
+		}
+
+		FRadialIntMask* RadialIntMaskField = new FRadialIntMask();
+		RadialIntMaskField->Position = InSphere.Center + InTranslation;
+		RadialIntMaskField->Radius = InSphere.W;
+		RadialIntMaskField->InteriorValue = InteriorValue;
+		RadialIntMaskField->ExteriorValue = ExteriorValue;
+		RadialIntMaskField->SetMaskCondition = (ESetMaskConditionType)SetMaskCondition;
+
+		FFieldContext FieldContext{
+			ExecutionDatas,
+			FFieldContext::UniquePointerMap(),
+			0.0
+		};
+
+		TArray<int32> ResultsArray;
+		ResultsArray.Init(false, NumVertices);
+		TFieldArrayView<int32> ResultsView(ResultsArray, 0, ResultsArray.Num());
+		RadialIntMaskField->Evaluate(FieldContext, ResultsView);
+
+		TArray<int32> NewMaskArray;
+		NewMaskArray.Init(0, NumVertices);
+
+		if (IsConnected<FDataflowVertexSelection>(&VertexSelection))
+		{
+			const FDataflowVertexSelection& InVertexSelection = GetValue<FDataflowVertexSelection>(Context, &VertexSelection);
+
+			if (InVertexArray.Num() == InVertexSelection.Num())
+			{
+				for (int32 Idx = 0; Idx < NumVertices; ++Idx)
+				{
+					if (InVertexSelection.IsSelected(Idx))
+					{
+						NewMaskArray[Idx] = ResultsView[Idx];
+					}
+				}
+			}
+		}
+		else
+		{
+			for (int32 Idx = 0; Idx < NumVertices; ++Idx)
+			{
+				NewMaskArray[Idx] = ResultsView[Idx];
+			}
+		}
+
+		SetValue<TArray<int32>>(Context, NewMaskArray, &MaskArray);
+	}
+}
+
+
+void FUniformScalarFieldDataflowNode::Evaluate(Dataflow::FContext& Context, const FDataflowOutput* Out) const
+{
+	if (Out->IsA<TArray<float>>(&WeightArray))
+	{
+		const TArray<FVector3f>& InVertexArray = GetValue<TArray<FVector3f>>(Context, &VertexArray);
+
+		const int32 NumVertices = InVertexArray.Num();
+
+		FFieldExecutionDatas ExecutionDatas;
+
+		FFieldContextIndex::ContiguousIndices(ExecutionDatas.SampleIndices, NumVertices);
+		ExecutionDatas.SamplePositions.Init(FVector(0.f), NumVertices);
+
+		for (int32 Index = 0; Index < NumVertices; Index++)
+		{
+			ExecutionDatas.SamplePositions[Index] = (FVector)InVertexArray[Index];
+		}
+
+		FUniformScalar* UniformScalarField = new FUniformScalar();
+		UniformScalarField->Magnitude = Magnitude;
+
+		FFieldContext FieldContext{
+			ExecutionDatas,
+			FFieldContext::UniquePointerMap(),
+			0.0
+		};
+
+		TArray<float> ResultsArray;
+		ResultsArray.Init(false, NumVertices);
+		TFieldArrayView<float> ResultsView(ResultsArray, 0, ResultsArray.Num());
+		UniformScalarField->Evaluate(FieldContext, ResultsView);
+
+		TArray<float> NewWeightArray;
+		NewWeightArray.Init(0.f, NumVertices);
+
+		if (IsConnected<FDataflowVertexSelection>(&VertexSelection))
+		{
+			const FDataflowVertexSelection& InVertexSelection = GetValue<FDataflowVertexSelection>(Context, &VertexSelection);
+
+			if (InVertexArray.Num() == InVertexSelection.Num())
+			{
+				for (int32 Idx = 0; Idx < NumVertices; ++Idx)
+				{
+					if (InVertexSelection.IsSelected(Idx))
+					{
+						NewWeightArray[Idx] = ResultsView[Idx];
+					}
+				}
+			}
+		}
+		else
+		{
+			for (int32 Idx = 0; Idx < NumVertices; ++Idx)
+			{
+				NewWeightArray[Idx] = ResultsView[Idx];
+			}
+		}
+
+		SetValue<TArray<float>>(Context, NewWeightArray, &WeightArray);
+	}
+}
+
+
+void FUniformVectorFieldDataflowNode::Evaluate(Dataflow::FContext& Context, const FDataflowOutput* Out) const
+{
+	if (Out->IsA<TArray<FVector>>(&VectorArray))
+	{
+		const TArray<FVector3f>& InVertexArray = GetValue<TArray<FVector3f>>(Context, &VertexArray);
+
+		const int32 NumVertices = InVertexArray.Num();
+
+		FFieldExecutionDatas ExecutionDatas;
+
+		FFieldContextIndex::ContiguousIndices(ExecutionDatas.SampleIndices, NumVertices);
+		ExecutionDatas.SamplePositions.Init(FVector(0.f), NumVertices);
+
+		for (int32 Index = 0; Index < NumVertices; Index++)
+		{
+			ExecutionDatas.SamplePositions[Index] = (FVector)InVertexArray[Index];
+		}
+
+		FUniformVector* UniformVectorField = new FUniformVector();
+		UniformVectorField->Magnitude = Magnitude;
+		UniformVectorField->Direction = Direction;
+
+		FFieldContext FieldContext{
+			ExecutionDatas,
+			FFieldContext::UniquePointerMap(),
+			0.0
+		};
+
+		TArray<FVector> ResultsArray;
+		ResultsArray.Init(FVector(0.f), NumVertices);
+		TFieldArrayView<FVector> ResultsView(ResultsArray, 0, ResultsArray.Num());
+		UniformVectorField->Evaluate(FieldContext, ResultsView);
+
+		TArray<FVector> NewVectorArray;
+		NewVectorArray.Init(FVector(0.f), NumVertices);
+
+		if (IsConnected<FDataflowVertexSelection>(&VertexSelection))
+		{
+			const FDataflowVertexSelection& InVertexSelection = GetValue<FDataflowVertexSelection>(Context, &VertexSelection);
+
+			if (InVertexArray.Num() == InVertexSelection.Num())
+			{
+				for (int32 Idx = 0; Idx < NumVertices; ++Idx)
+				{
+					if (InVertexSelection.IsSelected(Idx))
+					{
+						NewVectorArray[Idx] = ResultsView[Idx];
+					}
+				}
+			}
+		}
+		else
+		{
+			for (int32 Idx = 0; Idx < NumVertices; ++Idx)
+			{
+				NewVectorArray[Idx] = ResultsView[Idx];
+			}
+		}
+
+		SetValue<TArray<FVector>>(Context, NewVectorArray, &VectorArray);
+	}
+}
+
+
+void FRadialVectorFieldDataflowNode::Evaluate(Dataflow::FContext& Context, const FDataflowOutput* Out) const
+{
+	if (Out->IsA<TArray<FVector>>(&VectorArray))
+	{
+		const TArray<FVector3f>& InVertexArray = GetValue<TArray<FVector3f>>(Context, &VertexArray);
+
+		const int32 NumVertices = InVertexArray.Num();
+
+		FFieldExecutionDatas ExecutionDatas;
+
+		FFieldContextIndex::ContiguousIndices(ExecutionDatas.SampleIndices, NumVertices);
+		ExecutionDatas.SamplePositions.Init(FVector(0.f), NumVertices);
+
+		for (int32 Index = 0; Index < NumVertices; Index++)
+		{
+			ExecutionDatas.SamplePositions[Index] = (FVector)InVertexArray[Index];
+		}
+
+		FRadialVector* RadialVectorField = new FRadialVector();
+		RadialVectorField->Magnitude = Magnitude;
+		RadialVectorField->Position = Position;
+
+		FFieldContext FieldContext{
+			ExecutionDatas,
+			FFieldContext::UniquePointerMap(),
+			0.0
+		};
+
+		TArray<FVector> ResultsArray;
+		ResultsArray.Init(FVector(0.f), NumVertices);
+		TFieldArrayView<FVector> ResultsView(ResultsArray, 0, ResultsArray.Num());
+		RadialVectorField->Evaluate(FieldContext, ResultsView);
+
+		TArray<FVector> NewVectorArray;
+		NewVectorArray.Init(FVector(0.f), NumVertices);
+
+		if (IsConnected<FDataflowVertexSelection>(&VertexSelection))
+		{
+			const FDataflowVertexSelection& InVertexSelection = GetValue<FDataflowVertexSelection>(Context, &VertexSelection);
+
+			if (InVertexArray.Num() == InVertexSelection.Num())
+			{
+				for (int32 Idx = 0; Idx < NumVertices; ++Idx)
+				{
+					if (InVertexSelection.IsSelected(Idx))
+					{
+						NewVectorArray[Idx] = ResultsView[Idx];
+					}
+				}
+			}
+		}
+		else
+		{
+			for (int32 Idx = 0; Idx < NumVertices; ++Idx)
+			{
+				NewVectorArray[Idx] = ResultsView[Idx];
+			}
+		}
+
+		SetValue<TArray<FVector>>(Context, NewVectorArray, &VectorArray);
+	}
+}
+
+
+void FRandomVectorFieldDataflowNode::Evaluate(Dataflow::FContext& Context, const FDataflowOutput* Out) const
+{
+	if (Out->IsA<TArray<FVector>>(&VectorArray))
+	{
+		const TArray<FVector3f>& InVertexArray = GetValue<TArray<FVector3f>>(Context, &VertexArray);
+
+		const int32 NumVertices = InVertexArray.Num();
+
+		FFieldExecutionDatas ExecutionDatas;
+
+		FFieldContextIndex::ContiguousIndices(ExecutionDatas.SampleIndices, NumVertices);
+		ExecutionDatas.SamplePositions.Init(FVector(0.f), NumVertices);
+
+		for (int32 Index = 0; Index < NumVertices; Index++)
+		{
+			ExecutionDatas.SamplePositions[Index] = (FVector)InVertexArray[Index];
+		}
+
+		FRandomVector* RandomVectorField = new FRandomVector();
+		RandomVectorField->Magnitude = Magnitude;
+
+		FFieldContext FieldContext{
+			ExecutionDatas,
+			FFieldContext::UniquePointerMap(),
+			0.0
+		};
+
+		TArray<FVector> ResultsArray;
+		ResultsArray.Init(FVector(0.f), NumVertices);
+		TFieldArrayView<FVector> ResultsView(ResultsArray, 0, ResultsArray.Num());
+		RandomVectorField->Evaluate(FieldContext, ResultsView);
+
+		TArray<FVector> NewVectorArray;
+		NewVectorArray.Init(FVector(0.f), NumVertices);
+
+		if (IsConnected<FDataflowVertexSelection>(&VertexSelection))
+		{
+			const FDataflowVertexSelection& InVertexSelection = GetValue<FDataflowVertexSelection>(Context, &VertexSelection);
+
+			if (InVertexArray.Num() == InVertexSelection.Num())
+			{
+				for (int32 Idx = 0; Idx < NumVertices; ++Idx)
+				{
+					if (InVertexSelection.IsSelected(Idx))
+					{
+						NewVectorArray[Idx] = ResultsView[Idx];
+					}
+				}
+			}
+		}
+		else
+		{
+			for (int32 Idx = 0; Idx < NumVertices; ++Idx)
+			{
+				NewVectorArray[Idx] = ResultsView[Idx];
+			}
+		}
+
+		SetValue<TArray<FVector>>(Context, NewVectorArray, &VectorArray);
+	}
+}
+
+
+
+void FNoiseFieldDataflowNode::Evaluate(Dataflow::FContext& Context, const FDataflowOutput* Out) const
+{
+	if (Out->IsA<TArray<float>>(&WeightArray))
+	{
+		const TArray<FVector3f>& InVertexArray = GetValue<TArray<FVector3f>>(Context, &VertexArray);
+
+		const int32 NumVertices = InVertexArray.Num();
+
+		FFieldExecutionDatas ExecutionDatas;
+
+		FFieldContextIndex::ContiguousIndices(ExecutionDatas.SampleIndices, NumVertices);
+		ExecutionDatas.SamplePositions.Init(FVector(0.f), NumVertices);
+
+		for (int32 Index = 0; Index < NumVertices; Index++)
+		{
+			ExecutionDatas.SamplePositions[Index] = (FVector)InVertexArray[Index];
+		}
+
+		FNoiseField* NoiseField = new FNoiseField();
+		NoiseField->MinRange = MinRange;
+		NoiseField->MaxRange = MaxRange;
+		NoiseField->Transform = Transform;
+
+		FFieldContext FieldContext{
+			ExecutionDatas,
+			FFieldContext::UniquePointerMap(),
+			0.0
+		};
+
+		TArray<float> ResultsArray;
+		ResultsArray.Init(false, NumVertices);
+		TFieldArrayView<float> ResultsView(ResultsArray, 0, ResultsArray.Num());
+		NoiseField->Evaluate(FieldContext, ResultsView);
+
+		TArray<float> NewWeightArray;
+		NewWeightArray.Init(0.f, NumVertices);
+
+		if (IsConnected<FDataflowVertexSelection>(&VertexSelection))
+		{
+			const FDataflowVertexSelection& InVertexSelection = GetValue<FDataflowVertexSelection>(Context, &VertexSelection);
+
+			if (InVertexArray.Num() == InVertexSelection.Num())
+			{
+				for (int32 Idx = 0; Idx < NumVertices; ++Idx)
+				{
+					if (InVertexSelection.IsSelected(Idx))
+					{
+						NewWeightArray[Idx] = ResultsView[Idx];
+					}
+				}
+			}
+		}
+		else
+		{
+			for (int32 Idx = 0; Idx < NumVertices; ++Idx)
+			{
+				NewWeightArray[Idx] = ResultsView[Idx];
+			}
+		}
+
+		SetValue<TArray<float>>(Context, NewWeightArray, &WeightArray);
+	}
+}
+
+
+void FUniformIntegerFieldDataflowNode::Evaluate(Dataflow::FContext& Context, const FDataflowOutput* Out) const
+{
+	if (Out->IsA<TArray<int32>>(&MaskArray))
+	{
+		const TArray<FVector3f>& InVertexArray = GetValue<TArray<FVector3f>>(Context, &VertexArray);
+
+		const int32 NumVertices = InVertexArray.Num();
+
+		FFieldExecutionDatas ExecutionDatas;
+
+		FFieldContextIndex::ContiguousIndices(ExecutionDatas.SampleIndices, NumVertices);
+		ExecutionDatas.SamplePositions.Init(FVector(0.f), NumVertices);
+
+		for (int32 Index = 0; Index < NumVertices; Index++)
+		{
+			ExecutionDatas.SamplePositions[Index] = (FVector)InVertexArray[Index];
+		}
+
+		FUniformInteger* UniformIntegerField = new FUniformInteger();
+		UniformIntegerField->Magnitude = Magnitude;
+
+		FFieldContext FieldContext{
+			ExecutionDatas,
+			FFieldContext::UniquePointerMap(),
+			0.0
+		};
+
+		TArray<int32> ResultsArray;
+		ResultsArray.Init(false, NumVertices);
+		TFieldArrayView<int32> ResultsView(ResultsArray, 0, ResultsArray.Num());
+		UniformIntegerField->Evaluate(FieldContext, ResultsView);
+
+		TArray<int32> NewMaskArray;
+		NewMaskArray.Init(0, NumVertices);
+
+		if (IsConnected<FDataflowVertexSelection>(&VertexSelection))
+		{
+			const FDataflowVertexSelection& InVertexSelection = GetValue<FDataflowVertexSelection>(Context, &VertexSelection);
+
+			if (InVertexArray.Num() == InVertexSelection.Num())
+			{
+				for (int32 Idx = 0; Idx < NumVertices; ++Idx)
+				{
+					if (InVertexSelection.IsSelected(Idx))
+					{
+						NewMaskArray[Idx] = ResultsView[Idx];
+					}
+				}
+			}
+		}
+		else
+		{
+			for (int32 Idx = 0; Idx < NumVertices; ++Idx)
+			{
+				NewMaskArray[Idx] = ResultsView[Idx];
+			}
+		}
+
+		SetValue<TArray<int32>>(Context, NewMaskArray, &MaskArray);
+	}
+}
+
+
+void FWaveScalarFieldDataflowNode::Evaluate(Dataflow::FContext& Context, const FDataflowOutput* Out) const
+{
+	if (Out->IsA<TArray<float>>(&WeightArray))
+	{
+		const TArray<FVector3f>& InVertexArray = GetValue<TArray<FVector3f>>(Context, &VertexArray);
+		const FVector& InTranslation = GetValue<FVector>(Context, &Translation);
+
+		const int32 NumVertices = InVertexArray.Num();
+
+		FFieldExecutionDatas ExecutionDatas;
+
+		FFieldContextIndex::ContiguousIndices(ExecutionDatas.SampleIndices, NumVertices);
+		ExecutionDatas.SamplePositions.Init(FVector(0.f), NumVertices);
+
+		for (int32 Index = 0; Index < NumVertices; Index++)
+		{
+			ExecutionDatas.SamplePositions[Index] = (FVector)InVertexArray[Index];
+		}
+
+		FWaveScalar* WaveScalarField = new FWaveScalar();
+		WaveScalarField->Magnitude = Magnitude;
+		WaveScalarField->Position = Position + InTranslation;
+		WaveScalarField->Wavelength = Wavelength;
+		WaveScalarField->Period = Period;
+		WaveScalarField->Function = (EWaveFunctionType)FunctionType;
+		WaveScalarField->Falloff = (EFieldFalloffType)FalloffType;
+
+		FFieldContext FieldContext{
+			ExecutionDatas,
+			FFieldContext::UniquePointerMap(),
+			0.0
+		};
+
+		TArray<float> ResultsArray;
+		ResultsArray.Init(false, NumVertices);
+		TFieldArrayView<float> ResultsView(ResultsArray, 0, ResultsArray.Num());
+		WaveScalarField->Evaluate(FieldContext, ResultsView);
+
+		TArray<float> NewWeightArray;
+		NewWeightArray.Init(0.f, NumVertices);
+
+		if (IsConnected<FDataflowVertexSelection>(&VertexSelection))
+		{
+			const FDataflowVertexSelection& InVertexSelection = GetValue<FDataflowVertexSelection>(Context, &VertexSelection);
+
+			if (InVertexArray.Num() == InVertexSelection.Num())
+			{
+				for (int32 Idx = 0; Idx < NumVertices; ++Idx)
+				{
+					if (InVertexSelection.IsSelected(Idx))
+					{
+						NewWeightArray[Idx] = ResultsView[Idx];
+					}
+				}
+			}
+		}
+		else
+		{
+			for (int32 Idx = 0; Idx < NumVertices; ++Idx)
+			{
+				NewWeightArray[Idx] = ResultsView[Idx];
+			}
+		}
+
+		SetValue<TArray<float>>(Context, NewWeightArray, &WeightArray);
+	}
+}
+
+
+void FBoxFalloffFieldDataflowNode::Evaluate(Dataflow::FContext& Context, const FDataflowOutput* Out) const
+{
+	if (Out->IsA<TArray<float>>(&WeightArray))
+	{
+		const TArray<FVector3f>& InVertexArray = GetValue<TArray<FVector3f>>(Context, &VertexArray);
+		const FBox& InBox = GetValue<FBox>(Context, &Box);
+		const FTransform& InTransform = GetValue<FTransform>(Context, &Transform);
+
+		const int32 NumVertices = InVertexArray.Num();
+
+		FFieldExecutionDatas ExecutionDatas;
+
+		FFieldContextIndex::ContiguousIndices(ExecutionDatas.SampleIndices, NumVertices);
+		ExecutionDatas.SamplePositions.Init(FVector(0.f), NumVertices);
+
+		for (int32 Index = 0; Index < NumVertices; Index++)
+		{
+			ExecutionDatas.SamplePositions[Index] = (FVector)InVertexArray[Index];
+		}
+
+		FBoxFalloff* BoxFalloffField = new FBoxFalloff();
+		BoxFalloffField->Magnitude = Magnitude;
+		BoxFalloffField->MinRange = MinRange;
+		BoxFalloffField->MaxRange = MaxRange;
+		BoxFalloffField->Default = Default;
+		BoxFalloffField->Transform = InTransform;
+		BoxFalloffField->Falloff = (EFieldFalloffType)FalloffType;
+
+		FFieldContext FieldContext{
+			ExecutionDatas,
+			FFieldContext::UniquePointerMap(),
+			0.0
+		};
+
+		TArray<float> ResultsArray;
+		ResultsArray.Init(false, NumVertices);
+		TFieldArrayView<float> ResultsView(ResultsArray, 0, ResultsArray.Num());
+		BoxFalloffField->Evaluate(FieldContext, ResultsView);
+
+		TArray<float> NewWeightArray;
+		NewWeightArray.Init(0.f, NumVertices);
+
+		if (IsConnected<FDataflowVertexSelection>(&VertexSelection))
+		{
+			const FDataflowVertexSelection& InVertexSelection = GetValue<FDataflowVertexSelection>(Context, &VertexSelection);
+
+			if (InVertexArray.Num() == InVertexSelection.Num())
+			{
+				for (int32 Idx = 0; Idx < NumVertices; ++Idx)
+				{
+					if (InVertexSelection.IsSelected(Idx))
+					{
+						NewWeightArray[Idx] = ResultsView[Idx];
+					}
+				}
+			}
+		}
+		else
+		{
+			for (int32 Idx = 0; Idx < NumVertices; ++Idx)
+			{
+				NewWeightArray[Idx] = ResultsView[Idx];
+			}
+		}
+
+		SetValue<TArray<float>>(Context, NewWeightArray, &WeightArray);
+	}
+}
+
+
+
+
+
+
+
+
+
+
diff --git a/Engine/Plugins/Experimental/GeometryCollectionPlugin/Source/GeometryCollectionNodes/Private/Dataflow/GeometryCollectionMaterialNodes.cpp b/Engine/Plugins/Experimental/GeometryCollectionPlugin/Source/GeometryCollectionNodes/Private/Dataflow/GeometryCollectionMaterialNodes.cpp
new file mode 100644
--- /dev/null
+++ b/Engine/Plugins/Experimental/GeometryCollectionPlugin/Source/GeometryCollectionNodes/Private/Dataflow/GeometryCollectionMaterialNodes.cpp
@@ -0,0 +1,268 @@
+// Copyright Epic Games, Inc. All Rights Reserved.
+
+#include "Dataflow/GeometryCollectionMaterialNodes.h"
+#include "Dataflow/DataflowCore.h"
+
+#include "Engine/Engine.h"
+#include "Engine/StaticMesh.h"
+#include "GeometryCollection/Facades/CollectionMeshFacade.h"
+#include "GeometryCollection/GeometryCollectionObject.h"
+#include "GeometryCollection/ManagedArrayCollection.h"
+#include "GeometryCollection/GeometryCollection.h"
+#include "GeometryCollection/GeometryCollectionEngineUtility.h"
+#include "GeometryCollection/GeometryCollectionEngineRemoval.h"
+#include "GeometryCollection/GeometryCollectionEngineConversion.h"
+#include "Logging/LogMacros.h"
+#include "Templates/SharedPointer.h"
+#include "UObject/UnrealTypePrivate.h"
+#include "DynamicMeshToMeshDescription.h"
+#include "MeshDescriptionToDynamicMesh.h"
+#include "StaticMeshAttributes.h"
+#include "DynamicMeshEditor.h"
+#include "Operations/MeshBoolean.h"
+#include "Materials/Material.h"
+
+#include "EngineGlobals.h"
+#include "GeometryCollection/GeometryCollectionAlgo.h"
+#include "GeometryCollection/GeometryCollectionClusteringUtility.h"
+#include "GeometryCollection/GeometryCollectionConvexUtility.h"
+#include "Voronoi/Voronoi.h"
+#include "PlanarCut.h"
+#include "GeometryCollection/GeometryCollectionProximityUtility.h"
+#include "FractureEngineClustering.h"
+#include "FractureEngineSelection.h"
+#include "GeometryCollection/Facades/CollectionBoundsFacade.h"
+#include "GeometryCollection/Facades/CollectionAnchoringFacade.h"
+#include "GeometryCollection/Facades/CollectionRemoveOnBreakFacade.h"
+#include "GeometryCollection/Facades/CollectionTransformFacade.h"
+#include "DynamicMesh/MeshTransforms.h"
+#include "DynamicMesh/DynamicMesh3.h"
+
+//#include UE_INLINE_GENERATED_CPP_BY_NAME(GeometryCollectionMaterialNodes)
+
+namespace Dataflow
+{
+	void GeometryCollectionMaterialNodes()
+	{
+		static const FLinearColor CDefaultNodeBodyTintColor = FLinearColor(0.f, 0.f, 0.f, 0.5f);
+
+		DATAFLOW_NODE_REGISTER_CREATION_FACTORY(FAddMaterialToCollectionDataflowNode);
+		DATAFLOW_NODE_REGISTER_CREATION_FACTORY(FReAssignMaterialInCollectionDataflowNode);
+		DATAFLOW_NODE_REGISTER_CREATION_FACTORY(FMaterialsInfoDataflowNode);
+		DATAFLOW_NODE_REGISTER_CREATION_FACTORY(FGetMaterialFromMaterialsArrayDataflowNode);
+		DATAFLOW_NODE_REGISTER_CREATION_FACTORY(FSetMaterialInMaterialsArrayDataflowNode);
+		DATAFLOW_NODE_REGISTER_CREATION_FACTORY(FMakeMaterialDataflowNode);
+		DATAFLOW_NODE_REGISTER_CREATION_FACTORY(FMakeMaterialsArrayDataflowNode);
+
+		// Material
+		DATAFLOW_NODE_REGISTER_CREATION_FACTORY_NODE_COLORS_BY_CATEGORY("Generators", FLinearColor(.7f, 0.7f, 0.7f), CDefaultNodeBodyTintColor);
+	}
+}
+
+
+void FAddMaterialToCollectionDataflowNode::Evaluate(Dataflow::FContext& Context, const FDataflowOutput* Out) const
+{
+	if (Out->IsA<FManagedArrayCollection>(&Collection) || Out->IsA<TArray<TObjectPtr<UMaterial>>>(&Materials))
+	{
+		FManagedArrayCollection InCollection = GetValue<FManagedArrayCollection>(Context, &Collection);
+		TArray<TObjectPtr<UMaterial>> InMaterials = GetValue<TArray<TObjectPtr<UMaterial>>>(Context, &Materials);
+		const FDataflowFaceSelection& InFaceSelection = GetValue<FDataflowFaceSelection>(Context, &FaceSelection);
+		TObjectPtr<UMaterial> InOutsideMaterial = GetValue<TObjectPtr<UMaterial>>(Context, &OutsideMaterial);
+		TObjectPtr<UMaterial> InInsideMaterial = GetValue<TObjectPtr<UMaterial>>(Context, &InsideMaterial);
+
+		int32 NewIndex;
+
+		if (InCollection.HasAttribute("Internal", FGeometryCollection::FacesGroup))
+		{
+			if (bAssignOutsideMaterial || bAssignInsideMaterial)
+			{
+				TManagedArray<int32>& MaterialIDs = InCollection.ModifyAttribute<int32>("MaterialID", FGeometryCollection::FacesGroup);
+
+				const TManagedArray<bool>& Internals = InCollection.GetAttribute<bool>("Internal", FGeometryCollection::FacesGroup);
+
+				if (bAssignOutsideMaterial)
+				{
+					NewIndex = InMaterials.Add(InOutsideMaterial);
+
+					// Update MaterialIdx for selected outside faces
+					for (int32 FaceIdx = 0; FaceIdx < InFaceSelection.Num(); ++FaceIdx)
+					{
+						if (!Internals[FaceIdx] && InFaceSelection.IsSelected(FaceIdx))
+						{
+							MaterialIDs[FaceIdx] = NewIndex;
+						}
+					}
+				}
+
+				if (bAssignInsideMaterial)
+				{
+					NewIndex = InMaterials.Add(InInsideMaterial);
+
+					// Update MaterialIdx for selected inside faces
+					for (int32 FaceIdx = 0; FaceIdx < InFaceSelection.Num(); ++FaceIdx)
+					{
+						if (Internals[FaceIdx] && InFaceSelection.IsSelected(FaceIdx))
+						{
+							MaterialIDs[FaceIdx] = NewIndex;
+						}
+					}
+				}
+			}
+		}
+
+		SetValue<TArray<TObjectPtr<UMaterial>>>(Context, MoveTemp(InMaterials), &Materials);
+		SetValue<FManagedArrayCollection>(Context, MoveTemp(InCollection), &Collection);
+	}
+}
+
+
+void FReAssignMaterialInCollectionDataflowNode::Evaluate(Dataflow::FContext& Context, const FDataflowOutput* Out) const
+{
+	if (Out->IsA<FManagedArrayCollection>(&Collection))
+	{
+		FManagedArrayCollection InCollection = GetValue<FManagedArrayCollection>(Context, &Collection);
+		const FDataflowFaceSelection& InFaceSelection = GetValue<FDataflowFaceSelection>(Context, &FaceSelection);
+		const TArray<TObjectPtr<UMaterial>>& InMaterials = GetValue<TArray<TObjectPtr<UMaterial>>>(Context, &Materials);
+		int32 InOutsideMaterialIdx = GetValue<int32>(Context, &OutsideMaterialIdx);
+		int32 InInsideMaterialIdx = GetValue<int32>(Context, &InsideMaterialIdx);
+
+		if (InMaterials.Num() > 0)
+		{
+			if (InCollection.HasAttribute("Internal", FGeometryCollection::FacesGroup))
+			{
+				if (bAssignOutsideMaterial || bAssignInsideMaterial)
+				{
+					TManagedArray<int32>& MaterialIDs = InCollection.ModifyAttribute<int32>("MaterialID", FGeometryCollection::FacesGroup);
+
+					const TManagedArray<bool>& Internals = InCollection.GetAttribute<bool>("Internal", FGeometryCollection::FacesGroup);
+
+					if (bAssignOutsideMaterial)
+					{
+						if (InOutsideMaterialIdx >= 0 && InOutsideMaterialIdx < InMaterials.Num())
+						{
+							// Update MaterialIdx for selected outside faces
+							for (int32 FaceIdx = 0; FaceIdx < InFaceSelection.Num(); ++FaceIdx)
+							{
+								if (!Internals[FaceIdx] && InFaceSelection.IsSelected(FaceIdx))
+								{
+									MaterialIDs[FaceIdx] = InOutsideMaterialIdx;
+								}
+							}
+						}
+					}
+
+					if (bAssignInsideMaterial)
+					{
+						if (InInsideMaterialIdx >= 0 && InInsideMaterialIdx < InMaterials.Num())
+						{
+							// Update MaterialIdx for selected inside faces
+							for (int32 FaceIdx = 0; FaceIdx < InFaceSelection.Num(); ++FaceIdx)
+							{
+								if (Internals[FaceIdx] && InFaceSelection.IsSelected(FaceIdx))
+								{
+									MaterialIDs[FaceIdx] = InInsideMaterialIdx;
+								}
+							}
+						}
+					}
+				}
+			}
+		}
+
+		// move the collection to the output to avoid making another copy
+		SetValue<FManagedArrayCollection>(Context, MoveTemp(InCollection), &Collection);
+		SetValue<TArray<TObjectPtr<UMaterial>>>(Context, InMaterials, &Materials);
+	}
+}
+
+
+void FMaterialsInfoDataflowNode::Evaluate(Dataflow::FContext& Context, const FDataflowOutput* Out) const
+{
+	if (Out->IsA<FString>(&String))
+	{
+		const TArray<TObjectPtr<UMaterial>>& InMaterials = GetValue<TArray<TObjectPtr<UMaterial>>>(Context, &Materials);
+
+		FString OutputStr;
+
+		OutputStr.Appendf(TEXT("\n----------------------------------------\n"));
+		OutputStr.Appendf(TEXT("Number of Materials: %d\n"), InMaterials.Num());
+
+		int32 Idx = 0;
+		for (TObjectPtr<UMaterial> Material : InMaterials)
+		{
+			OutputStr.Appendf(TEXT("%4d: %s\n"), Idx, *(Material->GetFullName()));
+
+			Idx++;
+		}
+
+		OutputStr.Appendf(TEXT("----------------------------------------\n"));
+
+		SetValue<FString>(Context, OutputStr, &String);
+	}
+}
+
+
+void FGetMaterialFromMaterialsArrayDataflowNode::Evaluate(Dataflow::FContext& Context, const FDataflowOutput* Out) const
+{
+	if (Out->IsA<TObjectPtr<UMaterial>>(&Material))
+	{
+		const TArray<TObjectPtr<UMaterial>>& InMaterials = GetValue<TArray<TObjectPtr<UMaterial>>>(Context, &Materials);
+		int32 InMaterialIdx = GetValue<int32>(Context, &MaterialIdx);
+
+		if (InMaterials.Num() > 0)
+		{
+			if (InMaterialIdx >= 0 && InMaterialIdx < InMaterials.Num())
+			{
+				SetValue<TObjectPtr<UMaterial>>(Context, InMaterials[InMaterialIdx], &Material);
+
+				return;
+			}
+		}
+
+		SetValue<TObjectPtr<UMaterial>>(Context, TObjectPtr<UMaterial>(), &Material);
+	}
+}
+
+
+void FSetMaterialInMaterialsArrayDataflowNode::Evaluate(Dataflow::FContext& Context, const FDataflowOutput* Out) const
+{
+	if (Out->IsA<TArray<TObjectPtr<UMaterial>>>(&Materials))
+	{
+		TArray<TObjectPtr<UMaterial>> InMaterials = GetValue<TArray<TObjectPtr<UMaterial>>>(Context, &Materials);
+		TObjectPtr<UMaterial> InMaterial = GetValue<TObjectPtr<UMaterial>>(Context, &Material);
+		int32 InMaterialIdx = GetValue<int32>(Context, &MaterialIdx);
+
+		if (Operation == ESetMaterialOperationTypeEnum::Dataflow_SetMaterialOperationType_Add)
+		{
+			InMaterials.Add(InMaterial);
+		}
+		else
+		{
+			if (InMaterialIdx >= 0 && InMaterialIdx < InMaterials.Num())
+			{
+				InMaterials.Insert(InMaterial, InMaterialIdx);
+			}
+		}
+
+		SetValue<TArray<TObjectPtr<UMaterial>>>(Context, MoveTemp(InMaterials), &Materials);
+	}
+}
+
+
+void FMakeMaterialDataflowNode::Evaluate(Dataflow::FContext& Context, const FDataflowOutput* Out) const
+{
+	if (Out->IsA<TObjectPtr<UMaterial>>(&Material))
+	{
+		SetValue<TObjectPtr<UMaterial>>(Context, InMaterial, &Material);
+	}
+}
+
+
+void FMakeMaterialsArrayDataflowNode::Evaluate(Dataflow::FContext& Context, const FDataflowOutput* Out) const
+{
+	if (Out->IsA<TArray<TObjectPtr<UMaterial>>>(&Materials))
+	{
+		SetValue<TArray<TObjectPtr<UMaterial>>>(Context, TArray<TObjectPtr<UMaterial>>(), &Materials);
+	}
+}
+
diff --git a/Engine/Plugins/Experimental/GeometryCollectionPlugin/Source/GeometryCollectionNodes/Private/Dataflow/GeometryCollectionNodes.cpp b/Engine/Plugins/Experimental/GeometryCollectionPlugin/Source/GeometryCollectionNodes/Private/Dataflow/GeometryCollectionNodes.cpp
--- a/Engine/Plugins/Experimental/GeometryCollectionPlugin/Source/GeometryCollectionNodes/Private/Dataflow/GeometryCollectionNodes.cpp
+++ b/Engine/Plugins/Experimental/GeometryCollectionPlugin/Source/GeometryCollectionNodes/Private/Dataflow/GeometryCollectionNodes.cpp
@@ -96,13 +96,9 @@ namespace Dataflow
 		DATAFLOW_NODE_REGISTER_CREATION_FACTORY(FGetCollectionAttributeDataTypedDataflowNode);
 		DATAFLOW_NODE_REGISTER_CREATION_FACTORY(FSetCollectionAttributeDataTypedDataflowNode);
 		DATAFLOW_NODE_REGISTER_CREATION_FACTORY(FBoolArrayToFaceSelectionDataflowNode);
-		DATAFLOW_NODE_REGISTER_CREATION_FACTORY(FAddMaterialToCollectionDataflowNode);
-		DATAFLOW_NODE_REGISTER_CREATION_FACTORY(FReAssignMaterialInCollectionDataflowNode);
-		DATAFLOW_NODE_REGISTER_CREATION_FACTORY(FMaterialsInfoDataflowNode);
-		DATAFLOW_NODE_REGISTER_CREATION_FACTORY(FGetMaterialFromMaterialsArrayDataflowNode);
-		DATAFLOW_NODE_REGISTER_CREATION_FACTORY(FSetMaterialInMaterialsArrayDataflowNode);
-		DATAFLOW_NODE_REGISTER_CREATION_FACTORY(FMakeMaterialDataflowNode);
-		DATAFLOW_NODE_REGISTER_CREATION_FACTORY(FMakeMaterialsArrayDataflowNode);
+		DATAFLOW_NODE_REGISTER_CREATION_FACTORY(FFloatArrayToVertexSelectionDataflowNode);
+		DATAFLOW_NODE_REGISTER_CREATION_FACTORY(FSetVertexColorInCollectionDataflowNode);
+		DATAFLOW_NODE_REGISTER_CREATION_FACTORY(FMakeTransformDataflowNode);
 
 		// GeometryCollection
 		DATAFLOW_NODE_REGISTER_CREATION_FACTORY_NODE_COLORS_BY_CATEGORY("GeometryCollection", FLinearColor(0.55f, 0.45f, 1.0f), CDefaultNodeBodyTintColor);
@@ -117,7 +113,7 @@ namespace Dataflow
 		// Math
 		DATAFLOW_NODE_REGISTER_CREATION_FACTORY_NODE_COLORS_BY_CATEGORY("Math", FLinearColor(0.f, 0.4f, 0.8f), CDefaultNodeBodyTintColor);
 		// Generators
-		DATAFLOW_NODE_REGISTER_CREATION_FACTORY_NODE_COLORS_BY_CATEGORY("Generators", FLinearColor(.7f, 0.7f, 0.7f), CDefaultNodeBodyTintColor);
+		DATAFLOW_NODE_REGISTER_CREATION_FACTORY_NODE_COLORS_BY_CATEGORY("Generators", FLinearColor(.6f, 0.1f, 1.f), CDefaultNodeBodyTintColor);
 	}
 }
 
@@ -590,7 +586,14 @@ void FGetNumArrayElementsDataflowNode::Evaluate(Dataflow::FContext& Context, con
 {
 	if (Out->IsA<int32>(&NumElements))
 	{
-		SetValue<int32>(Context, GetValue<TArray<FVector>>(Context, &Points).Num(), &NumElements);
+		if (IsConnected<TArray<FVector>>(&Points))
+		{
+			SetValue<int32>(Context, GetValue<TArray<FVector>>(Context, &Points).Num(), &NumElements);
+		}
+		else if (IsConnected<TArray<FVector3f>>(&Vector3fArray))
+		{
+			SetValue<int32>(Context, GetValue<TArray<FVector3f>>(Context, &Vector3fArray).Num(), &NumElements);
+		}
 	}
 }
 
@@ -1121,7 +1124,7 @@ void FGetCollectionAttributeDataTypedDataflowNode::Evaluate(Dataflow::FContext&
 	if (Out->IsA<TArray<bool>>(&BoolAttributeData) ||
 		Out->IsA<TArray<float>>(&FloatAttributeData) ||
 		Out->IsA<TArray<int32>>(&IntAttributeData) ||
-		Out->IsA<TArray<FVector>>(&VectorAttributeData))
+		Out->IsA<TArray<FVector3f>>(&VectorAttributeData))
 	{
 		const FManagedArrayCollection& InCollection = GetValue<FManagedArrayCollection>(Context, &Collection);
 
@@ -1182,17 +1185,17 @@ void FGetCollectionAttributeDataTypedDataflowNode::Evaluate(Dataflow::FContext&
 							SetValue<TArray<int32>>(Context, TArray<int32>(), &IntAttributeData);
 						}
 					}
-					else if (Out->IsA<TArray<FVector>>(&VectorAttributeData))
+					else if (Out->IsA<TArray<FVector3f>>(&VectorAttributeData))
 					{
-						if (TypeStr == FString("Vector3d"))
+						if (TypeStr == FString("Vector"))
 						{
-							const TManagedArray<FVector>& VectorAttributeArr = InCollection.GetAttribute<FVector>(FName(*AttrName), GroupNameToUse);
+							const TManagedArray<FVector3f>& VectorAttributeArr = InCollection.GetAttribute<FVector3f>(FName(*AttrName), GroupNameToUse);
 
-							SetValue<TArray<FVector>>(Context, VectorAttributeArr.GetConstArray(), &VectorAttributeData);
+							SetValue<TArray<FVector3f>>(Context, VectorAttributeArr.GetConstArray(), &VectorAttributeData);
 						}
 						else
 						{
-							SetValue<TArray<FVector>>(Context, TArray<FVector>(), &VectorAttributeData);
+							SetValue<TArray<FVector3f>>(Context, TArray<FVector3f>(), &VectorAttributeData);
 						}
 					}
 				}
@@ -1326,209 +1329,99 @@ void FBoolArrayToFaceSelectionDataflowNode::Evaluate(Dataflow::FContext& Context
 }
 
 
-void FAddMaterialToCollectionDataflowNode::Evaluate(Dataflow::FContext& Context, const FDataflowOutput* Out) const
+void FFloatArrayToVertexSelectionDataflowNode::Evaluate(Dataflow::FContext& Context, const FDataflowOutput* Out) const
 {
-	if (Out->IsA<FManagedArrayCollection>(&Collection) || Out->IsA<TArray<TObjectPtr<UMaterial>>>(&Materials))
+	if (Out->IsA<FDataflowVertexSelection>(&VertexSelection))
 	{
-		FManagedArrayCollection InCollection = GetValue<FManagedArrayCollection>(Context, &Collection);
-		TArray<TObjectPtr<UMaterial>> InMaterials = GetValue<TArray<TObjectPtr<UMaterial>>>(Context, &Materials);
-		const FDataflowFaceSelection& InFaceSelection = GetValue<FDataflowFaceSelection>(Context, &FaceSelection);
-		TObjectPtr<UMaterial> InOutsideMaterial = GetValue<TObjectPtr<UMaterial>>(Context, &OutsideMaterial);
-		TObjectPtr<UMaterial> InInsideMaterial = GetValue<TObjectPtr<UMaterial>>(Context, &InsideMaterial);
+		const TArray<float>& InFloatArray = GetValue<TArray<float>>(Context, &FloatArray);
 
-		int32 NewIndex;
+		FDataflowVertexSelection NewVertexSelection;
+		NewVertexSelection.Initialize(InFloatArray.Num(), false);
 
-		if (InCollection.HasAttribute("Internal", FGeometryCollection::FacesGroup))
+		for (int32 Idx = 0; Idx < InFloatArray.Num(); ++Idx)
 		{
-			if (bAssignOutsideMaterial || bAssignInsideMaterial)
+			if (Operation == ECompareOperationEnum::Dataflow_Compare_Equal)
 			{
-				TManagedArray<int32>& MaterialIDs = InCollection.ModifyAttribute<int32>("MaterialID", FGeometryCollection::FacesGroup);
-
-				const TManagedArray<bool>& Internals = InCollection.GetAttribute<bool>("Internal", FGeometryCollection::FacesGroup);
-
-				if (bAssignOutsideMaterial)
+				if (InFloatArray[Idx] == Threshold)
 				{
-					NewIndex = InMaterials.Add(InOutsideMaterial);
-
-					// Update MaterialIdx for selected outside faces
-					for (int32 FaceIdx = 0; FaceIdx < InFaceSelection.Num(); ++FaceIdx)
-					{
-						if (!Internals[FaceIdx] && InFaceSelection.IsSelected(FaceIdx))
-						{
-							MaterialIDs[FaceIdx] = NewIndex;
-						}
-					}
+					NewVertexSelection.SetSelected(Idx);
 				}
-
-				if (bAssignInsideMaterial)
+			}
+			else if (Operation == ECompareOperationEnum::Dataflow_Compare_Smaller)
+			{
+				if (InFloatArray[Idx] < Threshold)
 				{
-					NewIndex = InMaterials.Add(InInsideMaterial);
-
-					// Update MaterialIdx for selected inside faces
-					for (int32 FaceIdx = 0; FaceIdx < InFaceSelection.Num(); ++FaceIdx)
-					{
-						if (Internals[FaceIdx] && InFaceSelection.IsSelected(FaceIdx))
-						{
-							MaterialIDs[FaceIdx] = NewIndex;
-						}
-					}
+					NewVertexSelection.SetSelected(Idx);
+				}
+			}
+			else if (Operation == ECompareOperationEnum::Dataflow_Compare_SmallerOrEqual)
+			{
+				if (InFloatArray[Idx] <= Threshold)
+				{
+					NewVertexSelection.SetSelected(Idx);
+				}
+			}
+			else if (Operation == ECompareOperationEnum::Dataflow_Compare_Greater)
+			{
+				if (InFloatArray[Idx] > Threshold)
+				{
+					NewVertexSelection.SetSelected(Idx);
+				}
+			}
+			else if (Operation == ECompareOperationEnum::Dataflow_Compare_GreaterOrEqual)
+			{
+				if (InFloatArray[Idx] >= Threshold)
+				{
+					NewVertexSelection.SetSelected(Idx);
 				}
 			}
 		}
 
-		SetValue<TArray<TObjectPtr<UMaterial>>>(Context, MoveTemp(InMaterials), &Materials);
-		SetValue<FManagedArrayCollection>(Context, MoveTemp(InCollection), &Collection);
+		SetValue<FDataflowVertexSelection>(Context, NewVertexSelection, &VertexSelection);
 	}
 }
 
 
-void FReAssignMaterialInCollectionDataflowNode::Evaluate(Dataflow::FContext& Context, const FDataflowOutput* Out) const
+void FSetVertexColorInCollectionDataflowNode::Evaluate(Dataflow::FContext& Context, const FDataflowOutput* Out) const
 {
 	if (Out->IsA<FManagedArrayCollection>(&Collection))
 	{
 		FManagedArrayCollection InCollection = GetValue<FManagedArrayCollection>(Context, &Collection);
-		const FDataflowFaceSelection& InFaceSelection = GetValue<FDataflowFaceSelection>(Context, &FaceSelection);
-		const TArray<TObjectPtr<UMaterial>>& InMaterials = GetValue<TArray<TObjectPtr<UMaterial>>>(Context, &Materials);
-		int32 InOutsideMaterialIdx = GetValue<int32>(Context, &OutsideMaterialIdx);
-		int32 InInsideMaterialIdx = GetValue<int32>(Context, &InsideMaterialIdx);
+		const FDataflowVertexSelection& InVertexSelection = GetValue<FDataflowVertexSelection>(Context, &VertexSelection);
 
-		if (InMaterials.Num() > 0)
+		if (InCollection.NumElements(FGeometryCollection::VerticesGroup) == InVertexSelection.Num())
 		{
-			if (InCollection.HasAttribute("Internal", FGeometryCollection::FacesGroup))
+			const int32 NumVertices = InCollection.NumElements(FGeometryCollection::VerticesGroup);
+
+//			TManagedArray<FLinearColor>& VertexColors = InCollection.ModifyAttribute<FLinearColor>("Color", FGeometryCollection::VerticesGroup);
+			if (TManagedArray<FLinearColor>* VertexColors = InCollection.FindAttribute<FLinearColor>("Color", FGeometryCollection::VerticesGroup))
 			{
-				if (bAssignOutsideMaterial || bAssignInsideMaterial)
+				for (int32 Idx = 0; Idx < NumVertices; ++Idx)
 				{
-					TManagedArray<int32>& MaterialIDs = InCollection.ModifyAttribute<int32>("MaterialID", FGeometryCollection::FacesGroup);
-
-					const TManagedArray<bool>& Internals = InCollection.GetAttribute<bool>("Internal", FGeometryCollection::FacesGroup);
-
-					if (bAssignOutsideMaterial)
+					if (InVertexSelection.IsSelected(Idx))
 					{
-						if (InOutsideMaterialIdx >= 0 && InOutsideMaterialIdx < InMaterials.Num())
-						{
-							// Update MaterialIdx for selected outside faces
-							for (int32 FaceIdx = 0; FaceIdx < InFaceSelection.Num(); ++FaceIdx)
-							{
-								if (!Internals[FaceIdx] && InFaceSelection.IsSelected(FaceIdx))
-								{
-									MaterialIDs[FaceIdx] = InOutsideMaterialIdx;
-								}
-							}
-						}
+						(*VertexColors)[Idx] = SelectedColor;
 					}
-
-					if (bAssignInsideMaterial)
+					else
 					{
-						if (InInsideMaterialIdx >= 0 && InInsideMaterialIdx < InMaterials.Num())
-						{
-							// Update MaterialIdx for selected inside faces
-							for (int32 FaceIdx = 0; FaceIdx < InFaceSelection.Num(); ++FaceIdx)
-							{
-								if (Internals[FaceIdx] && InFaceSelection.IsSelected(FaceIdx))
-								{
-									MaterialIDs[FaceIdx] = InInsideMaterialIdx;
-								}
-							}
-						}
+						(*VertexColors)[Idx] = NonSelectedColor;
 					}
 				}
 			}
-		}
 
-		// move the collection to the output to avoid making another copy
-		SetValue<FManagedArrayCollection>(Context, MoveTemp(InCollection), &Collection);
-		SetValue<TArray<TObjectPtr<UMaterial>>>(Context, InMaterials, &Materials);
-	}
-}
-
-
-void FMaterialsInfoDataflowNode::Evaluate(Dataflow::FContext& Context, const FDataflowOutput* Out) const
-{
-	if (Out->IsA<FString>(&String))
-	{
-		const TArray<TObjectPtr<UMaterial>>& InMaterials = GetValue<TArray<TObjectPtr<UMaterial>>>(Context, &Materials);
-
-		FString OutputStr;
-
-		OutputStr.Appendf(TEXT("\n----------------------------------------\n"));
-		OutputStr.Appendf(TEXT("Number of Materials: %d\n"), InMaterials.Num());
-
-		int32 Idx = 0;
-		for (TObjectPtr<UMaterial> Material : InMaterials)
-		{
-			OutputStr.Appendf(TEXT("%4d: %s\n"), Idx, *(Material->GetFullName()));
-
-			Idx++;
+			SetValue<FManagedArrayCollection>(Context, MoveTemp(InCollection), &Collection);
 		}
-
-		OutputStr.Appendf(TEXT("----------------------------------------\n"));
-
-		SetValue<FString>(Context, OutputStr, &String);
 	}
 }
 
 
-void FGetMaterialFromMaterialsArrayDataflowNode::Evaluate(Dataflow::FContext& Context, const FDataflowOutput* Out) const
+void FMakeTransformDataflowNode::Evaluate(Dataflow::FContext& Context, const FDataflowOutput* Out) const
 {
-	if (Out->IsA<TObjectPtr<UMaterial>>(&Material))
+	if (Out->IsA<FTransform>(&Transform))
 	{
-		const TArray<TObjectPtr<UMaterial>>& InMaterials = GetValue<TArray<TObjectPtr<UMaterial>>>(Context, &Materials);
-		int32 InMaterialIdx = GetValue<int32>(Context, &MaterialIdx);
-
-		if (InMaterials.Num() > 0)
-		{
-			if (InMaterialIdx >= 0 && InMaterialIdx < InMaterials.Num())
-			{
-				SetValue<TObjectPtr<UMaterial>>(Context, InMaterials[InMaterialIdx], &Material);
-
-				return;
-			}
-		}
-
-		SetValue<TObjectPtr<UMaterial>>(Context, TObjectPtr<UMaterial>(), &Material);
+		SetValue<FTransform>(Context, FTransform(InTransform), &Transform);
 	}
 }
 
 
-void FSetMaterialInMaterialsArrayDataflowNode::Evaluate(Dataflow::FContext& Context, const FDataflowOutput* Out) const
-{
-	if (Out->IsA<TArray<TObjectPtr<UMaterial>>>(&Materials))
-	{
-		TArray<TObjectPtr<UMaterial>> InMaterials = GetValue<TArray<TObjectPtr<UMaterial>>>(Context, &Materials);
-		TObjectPtr<UMaterial> InMaterial = GetValue<TObjectPtr<UMaterial>>(Context, &Material);
-		int32 InMaterialIdx = GetValue<int32>(Context, &MaterialIdx);
-
-		if (Operation == ESetMaterialOperationTypeEnum::Dataflow_SetMaterialOperationType_Add)
-		{
-			InMaterials.Add(InMaterial);
-		}
-		else
-		{
-			if (InMaterialIdx >= 0 && InMaterialIdx < InMaterials.Num())
-			{
-				InMaterials.Insert(InMaterial, InMaterialIdx);
-			}
-		}
-
-		SetValue<TArray<TObjectPtr<UMaterial>>>(Context, MoveTemp(InMaterials), &Materials);
-	}
-}
-
-
-void FMakeMaterialDataflowNode::Evaluate(Dataflow::FContext& Context, const FDataflowOutput* Out) const
-{
-	if (Out->IsA<TObjectPtr<UMaterial>>(&Material))
-	{
-		SetValue<TObjectPtr<UMaterial>>(Context, InMaterial, &Material);
-	}
-}
-
-
-void FMakeMaterialsArrayDataflowNode::Evaluate(Dataflow::FContext& Context, const FDataflowOutput* Out) const
-{
-	if (Out->IsA<TArray<TObjectPtr<UMaterial>>>(&Materials))
-	{
-		SetValue<TArray<TObjectPtr<UMaterial>>>(Context, TArray<TObjectPtr<UMaterial>>(), &Materials);
-	}
-}
 
diff --git a/Engine/Plugins/Experimental/GeometryCollectionPlugin/Source/GeometryCollectionNodes/Private/Dataflow/GeometryCollectionNodesPlugin.cpp b/Engine/Plugins/Experimental/GeometryCollectionPlugin/Source/GeometryCollectionNodes/Private/Dataflow/GeometryCollectionNodesPlugin.cpp
--- a/Engine/Plugins/Experimental/GeometryCollectionPlugin/Source/GeometryCollectionNodes/Private/Dataflow/GeometryCollectionNodesPlugin.cpp
+++ b/Engine/Plugins/Experimental/GeometryCollectionPlugin/Source/GeometryCollectionNodes/Private/Dataflow/GeometryCollectionNodesPlugin.cpp
@@ -15,6 +15,8 @@
 #include "Dataflow/GeometryCollectionFracturingNodes.h"
 #include "Dataflow/GeometryCollectionEditNodes.h"
 #include "Dataflow/GeometryCollectionUtilityNodes.h"
+#include "Dataflow/GeometryCollectionMaterialNodes.h"
+#include "Dataflow/GeometryCollectionFieldNodes.h"
 
 #define LOCTEXT_NAMESPACE "DataflowNodes"
 
@@ -31,6 +33,8 @@ void IGeometryCollectionNodesPlugin::StartupModule()
 	Dataflow::GeometryCollectionFracturingNodes();
 	Dataflow::GeometryCollectionEditNodes();
 	Dataflow::GeometryCollectionUtilityNodes();
+	Dataflow::GeometryCollectionMaterialNodes();
+	Dataflow::GeometryCollectionFieldNodes();
 }
 
 void IGeometryCollectionNodesPlugin::ShutdownModule()
diff --git a/Engine/Plugins/Experimental/GeometryCollectionPlugin/Source/GeometryCollectionNodes/Public/Dataflow/GeometryCollectionFieldNodes.h b/Engine/Plugins/Experimental/GeometryCollectionPlugin/Source/GeometryCollectionNodes/Public/Dataflow/GeometryCollectionFieldNodes.h
new file mode 100644
--- /dev/null
+++ b/Engine/Plugins/Experimental/GeometryCollectionPlugin/Source/GeometryCollectionNodes/Public/Dataflow/GeometryCollectionFieldNodes.h
@@ -0,0 +1,666 @@
+// Copyright Epic Games, Inc. All Rights Reserved.
+
+#pragma once 
+
+#include "CoreMinimal.h"
+#include "Dataflow/DataflowEngine.h"
+#include "GeometryCollection/ManagedArrayCollection.h"
+#include "Dataflow/DataflowSelection.h"
+#include "Field/FieldSystemTypes.h"
+
+#include "GeometryCollectionFieldNodes.generated.h"
+
+
+/**
+ *
+ * 
+ *
+ */
+UENUM(BlueprintType)
+enum class EDataflowFieldFalloffType : uint8
+{
+	Dataflow_FieldFalloffType_None			UMETA(DisplayName = "None", ToolTip = "No falloff function is used"),
+	Dataflow_FieldFalloffType_Linear		UMETA(DisplayName = "Linear", ToolTip = "The falloff function will be proportional to x"),
+	Dataflow_FieldFalloffType_Inverse		UMETA(DisplayName = "Inverse", ToolTip = "The falloff function will be proportional to 1.0/x"),
+	Dataflow_FieldFalloffType_Squared		UMETA(DisplayName = "Squared", ToolTip = "The falloff function will be proportional to x*x"),
+	Dataflow_FieldFalloffType_Logarithmic	UMETA(DisplayName = "Logarithmic", ToolTip = "The falloff function will be proportional to log(x)"),
+	//~~~
+	//256th entry
+	Dataflow_Max                UMETA(Hidden)
+};
+
+/**
+ *
+ * RadialFalloff Field Dataflow node
+ *
+ */
+USTRUCT()
+struct FRadialFalloffFieldDataflowNode : public FDataflowNode
+{
+	GENERATED_USTRUCT_BODY()
+	DATAFLOW_NODE_DEFINE_INTERNAL(FRadialFalloffFieldDataflowNode, "RadialFalloffField", "Fields", "")
+
+public:
+	UPROPERTY(meta = (DataflowInput, DataflowIntrinsic))
+	TArray<FVector3f> VertexArray;
+
+	/**  */
+	UPROPERTY(meta = (DataflowInput, DisplayName = "VertexSelection"))
+	FDataflowVertexSelection VertexSelection;
+
+	/**  */
+	UPROPERTY(meta = (DataflowInput, DataflowIntrinsic))
+	FSphere Sphere = FSphere(ForceInit);
+
+	/**  */
+	UPROPERTY(EditAnywhere, Category = "Field", meta = (DataflowInput))
+	FVector Translation = FVector(0.f);
+
+	/**  */
+	UPROPERTY(EditAnywhere, Category = "Field")
+	float Magnitude = 1.f;
+
+	/**  */
+	UPROPERTY(EditAnywhere, Category = "Field")
+	float MinRange = 0.f;
+
+	/**  */
+	UPROPERTY(EditAnywhere, Category = "Field")
+	float MaxRange = 1.f;
+
+	/**  */
+	UPROPERTY(EditAnywhere, Category = "Field")
+	float Default = 0.f;
+
+	/**  */
+	UPROPERTY(EditAnywhere, Category = "Field")
+		EDataflowFieldFalloffType FalloffType = EDataflowFieldFalloffType::Dataflow_FieldFalloffType_Linear;
+
+	/**  */
+	UPROPERTY(meta = (DataflowOutput))
+	TArray<float> WeightArray;
+
+	FRadialFalloffFieldDataflowNode(const Dataflow::FNodeParameters& InParam, FGuid InGuid = FGuid::NewGuid())
+		: FDataflowNode(InParam, InGuid)
+	{
+		RegisterInputConnection(&VertexArray);
+		RegisterInputConnection(&VertexSelection);
+		RegisterInputConnection(&Sphere);
+		RegisterInputConnection(&Translation);
+		RegisterOutputConnection(&WeightArray);
+	}
+
+	virtual void Evaluate(Dataflow::FContext& Context, const FDataflowOutput* Out) const override;
+
+};
+
+
+/**
+ *
+ * PlaneFalloff Field Dataflow node
+ *
+ */
+USTRUCT()
+struct FPlaneFalloffFieldDataflowNode : public FDataflowNode
+{
+	GENERATED_USTRUCT_BODY()
+	DATAFLOW_NODE_DEFINE_INTERNAL(FPlaneFalloffFieldDataflowNode, "PlaneFalloffField", "Fields", "")
+
+public:
+	/**  */
+	UPROPERTY(meta = (DataflowInput, DataflowIntrinsic))
+	TArray<FVector3f> VertexArray;
+
+	/**  */
+	UPROPERTY(meta = (DataflowInput, DisplayName = "VertexSelection"))
+	FDataflowVertexSelection VertexSelection;
+
+	/**  */
+	UPROPERTY(EditAnywhere, Category = "Field", meta = (DataflowInput))
+	FVector Translation = FVector(0.f);
+
+	/**  */
+	UPROPERTY(EditAnywhere, Category = "Field")
+	float Magnitude = 1.f;
+
+	/**  */
+	UPROPERTY(EditAnywhere, Category = "Field")
+	float MinRange = 0.f;
+
+	/**  */
+	UPROPERTY(EditAnywhere, Category = "Field")
+	float MaxRange = 1.f;
+
+	/**  */
+	UPROPERTY(EditAnywhere, Category = "Field")
+	float Distance = 0.f;
+
+	/**  */
+	UPROPERTY(EditAnywhere, Category = "Field")
+	FVector Position = FVector(0.f);
+
+	/**  */
+	UPROPERTY(EditAnywhere, Category = "Field")
+	FVector Normal = FVector(0.f, 0.f, 1.f);
+
+	/**  */
+	UPROPERTY(EditAnywhere, Category = "Field")
+	float Default = 0.f;
+
+	/**  */
+	UPROPERTY(EditAnywhere, Category = "Field")
+	EDataflowFieldFalloffType FalloffType = EDataflowFieldFalloffType::Dataflow_FieldFalloffType_Linear;
+
+	/**  */
+	UPROPERTY(meta = (DataflowOutput))
+	TArray<float> WeightArray;
+
+	FPlaneFalloffFieldDataflowNode(const Dataflow::FNodeParameters& InParam, FGuid InGuid = FGuid::NewGuid())
+		: FDataflowNode(InParam, InGuid)
+	{
+		RegisterInputConnection(&VertexArray);
+		RegisterInputConnection(&VertexSelection);
+		RegisterInputConnection(&Position);
+		RegisterInputConnection(&Normal);
+		RegisterInputConnection(&Translation);
+		RegisterOutputConnection(&WeightArray);
+	}
+
+	virtual void Evaluate(Dataflow::FContext& Context, const FDataflowOutput* Out) const override;
+
+};
+
+
+/**
+ *
+ *
+ *
+ */
+UENUM(BlueprintType)
+enum class EDataflowSetMaskConditionType : uint8
+{
+	Dataflow_SetMaskConditionType_Always			UMETA(DisplayName = "Set Always", ToolTip = "The particle output value will be equal to Interior-value if the particle position is inside a sphere / Exterior-value otherwise."),
+	Dataflow_SetMaskConditionType_IFF_NOT_Interior	UMETA(DisplayName = "Merge Interior", ToolTip = "The particle output value will be equal to Interior-value if the particle position is inside the sphere or if the particle input value is already Interior-Value / Exterior-value otherwise."),
+	Dataflow_SetMaskConditionType_IFF_NOT_Exterior  UMETA(DisplayName = "Merge Exterior", ToolTip = "The particle output value will be equal to Exterior-value if the particle position is outside the sphere or if the particle input value is already Exterior-Value / Interior-value otherwise."),
+	//~~~
+	//256th entry
+	Dataflow_Max                UMETA(Hidden)
+};
+
+/**
+ *
+ * RadialIntMask Field Dataflow node
+ *
+ */
+USTRUCT()
+struct FRadialIntMaskFieldDataflowNode : public FDataflowNode
+{
+	GENERATED_USTRUCT_BODY()
+	DATAFLOW_NODE_DEFINE_INTERNAL(FRadialIntMaskFieldDataflowNode, "RadialIntMaskField", "Fields", "")
+
+public:
+	/**  */
+	UPROPERTY(meta = (DataflowInput, DataflowIntrinsic))
+	TArray<FVector3f> VertexArray;
+
+	/**  */
+	UPROPERTY(meta = (DataflowInput, DisplayName = "VertexSelection"))
+	FDataflowVertexSelection VertexSelection;
+
+	/**  */
+	UPROPERTY(meta = (DataflowInput, DataflowIntrinsic))
+	FSphere Sphere = FSphere(ForceInit);
+
+	/**  */
+	UPROPERTY(EditAnywhere, Category = "Field", meta = (DataflowInput))
+	FVector Translation = FVector(0.f);
+
+	/**  */
+	UPROPERTY(EditAnywhere, Category = "Field")
+	int32 InteriorValue = 1;
+
+	/**  */
+	UPROPERTY(EditAnywhere, Category = "Field")
+	int32 ExteriorValue = 0;
+
+	/**  */
+	UPROPERTY(EditAnywhere, Category = "Field")
+	EDataflowSetMaskConditionType SetMaskCondition = EDataflowSetMaskConditionType::Dataflow_SetMaskConditionType_Always;
+
+	/**  */
+	UPROPERTY(meta = (DataflowOutput))
+	TArray<int32> MaskArray;
+
+	FRadialIntMaskFieldDataflowNode(const Dataflow::FNodeParameters& InParam, FGuid InGuid = FGuid::NewGuid())
+		: FDataflowNode(InParam, InGuid)
+	{
+		RegisterInputConnection(&VertexArray);
+		RegisterInputConnection(&VertexSelection);
+		RegisterInputConnection(&Sphere);
+		RegisterInputConnection(&Translation);
+		RegisterOutputConnection(&MaskArray);
+	}
+
+	virtual void Evaluate(Dataflow::FContext& Context, const FDataflowOutput* Out) const override;
+
+};
+
+
+/**
+ *
+ * UniformScalar Field Dataflow node
+ *
+ */
+USTRUCT()
+struct FUniformScalarFieldDataflowNode : public FDataflowNode
+{
+	GENERATED_USTRUCT_BODY()
+	DATAFLOW_NODE_DEFINE_INTERNAL(FUniformScalarFieldDataflowNode, "UniformScalarField", "Fields", "")
+
+public:
+	/**  */
+	UPROPERTY(meta = (DataflowInput, DataflowIntrinsic))
+	TArray<FVector3f> VertexArray;
+
+	/**  */
+	UPROPERTY(meta = (DataflowInput, DisplayName = "VertexSelection"))
+	FDataflowVertexSelection VertexSelection;
+
+	/**  */
+	UPROPERTY(EditAnywhere, Category = "Field")
+	float Magnitude = 1.f;
+
+	/**  */
+	UPROPERTY(meta = (DataflowOutput))
+	TArray<float> WeightArray;
+
+	FUniformScalarFieldDataflowNode(const Dataflow::FNodeParameters& InParam, FGuid InGuid = FGuid::NewGuid())
+		: FDataflowNode(InParam, InGuid)
+	{
+		RegisterInputConnection(&VertexArray);
+		RegisterInputConnection(&VertexSelection);
+		RegisterOutputConnection(&WeightArray);
+	}
+
+	virtual void Evaluate(Dataflow::FContext& Context, const FDataflowOutput* Out) const override;
+
+};
+
+
+/**
+ *
+ * UniformVector Field Dataflow node
+ *
+ */
+USTRUCT()
+struct FUniformVectorFieldDataflowNode : public FDataflowNode
+{
+	GENERATED_USTRUCT_BODY()
+	DATAFLOW_NODE_DEFINE_INTERNAL(FUniformVectorFieldDataflowNode, "UniformVectorField", "Fields", "")
+
+public:
+	/**  */
+	UPROPERTY(meta = (DataflowInput, DataflowIntrinsic))
+	TArray<FVector3f> VertexArray;
+
+	/**  */
+	UPROPERTY(meta = (DataflowInput, DisplayName = "VertexSelection"))
+	FDataflowVertexSelection VertexSelection;
+
+	/**  */
+	UPROPERTY(EditAnywhere, Category = "Field")
+	float Magnitude = 1.f;
+
+	/**  */
+	UPROPERTY(EditAnywhere, Category = "Field")
+	FVector Direction = FVector(0.f);
+
+	/**  */
+	UPROPERTY(meta = (DataflowOutput))
+	TArray<FVector> VectorArray;
+
+	FUniformVectorFieldDataflowNode(const Dataflow::FNodeParameters& InParam, FGuid InGuid = FGuid::NewGuid())
+		: FDataflowNode(InParam, InGuid)
+	{
+		RegisterInputConnection(&VertexArray);
+		RegisterInputConnection(&VertexSelection);
+		RegisterOutputConnection(&VectorArray);
+	}
+
+	virtual void Evaluate(Dataflow::FContext& Context, const FDataflowOutput* Out) const override;
+
+};
+
+
+/**
+ *
+ * RadialVector Field Dataflow node
+ *
+ */
+USTRUCT()
+struct FRadialVectorFieldDataflowNode : public FDataflowNode
+{
+	GENERATED_USTRUCT_BODY()
+	DATAFLOW_NODE_DEFINE_INTERNAL(FRadialVectorFieldDataflowNode, "RadialVectorField", "Fields", "")
+
+public:
+	/**  */
+	UPROPERTY(meta = (DataflowInput, DataflowIntrinsic))
+	TArray<FVector3f> VertexArray;
+
+	/**  */
+	UPROPERTY(meta = (DataflowInput, DisplayName = "VertexSelection"))
+	FDataflowVertexSelection VertexSelection;
+
+	/**  */
+	UPROPERTY(EditAnywhere, Category = "Field")
+	float Magnitude = 1.f;
+
+	/**  */
+	UPROPERTY(EditAnywhere, Category = "Field")
+	FVector Position = FVector(0.f);
+
+	/**  */
+	UPROPERTY(meta = (DataflowOutput))
+	TArray<FVector> VectorArray;
+
+	FRadialVectorFieldDataflowNode(const Dataflow::FNodeParameters& InParam, FGuid InGuid = FGuid::NewGuid())
+		: FDataflowNode(InParam, InGuid)
+	{
+		RegisterInputConnection(&VertexArray);
+		RegisterInputConnection(&VertexSelection);
+		RegisterOutputConnection(&VectorArray);
+	}
+
+	virtual void Evaluate(Dataflow::FContext& Context, const FDataflowOutput* Out) const override;
+
+};
+
+
+/**
+ *
+ * RandomVector Field Dataflow node
+ *
+ */
+USTRUCT()
+struct FRandomVectorFieldDataflowNode : public FDataflowNode
+{
+	GENERATED_USTRUCT_BODY()
+	DATAFLOW_NODE_DEFINE_INTERNAL(FRandomVectorFieldDataflowNode, "RandomVectorField", "Fields", "")
+
+public:
+	/**  */
+	UPROPERTY(meta = (DataflowInput, DataflowIntrinsic))
+	TArray<FVector3f> VertexArray;
+
+	/**  */
+	UPROPERTY(meta = (DataflowInput, DisplayName = "VertexSelection"))
+	FDataflowVertexSelection VertexSelection;
+
+	/**  */
+	UPROPERTY(EditAnywhere, Category = "Field")
+	float Magnitude = 1.f;
+
+	/**  */
+	UPROPERTY(meta = (DataflowOutput))
+	TArray<FVector> VectorArray;
+
+	FRandomVectorFieldDataflowNode(const Dataflow::FNodeParameters& InParam, FGuid InGuid = FGuid::NewGuid())
+		: FDataflowNode(InParam, InGuid)
+	{
+		RegisterInputConnection(&VertexArray);
+		RegisterInputConnection(&VertexSelection);
+		RegisterOutputConnection(&VectorArray);
+	}
+
+	virtual void Evaluate(Dataflow::FContext& Context, const FDataflowOutput* Out) const override;
+
+};
+
+
+/**
+ *
+ * Noise Field Dataflow node
+ *
+ */
+USTRUCT()
+struct FNoiseFieldDataflowNode : public FDataflowNode
+{
+	GENERATED_USTRUCT_BODY()
+	DATAFLOW_NODE_DEFINE_INTERNAL(FNoiseFieldDataflowNode, "NoiseField", "Fields", "")
+
+public:
+	/**  */
+	UPROPERTY(meta = (DataflowInput, DataflowIntrinsic))
+	TArray<FVector3f> VertexArray;
+
+	/**  */
+	UPROPERTY(meta = (DataflowInput, DisplayName = "VertexSelection"))
+	FDataflowVertexSelection VertexSelection;
+
+	/**  */
+	UPROPERTY(EditAnywhere, Category = "Field")
+	float MinRange = 0.f;
+
+	/**  */
+	UPROPERTY(EditAnywhere, Category = "Field")
+	float MaxRange = 1.f;
+
+	/**  */
+	UPROPERTY(EditAnywhere, Category = "Field", meta = (DataflowInput))
+	FTransform Transform = FTransform::Identity;
+
+	/**  */
+	UPROPERTY(meta = (DataflowOutput))
+	TArray<float> WeightArray;
+
+	FNoiseFieldDataflowNode(const Dataflow::FNodeParameters& InParam, FGuid InGuid = FGuid::NewGuid())
+		: FDataflowNode(InParam, InGuid)
+	{
+		RegisterInputConnection(&VertexArray);
+		RegisterInputConnection(&VertexSelection);
+		RegisterInputConnection(&Transform);
+		RegisterOutputConnection(&WeightArray);
+	}
+
+	virtual void Evaluate(Dataflow::FContext& Context, const FDataflowOutput* Out) const override;
+
+};
+
+
+/**
+ *
+ * UniformInteger Field Dataflow node
+ *
+ */
+USTRUCT()
+struct FUniformIntegerFieldDataflowNode : public FDataflowNode
+{
+	GENERATED_USTRUCT_BODY()
+	DATAFLOW_NODE_DEFINE_INTERNAL(FUniformIntegerFieldDataflowNode, "UniformIntegerField", "Fields", "")
+
+public:
+	/**  */
+	UPROPERTY(meta = (DataflowInput, DataflowIntrinsic))
+	TArray<FVector3f> VertexArray;
+
+	/**  */
+	UPROPERTY(meta = (DataflowInput, DisplayName = "VertexSelection"))
+	FDataflowVertexSelection VertexSelection;
+
+	/**  */
+	UPROPERTY(EditAnywhere, Category = "Field")
+	int32 Magnitude = 0;
+
+	/**  */
+	UPROPERTY(meta = (DataflowOutput))
+	TArray<int32> MaskArray;
+
+	FUniformIntegerFieldDataflowNode(const Dataflow::FNodeParameters& InParam, FGuid InGuid = FGuid::NewGuid())
+		: FDataflowNode(InParam, InGuid)
+	{
+		RegisterInputConnection(&VertexArray);
+		RegisterInputConnection(&VertexSelection);
+		RegisterOutputConnection(&MaskArray);
+	}
+
+	virtual void Evaluate(Dataflow::FContext& Context, const FDataflowOutput* Out) const override;
+
+};
+
+
+/**
+ *
+ *
+ *
+ */
+UENUM(BlueprintType)
+enum class EDataflowWaveFunctionType : uint8
+{
+	Dataflow_WaveFunctionType_Cosine	UMETA(DisplayName = "Cosine", ToolTip = "Cosine wave that will move in time."),
+	Dataflow_WaveFunctionType_Gaussian  UMETA(DisplayName = "Gaussian", ToolTip = "Gaussian wave that will move in time."),
+	Dataflow_WaveFunctionType_Falloff	UMETA(DisplayName = "Falloff", ToolTip = "The radial falloff radius will move along temporal wave."),
+	Dataflow_WaveFunctionType_Decay		UMETA(DisplayName = "Decay", ToolTip = "The magnitude of the field will decay in time."),
+	//~~~
+	//256th entry
+	Dataflow_Max                UMETA(Hidden)
+};
+
+/**
+ *
+ * WaveScalar Field Dataflow node v2
+ *
+ */
+USTRUCT()
+struct FWaveScalarFieldDataflowNode : public FDataflowNode
+{
+	GENERATED_USTRUCT_BODY()
+	DATAFLOW_NODE_DEFINE_INTERNAL(FWaveScalarFieldDataflowNode, "WaveScalarField", "Fields", "")
+
+public:
+	/**  */
+	UPROPERTY(meta = (DataflowInput, DataflowIntrinsic))
+	TArray<FVector3f> VertexArray;
+
+	/**  */
+	UPROPERTY(meta = (DataflowInput, DisplayName = "VertexSelection"))
+	FDataflowVertexSelection VertexSelection;
+
+	/**  */
+	UPROPERTY(EditAnywhere, Category = "Field")
+	float Magnitude = 1.f;
+
+	/**  */
+	UPROPERTY(meta = (DataflowInput, DataflowIntrinsic))
+	FVector Position = FVector(0.f);
+
+	/**  */
+	UPROPERTY(EditAnywhere, Category = "Field", meta = (DataflowInput))
+	FVector Translation = FVector(0.f);
+
+	/**  */
+	UPROPERTY(EditAnywhere, Category = "Field")
+	float Wavelength = 1000.f;
+
+	/**  */
+	UPROPERTY(EditAnywhere, Category = "Field")
+	float Period = 1.f;
+
+	/**  */
+	UPROPERTY(EditAnywhere, Category = "Field")
+	EDataflowWaveFunctionType FunctionType = EDataflowWaveFunctionType::Dataflow_WaveFunctionType_Cosine;
+
+	/**  */
+	UPROPERTY(EditAnywhere, Category = "Field")
+	EDataflowFieldFalloffType FalloffType = EDataflowFieldFalloffType::Dataflow_FieldFalloffType_Linear;
+
+	/**  */
+	UPROPERTY(meta = (DataflowOutput))
+	TArray<float> WeightArray;
+
+	FWaveScalarFieldDataflowNode(const Dataflow::FNodeParameters& InParam, FGuid InGuid = FGuid::NewGuid())
+		: FDataflowNode(InParam, InGuid)
+	{
+		RegisterInputConnection(&VertexArray);
+		RegisterInputConnection(&VertexSelection);
+		RegisterInputConnection(&Translation);
+		RegisterOutputConnection(&WeightArray);
+	}
+
+	virtual void Evaluate(Dataflow::FContext& Context, const FDataflowOutput* Out) const override;
+
+};
+
+
+/**
+ *
+ * BoxFalloff Field Dataflow node
+ *
+ */
+USTRUCT()
+struct FBoxFalloffFieldDataflowNode : public FDataflowNode
+{
+	GENERATED_USTRUCT_BODY()
+	DATAFLOW_NODE_DEFINE_INTERNAL(FBoxFalloffFieldDataflowNode, "BoxFalloffField", "Fields", "")
+
+public:
+	UPROPERTY(meta = (DataflowInput, DataflowIntrinsic))
+	TArray<FVector3f> VertexArray;
+
+	/**  */
+	UPROPERTY(meta = (DataflowInput, DisplayName = "VertexSelection"))
+	FDataflowVertexSelection VertexSelection;
+
+	/**  */
+	UPROPERTY(meta = (DataflowInput, DataflowIntrinsic))
+	FBox Box = FBox(ForceInit);
+
+	/**  */
+	UPROPERTY(EditAnywhere, Category = "Field", meta = (DataflowInput))
+	FTransform Transform = FTransform::Identity;
+
+	/**  */
+	UPROPERTY(EditAnywhere, Category = "Field")
+	float Magnitude = 1.f;
+
+	/**  */
+	UPROPERTY(EditAnywhere, Category = "Field")
+	float MinRange = 0.f;
+
+	/**  */
+	UPROPERTY(EditAnywhere, Category = "Field")
+	float MaxRange = 1.f;
+
+	/**  */
+	UPROPERTY(EditAnywhere, Category = "Field")
+	float Default = 0.f;
+
+	/**  */
+	UPROPERTY(EditAnywhere, Category = "Field")
+	EDataflowFieldFalloffType FalloffType = EDataflowFieldFalloffType::Dataflow_FieldFalloffType_Linear;
+
+	/**  */
+	UPROPERTY(meta = (DataflowOutput))
+	TArray<float> WeightArray;
+
+	FBoxFalloffFieldDataflowNode(const Dataflow::FNodeParameters& InParam, FGuid InGuid = FGuid::NewGuid())
+		: FDataflowNode(InParam, InGuid)
+	{
+		RegisterInputConnection(&VertexArray);
+		RegisterInputConnection(&VertexSelection);
+		RegisterInputConnection(&Box);
+		RegisterInputConnection(&Transform);
+		RegisterOutputConnection(&WeightArray);
+	}
+
+	virtual void Evaluate(Dataflow::FContext& Context, const FDataflowOutput* Out) const override;
+
+};
+
+
+
+
+namespace Dataflow
+{
+	void GeometryCollectionFieldNodes();
+}
diff --git a/Engine/Plugins/Experimental/GeometryCollectionPlugin/Source/GeometryCollectionNodes/Public/Dataflow/GeometryCollectionMaterialNodes.h b/Engine/Plugins/Experimental/GeometryCollectionPlugin/Source/GeometryCollectionNodes/Public/Dataflow/GeometryCollectionMaterialNodes.h
new file mode 100644
--- /dev/null
+++ b/Engine/Plugins/Experimental/GeometryCollectionPlugin/Source/GeometryCollectionNodes/Public/Dataflow/GeometryCollectionMaterialNodes.h
@@ -0,0 +1,318 @@
+// Copyright Epic Games, Inc. All Rights Reserved.
+
+#pragma once 
+
+#include "CoreMinimal.h"
+#include "Dataflow/DataflowEngine.h"
+#include "GeometryCollection/ManagedArrayCollection.h"
+//#include "DynamicMesh/DynamicMesh3.h"
+//#include "UDynamicMesh.h"
+#include "Dataflow/DataflowSelection.h"
+//#include "Math/MathFwd.h"
+//#include "DynamicMesh/DynamicMesh3.h"
+
+#include "GeometryCollectionMaterialNodes.generated.h"
+
+
+class FGeometryCollection;
+class UGeometryCollection;
+class UStaticMesh;
+class UMaterial;
+
+/**
+ *
+ * Adds Outside/Inside Materials to Outside/Inside faces
+ *
+ */
+USTRUCT(meta = (DataflowGeometryCollection))
+struct FAddMaterialToCollectionDataflowNode : public FDataflowNode
+{
+	GENERATED_USTRUCT_BODY()
+	DATAFLOW_NODE_DEFINE_INTERNAL(FAddMaterialToCollectionDataflowNode, "AddMaterialToCollection", "GeometryCollection|Materials", "")
+
+public:
+	/** Collection to add material(s) to */
+	UPROPERTY(meta = (DataflowInput, DataflowOutput, DataflowPassthrough = "Collection", DataflowIntrinsic))
+	FManagedArrayCollection Collection;
+
+	/** Face selection, the material(s) will be added to the selected faces */
+	UPROPERTY(meta = (DataflowInput, DisplayName = "FaceSelection", DataflowIntrinsic))
+	FDataflowFaceSelection FaceSelection;
+
+	/** Materials array storing the materials */
+	UPROPERTY(meta = (DataflowInput, DataflowOutput, DataflowPassthrough = "Materials", DataflowIntrinsic))
+	TArray<TObjectPtr<UMaterial>> Materials;
+
+	/** Outside material to assign to the outside faces from the face selection */
+	UPROPERTY(EditAnywhere, Category = "Material", meta = (DataflowInput))
+	TObjectPtr<UMaterial> OutsideMaterial;
+
+	/** Inside material to assign to the inside faces from the face selection */
+	UPROPERTY(EditAnywhere, Category = "Material", meta = (DataflowInput))
+	TObjectPtr<UMaterial> InsideMaterial;
+
+	/** If true, the outside material will be assigned to the outside faces from the face selection */
+	UPROPERTY(EditAnywhere, Category = "Material", meta = (DisplayName = "Assign Outside Material to Outside Faces"))
+	bool bAssignOutsideMaterial = true;
+
+	/** If true, the inside material will be assigned to the inside faces from the face selection */
+	UPROPERTY(EditAnywhere, Category = "Material", meta = (DisplayName = "Assign Inside Material to Inside Faces"))
+	bool bAssignInsideMaterial = false;
+
+	FAddMaterialToCollectionDataflowNode(const Dataflow::FNodeParameters& InParam, FGuid InGuid = FGuid::NewGuid())
+		: FDataflowNode(InParam, InGuid)
+	{
+		RegisterInputConnection(&Collection);
+		RegisterInputConnection(&Materials);
+		RegisterInputConnection(&FaceSelection);
+		RegisterInputConnection(&OutsideMaterial);
+		RegisterInputConnection(&InsideMaterial);
+		RegisterOutputConnection(&Collection, &Collection);
+		RegisterOutputConnection(&Materials, &Materials);
+	}
+
+	virtual void Evaluate(Dataflow::FContext& Context, const FDataflowOutput* Out) const override;
+
+};
+
+
+/**
+ *
+ * Reassign existing material(s) to Outside/Inside faces
+ *
+ */
+USTRUCT(meta = (DataflowGeometryCollection))
+struct FReAssignMaterialInCollectionDataflowNode : public FDataflowNode
+{
+	GENERATED_USTRUCT_BODY()
+		DATAFLOW_NODE_DEFINE_INTERNAL(FReAssignMaterialInCollectionDataflowNode, "ReAssignMaterialInCollection", "GeometryCollection|Materials", "")
+
+public:
+	/** Collection for reassign the material(s) */
+	UPROPERTY(meta = (DataflowInput, DataflowOutput, DataflowPassthrough = "Collection", DataflowIntrinsic))
+	FManagedArrayCollection Collection;
+
+	/** Face selection, the material(s) will be assigned to the selected faces */
+	UPROPERTY(meta = (DataflowInput, DisplayName = "FaceSelection", DataflowIntrinsic))
+	FDataflowFaceSelection FaceSelection;
+
+	/** Materials array storing the materials */
+	UPROPERTY(EditAnywhere, Category = "Material", meta = (DataflowInput, DataflowOutput, DataflowPassthrough = "Materials", DataflowIntrinsic))
+	TArray<TObjectPtr<UMaterial>> Materials;
+
+	/** Index of the material in the Materials array to assign to the outside faces from the face selection */
+	UPROPERTY(EditAnywhere, Category = "Material", meta = (DataflowInput, DisplayName = "Outside Material Index", ArrayClamp = "Materials"))
+	int32 OutsideMaterialIdx = -1;
+
+	/** Index of the material in the Materials array to assign to the inside faces from the face selection */
+	UPROPERTY(EditAnywhere, Category = "Material", meta = (DataflowInput, DisplayName = "Inside Material Index", ArrayClamp = "Materials"))
+	int32 InsideMaterialIdx = -1;
+
+	/** If true, the selected material from the Materials array will be assigned to the outside faces from the face selection */
+	UPROPERTY(EditAnywhere, Category = "Material", meta = (DisplayName = "Assign Outside Material to Outside Faces"))
+	bool bAssignOutsideMaterial = false;
+
+	/** If true, the selected material from the Materials array will be assigned to the inside faces from the face selection */
+	UPROPERTY(EditAnywhere, Category = "Material", meta = (DisplayName = "Assign Inside Material to Inside Faces"))
+	bool bAssignInsideMaterial = false;
+
+	FReAssignMaterialInCollectionDataflowNode(const Dataflow::FNodeParameters& InParam, FGuid InGuid = FGuid::NewGuid())
+		: FDataflowNode(InParam, InGuid)
+	{
+		RegisterInputConnection(&Collection);
+		RegisterInputConnection(&Materials);
+		RegisterInputConnection(&FaceSelection);
+		RegisterInputConnection(&OutsideMaterialIdx);
+		RegisterInputConnection(&InsideMaterialIdx);
+		RegisterOutputConnection(&Collection, &Collection);
+		RegisterOutputConnection(&Materials, &Materials);
+	}
+
+	virtual void Evaluate(Dataflow::FContext& Context, const FDataflowOutput* Out) const override;
+
+};
+
+
+/**
+ *
+ * Generates a formatted string of materials from the Materials array
+ *
+ */
+USTRUCT(meta = (DataflowGeometryCollection))
+struct FMaterialsInfoDataflowNode : public FDataflowNode
+{
+	GENERATED_USTRUCT_BODY()
+		DATAFLOW_NODE_DEFINE_INTERNAL(FMaterialsInfoDataflowNode, "MaterialsInfo", "GeometryCollection|Materials", "")
+
+public:
+	/** Materials array storing the materials */
+	UPROPERTY(meta = (DataflowInput, DataflowIntrinsic))
+	TArray<TObjectPtr<UMaterial>> Materials;
+
+	/** Formatted string of the materials */
+	UPROPERTY(meta = (DataflowOutput))
+	FString String;
+
+	FMaterialsInfoDataflowNode(const Dataflow::FNodeParameters& InParam, FGuid InGuid = FGuid::NewGuid())
+		: FDataflowNode(InParam, InGuid)
+	{
+		RegisterInputConnection(&Materials);
+		RegisterOutputConnection(&String);
+	}
+
+	virtual void Evaluate(Dataflow::FContext& Context, const FDataflowOutput* Out) const override;
+
+};
+
+
+/**
+ *
+ * Get a Material from a Materials array
+ *
+ */
+USTRUCT(meta = (DataflowGeometryCollection))
+struct FGetMaterialFromMaterialsArrayDataflowNode : public FDataflowNode
+{
+	GENERATED_USTRUCT_BODY()
+		DATAFLOW_NODE_DEFINE_INTERNAL(FGetMaterialFromMaterialsArrayDataflowNode, "GetMaterialFromMaterialsArray", "Utilities|Materials", "")
+
+public:
+	/** Materials array storing the materials */
+	UPROPERTY(meta = (DataflowInput, DataflowIntrinsic))
+	TArray<TObjectPtr<UMaterial>> Materials;
+
+	/** Selected material from the Materials array */
+	UPROPERTY(meta = (DataflowOutput))
+	TObjectPtr<UMaterial> Material;
+
+	/** Index in the Materials array for the selected material */
+	UPROPERTY(EditAnywhere, Category = "Material", meta = (DataflowInput, DisplayName = "Material Index"))
+	int32 MaterialIdx = 0;
+
+	FGetMaterialFromMaterialsArrayDataflowNode(const Dataflow::FNodeParameters& InParam, FGuid InGuid = FGuid::NewGuid())
+		: FDataflowNode(InParam, InGuid)
+	{
+		RegisterInputConnection(&Materials);
+		RegisterInputConnection(&MaterialIdx);
+		RegisterOutputConnection(&Material);
+	}
+
+	virtual void Evaluate(Dataflow::FContext& Context, const FDataflowOutput* Out) const override;
+
+};
+
+
+UENUM(BlueprintType)
+enum class ESetMaterialOperationTypeEnum : uint8
+{
+	Dataflow_SetMaterialOperationType_Add UMETA(DisplayName = "Add"),
+	Dataflow_SetMaterialOperationType_Insert UMETA(DisplayName = "Insert"),
+	//~~~
+	//256th entry
+	Dataflow_Max                UMETA(Hidden)
+};
+
+/**
+ *
+ * Set a Material in a Materials array
+ *
+ */
+USTRUCT(meta = (DataflowGeometryCollection))
+struct FSetMaterialInMaterialsArrayDataflowNode : public FDataflowNode
+{
+	GENERATED_USTRUCT_BODY()
+		DATAFLOW_NODE_DEFINE_INTERNAL(FSetMaterialInMaterialsArrayDataflowNode, "SetMaterialInMaterialsArray", "Utilities|Materials", "")
+
+public:
+	/** Materials array storing the materials */
+	UPROPERTY(meta = (DataflowInput, DataflowOutput, DataflowPassthrough = "Materials", DataflowIntrinsic))
+	TArray<TObjectPtr<UMaterial>> Materials;
+
+	/** Material to add/insert to/in Materials array */
+	UPROPERTY(EditAnywhere, Category = "Material", meta = (DataflowInput))
+	TObjectPtr<UMaterial> Material;
+
+	/** Operation type for setting the material, add will add the new material to the end off Materials array, insert will insert the
+	new material into Materials array at the index specified by MaterialIdx	*/
+	UPROPERTY(EditAnywhere, Category = "Material")
+	ESetMaterialOperationTypeEnum Operation = ESetMaterialOperationTypeEnum::Dataflow_SetMaterialOperationType_Add;
+
+	/** Index for inserting a nem material into the Materials array */
+	UPROPERTY(EditAnywhere, Category = "Material", meta = (DataflowInput, DisplayName = "Material Index", EditCondition = "Operation == ESetMaterialOperationTypeEnum::Dataflow_SetMaterialOperationType_Insert", EditConditionHides))
+	int32 MaterialIdx = 0;
+
+	FSetMaterialInMaterialsArrayDataflowNode(const Dataflow::FNodeParameters& InParam, FGuid InGuid = FGuid::NewGuid())
+		: FDataflowNode(InParam, InGuid)
+	{
+		RegisterInputConnection(&Materials);
+		RegisterInputConnection(&Material);
+		RegisterInputConnection(&MaterialIdx);
+		RegisterOutputConnection(&Materials, &Materials);
+	}
+
+	virtual void Evaluate(Dataflow::FContext& Context, const FDataflowOutput* Out) const override;
+
+};
+
+
+/**
+ *
+ * Makes a material
+ *
+ */
+USTRUCT()
+struct FMakeMaterialDataflowNode : public FDataflowNode
+{
+	GENERATED_USTRUCT_BODY()
+		DATAFLOW_NODE_DEFINE_INTERNAL(FMakeMaterialDataflowNode, "MakeMaterial", "Generators|Material", "")
+
+public:
+	/** Material which will be outputed */
+	UPROPERTY(EditAnywhere, Category = "Material", meta = (DisplayName = "Material"))
+	TObjectPtr<UMaterial> InMaterial;
+
+	/** Output material */
+	UPROPERTY(meta = (DataflowOutput));
+	TObjectPtr<UMaterial> Material;
+
+	FMakeMaterialDataflowNode(const Dataflow::FNodeParameters& InParam, FGuid InGuid = FGuid::NewGuid())
+		: FDataflowNode(InParam, InGuid)
+	{
+		RegisterOutputConnection(&Material);
+	}
+
+	virtual void Evaluate(Dataflow::FContext& Context, const FDataflowOutput* Out) const override;
+
+};
+
+/**
+ *
+ * Makes an empty Materials array
+ *
+ */
+USTRUCT()
+struct FMakeMaterialsArrayDataflowNode : public FDataflowNode
+{
+	GENERATED_USTRUCT_BODY()
+	DATAFLOW_NODE_DEFINE_INTERNAL(FMakeMaterialsArrayDataflowNode, "MakeMaterialsArray", "Generators|Material", "")
+
+public:
+	/** Output Matarials array */
+	UPROPERTY(meta = (DataflowOutput));
+	TArray<TObjectPtr<UMaterial>> Materials;
+
+	FMakeMaterialsArrayDataflowNode(const Dataflow::FNodeParameters& InParam, FGuid InGuid = FGuid::NewGuid())
+		: FDataflowNode(InParam, InGuid)
+	{
+		RegisterOutputConnection(&Materials);
+	}
+
+	virtual void Evaluate(Dataflow::FContext& Context, const FDataflowOutput* Out) const override;
+
+};
+
+
+namespace Dataflow
+{
+	void GeometryCollectionMaterialNodes();
+}
diff --git a/Engine/Plugins/Experimental/GeometryCollectionPlugin/Source/GeometryCollectionNodes/Public/Dataflow/GeometryCollectionNodes.h b/Engine/Plugins/Experimental/GeometryCollectionPlugin/Source/GeometryCollectionNodes/Public/Dataflow/GeometryCollectionNodes.h
--- a/Engine/Plugins/Experimental/GeometryCollectionPlugin/Source/GeometryCollectionNodes/Public/Dataflow/GeometryCollectionNodes.h
+++ b/Engine/Plugins/Experimental/GeometryCollectionPlugin/Source/GeometryCollectionNodes/Public/Dataflow/GeometryCollectionNodes.h
@@ -1141,10 +1141,14 @@ struct FGetNumArrayElementsDataflowNode : public FDataflowNode
 	DATAFLOW_NODE_DEFINE_INTERNAL(FGetNumArrayElementsDataflowNode, "GetNumArrayElements", "Utilities|Array", "")
 
 public:
-	/** Array input */
-	UPROPERTY(meta = (DataflowInput))
+	/** FVector array input */
+	UPROPERTY(meta = (DataflowInput, DisplayName = "VectorArray"))
 	TArray<FVector> Points;
 
+	/** FVector3f array input */
+	UPROPERTY(meta = (DataflowInput, DisplayName = "Vector3fArray"))
+	TArray<FVector3f> Vector3fArray;
+
 	/** Number of elements in the array */
 	UPROPERTY(meta = (DataflowOutput))
 	int32 NumElements = 0;
@@ -1153,6 +1157,7 @@ public:
 		: FDataflowNode(InParam, InGuid)
 	{
 		RegisterInputConnection(&Points);
+		RegisterInputConnection(&Vector3fArray);
 		RegisterOutputConnection(&NumElements);
 	}
 
@@ -1893,7 +1898,7 @@ public:
 
 	/** Vector type attribute data */
 	UPROPERTY(meta = (DataflowOutput));
-	TArray<FVector> VectorAttributeData;
+	TArray<FVector3f> VectorAttributeData;
 
 	FGetCollectionAttributeDataTypedDataflowNode(const Dataflow::FNodeParameters& InParam, FGuid InGuid = FGuid::NewGuid())
 		: FDataflowNode(InParam, InGuid)
@@ -1983,7 +1988,7 @@ struct FBoolArrayToFaceSelectionDataflowNode : public FDataflowNode
 
 public:
 	/** TArray<bool> data */
-	UPROPERTY(meta = (DataflowInput));
+	UPROPERTY(meta = (DataflowInput, DataflowIntrinsic));
 	TArray<bool> BoolAttributeData;
 
 	UPROPERTY(meta = (DataflowOutput, DisplayName = "FaceSelection"))
@@ -2003,54 +2008,36 @@ public:
 
 /**
  *
- * Adds Outside/Inside Materials to Outside/Inside faces
+ * Converts a TArray<float> to a FDataflowVertexSelection
  *
  */
-USTRUCT(meta = (DataflowGeometryCollection))
-struct FAddMaterialToCollectionDataflowNode : public FDataflowNode
+USTRUCT()
+struct FFloatArrayToVertexSelectionDataflowNode : public FDataflowNode
 {
 	GENERATED_USTRUCT_BODY()
-		DATAFLOW_NODE_DEFINE_INTERNAL(FAddMaterialToCollectionDataflowNode, "AddMaterialToCollection", "GeometryCollection|Materials", "")
+	DATAFLOW_NODE_DEFINE_INTERNAL(FFloatArrayToVertexSelectionDataflowNode, "FloatArrayToVertexSelection", "Utilities|Array", "")
 
 public:
-	/** Collection to add material(s) to */
-	UPROPERTY(meta = (DataflowInput, DataflowOutput, DataflowPassthrough = "Collection", DataflowIntrinsic))
-	FManagedArrayCollection Collection;
+	/** TArray<floatl> array */
+	UPROPERTY(meta = (DataflowInput, DataflowIntrinsic));
+	TArray<float> FloatArray;
 
-	/** Face selection, the material(s) will be added to the selected faces */
-	UPROPERTY(meta = (DataflowInput, DisplayName = "FaceSelection", DataflowIntrinsic))
-	FDataflowFaceSelection FaceSelection;
-
-	/** Materials array storing the materials */
-	UPROPERTY(meta = (DataflowInput, DataflowOutput, DataflowPassthrough = "Materials", DataflowIntrinsic))
-	TArray<TObjectPtr<UMaterial>> Materials;
-
-	/** Outside material to assign to the outside faces from the face selection */
-	UPROPERTY(EditAnywhere, Category = "Material", meta = (DataflowInput))
-	TObjectPtr<UMaterial> OutsideMaterial;
-
-	/** Inside material to assign to the inside faces from the face selection */
-	UPROPERTY(EditAnywhere, Category = "Material", meta = (DataflowInput))
-	TObjectPtr<UMaterial> InsideMaterial;
+	/** Comparison operation */
+	UPROPERTY(EditAnywhere, Category = "Compare");
+	ECompareOperationEnum Operation = ECompareOperationEnum::Dataflow_Compare_Greater;
 
-	/** If true, the outside material will be assigned to the outside faces from the face selection */
-	UPROPERTY(EditAnywhere, Category = "Material", meta = (DisplayName = "Assign Outside Material to Outside Faces"))
-	bool bAssignOutsideMaterial = true;
+	/**  */
+	UPROPERTY(EditAnywhere, Category = "Compare")
+	float Threshold = 0.f;
 
-	/** If true, the inside material will be assigned to the inside faces from the face selection */
-	UPROPERTY(EditAnywhere, Category = "Material", meta = (DisplayName = "Assign Inside Material to Inside Faces"))
-	bool bAssignInsideMaterial = false;
+	UPROPERTY(meta = (DataflowOutput, DisplayName = "VertexSelection"))
+	FDataflowVertexSelection VertexSelection;
 
-	FAddMaterialToCollectionDataflowNode(const Dataflow::FNodeParameters& InParam, FGuid InGuid = FGuid::NewGuid())
+	FFloatArrayToVertexSelectionDataflowNode(const Dataflow::FNodeParameters& InParam, FGuid InGuid = FGuid::NewGuid())
 		: FDataflowNode(InParam, InGuid)
 	{
-		RegisterInputConnection(&Collection);
-		RegisterInputConnection(&Materials);
-		RegisterInputConnection(&FaceSelection);
-		RegisterInputConnection(&OutsideMaterial);
-		RegisterInputConnection(&InsideMaterial);
-		RegisterOutputConnection(&Collection, &Collection);
-		RegisterOutputConnection(&Materials, &Materials);
+		RegisterInputConnection(&FloatArray);
+		RegisterOutputConnection(&VertexSelection);
 	}
 
 	virtual void Evaluate(Dataflow::FContext& Context, const FDataflowOutput* Out) const override;
@@ -2060,54 +2047,38 @@ public:
 
 /**
  *
- * Reassign existing material(s) to Outside/Inside faces
+ * 
  *
  */
-USTRUCT(meta = (DataflowGeometryCollection))
-struct FReAssignMaterialInCollectionDataflowNode : public FDataflowNode
+USTRUCT()
+struct FSetVertexColorInCollectionDataflowNode : public FDataflowNode
 {
 	GENERATED_USTRUCT_BODY()
-		DATAFLOW_NODE_DEFINE_INTERNAL(FReAssignMaterialInCollectionDataflowNode, "ReAssignMaterialInCollection", "GeometryCollection|Materials", "")
+	DATAFLOW_NODE_DEFINE_INTERNAL(FSetVertexColorInCollectionDataflowNode, "SetVertexColorInCollection", "Collection|Utilities", "")
 
 public:
-	/** Collection for reassign the material(s) */
+	/** Collection */
 	UPROPERTY(meta = (DataflowInput, DataflowOutput, DataflowPassthrough = "Collection", DataflowIntrinsic))
 	FManagedArrayCollection Collection;
 
-	/** Face selection, the material(s) will be assigned to the selected faces */
-	UPROPERTY(meta = (DataflowInput, DisplayName = "FaceSelection", DataflowIntrinsic))
-	FDataflowFaceSelection FaceSelection;
-
-	/** Materials array storing the materials */
-	UPROPERTY(EditAnywhere, Category = "Material", meta = (DataflowInput, DataflowOutput, DataflowPassthrough = "Materials", DataflowIntrinsic))
-	TArray<TObjectPtr<UMaterial>> Materials;
-
-	/** Index of the material in the Materials array to assign to the outside faces from the face selection */
-	UPROPERTY(EditAnywhere, Category = "Material", meta = (DataflowInput, DisplayName = "Outside Material Index", ArrayClamp = "Materials"))
-	int32 OutsideMaterialIdx = -1;
+	/**  */
+	UPROPERTY(meta = (DataflowInput, DisplayName = "VertexSelection", DataflowIntrinsic))
+	FDataflowVertexSelection VertexSelection;
 
-	/** Index of the material in the Materials array to assign to the inside faces from the face selection */
-	UPROPERTY(EditAnywhere, Category = "Material", meta = (DataflowInput, DisplayName = "Inside Material Index", ArrayClamp = "Materials"))
-	int32 InsideMaterialIdx = -1;
+	/**  */
+	UPROPERTY(EditAnywhere, Category = "Color")
+	FLinearColor SelectedColor = FLinearColor(FColor::Yellow);
 
-	/** If true, the selected material from the Materials array will be assigned to the outside faces from the face selection */
-	UPROPERTY(EditAnywhere, Category = "Material", meta = (DisplayName = "Assign Outside Material to Outside Faces"))
-	bool bAssignOutsideMaterial = false;
+	/**  */
+	UPROPERTY(EditAnywhere, Category = "Color", meta = (DisplayName = "NonSelected Color"))
+	FLinearColor NonSelectedColor = FLinearColor(FColor::Blue);
 
-	/** If true, the selected material from the Materials array will be assigned to the inside faces from the face selection */
-	UPROPERTY(EditAnywhere, Category = "Material", meta = (DisplayName = "Assign Inside Material to Inside Faces"))
-	bool bAssignInsideMaterial = false;
-
-	FReAssignMaterialInCollectionDataflowNode(const Dataflow::FNodeParameters& InParam, FGuid InGuid = FGuid::NewGuid())
+	FSetVertexColorInCollectionDataflowNode(const Dataflow::FNodeParameters& InParam, FGuid InGuid = FGuid::NewGuid())
 		: FDataflowNode(InParam, InGuid)
 	{
 		RegisterInputConnection(&Collection);
-		RegisterInputConnection(&Materials);
-		RegisterInputConnection(&FaceSelection);
-		RegisterInputConnection(&OutsideMaterialIdx);
-		RegisterInputConnection(&InsideMaterialIdx);
+		RegisterInputConnection(&VertexSelection);
 		RegisterOutputConnection(&Collection, &Collection);
-		RegisterOutputConnection(&Materials, &Materials);
 	}
 
 	virtual void Evaluate(Dataflow::FContext& Context, const FDataflowOutput* Out) const override;
@@ -2117,181 +2088,32 @@ public:
 
 /**
  *
- * Generates a formatted string of materials from the Materials array
- *
- */
-USTRUCT(meta = (DataflowGeometryCollection))
-struct FMaterialsInfoDataflowNode : public FDataflowNode
-{
-	GENERATED_USTRUCT_BODY()
-		DATAFLOW_NODE_DEFINE_INTERNAL(FMaterialsInfoDataflowNode, "MaterialsInfo", "GeometryCollection|Materials", "")
-
-public:
-	/** Materials array storing the materials */
-	UPROPERTY(meta = (DataflowInput, DataflowIntrinsic))
-	TArray<TObjectPtr<UMaterial>> Materials;
-
-	/** Formatted string of the materials */
-	UPROPERTY(meta = (DataflowOutput))
-	FString String;
-
-	FMaterialsInfoDataflowNode(const Dataflow::FNodeParameters& InParam, FGuid InGuid = FGuid::NewGuid())
-		: FDataflowNode(InParam, InGuid)
-	{
-		RegisterInputConnection(&Materials);
-		RegisterOutputConnection(&String);
-	}
-
-	virtual void Evaluate(Dataflow::FContext& Context, const FDataflowOutput* Out) const override;
-
-};
-
-
-/**
- *
- * Get a Material from a Materials array
- *
- */
-USTRUCT(meta = (DataflowGeometryCollection))
-struct FGetMaterialFromMaterialsArrayDataflowNode : public FDataflowNode
-{
-	GENERATED_USTRUCT_BODY()
-		DATAFLOW_NODE_DEFINE_INTERNAL(FGetMaterialFromMaterialsArrayDataflowNode, "GetMaterialFromMaterialsArray", "Utilities|Materials", "")
-
-public:
-	/** Materials array storing the materials */
-	UPROPERTY(meta = (DataflowInput, DataflowIntrinsic))
-	TArray<TObjectPtr<UMaterial>> Materials;
-
-	/** Selected material from the Materials array */
-	UPROPERTY(meta = (DataflowOutput))
-	TObjectPtr<UMaterial> Material;
-
-	/** Index in the Materials array for the selected material */
-	UPROPERTY(EditAnywhere, Category = "Material", meta = (DataflowInput, DisplayName = "Material Index"))
-	int32 MaterialIdx = 0;
-
-	FGetMaterialFromMaterialsArrayDataflowNode(const Dataflow::FNodeParameters& InParam, FGuid InGuid = FGuid::NewGuid())
-		: FDataflowNode(InParam, InGuid)
-	{
-		RegisterInputConnection(&Materials);
-		RegisterInputConnection(&MaterialIdx);
-		RegisterOutputConnection(&Material);
-	}
-
-	virtual void Evaluate(Dataflow::FContext& Context, const FDataflowOutput* Out) const override;
-
-};
-
-
-UENUM(BlueprintType)
-enum class ESetMaterialOperationTypeEnum : uint8
-{
-	Dataflow_SetMaterialOperationType_Add UMETA(DisplayName = "Add"),
-	Dataflow_SetMaterialOperationType_Insert UMETA(DisplayName = "Insert"),
-	//~~~
-	//256th entry
-	Dataflow_Max                UMETA(Hidden)
-};
-
-/**
- *
- * Set a Material in a Materials array
- *
- */
-USTRUCT(meta = (DataflowGeometryCollection))
-struct FSetMaterialInMaterialsArrayDataflowNode : public FDataflowNode
-{
-	GENERATED_USTRUCT_BODY()
-		DATAFLOW_NODE_DEFINE_INTERNAL(FSetMaterialInMaterialsArrayDataflowNode, "SetMaterialInMaterialsArray", "Utilities|Materials", "")
-
-public:
-	/** Materials array storing the materials */
-	UPROPERTY(meta = (DataflowInput, DataflowOutput, DataflowPassthrough = "Materials", DataflowIntrinsic))
-	TArray<TObjectPtr<UMaterial>> Materials;
-
-	/** Material to add/insert to/in Materials array */
-	UPROPERTY(EditAnywhere, Category = "Material", meta = (DataflowInput))
-	TObjectPtr<UMaterial> Material;
-
-	/** Operation type for setting the material, add will add the new material to the end off Materials array, insert will insert the
-	new material into Materials array at the index specified by MaterialIdx	*/
-	UPROPERTY(EditAnywhere, Category = "Material")
-	ESetMaterialOperationTypeEnum Operation = ESetMaterialOperationTypeEnum::Dataflow_SetMaterialOperationType_Add;
-
-	/** Index for inserting a nem material into the Materials array */
-	UPROPERTY(EditAnywhere, Category = "Material", meta = (DataflowInput, DisplayName = "Material Index", EditCondition = "Operation == ESetMaterialOperationTypeEnum::Dataflow_SetMaterialOperationType_Insert", EditConditionHides))
-	int32 MaterialIdx = 0;
-
-	FSetMaterialInMaterialsArrayDataflowNode(const Dataflow::FNodeParameters& InParam, FGuid InGuid = FGuid::NewGuid())
-		: FDataflowNode(InParam, InGuid)
-	{
-		RegisterInputConnection(&Materials);
-		RegisterInputConnection(&Material);
-		RegisterInputConnection(&MaterialIdx);
-		RegisterOutputConnection(&Materials, &Materials);
-	}
-
-	virtual void Evaluate(Dataflow::FContext& Context, const FDataflowOutput* Out) const override;
-
-};
-
-
-/**
- *
- * Makes a material
+ * Description for this node
  *
  */
 USTRUCT()
-struct FMakeMaterialDataflowNode : public FDataflowNode
+struct FMakeTransformDataflowNode : public FDataflowNode
 {
 	GENERATED_USTRUCT_BODY()
-		DATAFLOW_NODE_DEFINE_INTERNAL(FMakeMaterialDataflowNode, "MakeMaterial", "Generators|Material", "")
+	DATAFLOW_NODE_DEFINE_INTERNAL(FMakeTransformDataflowNode, "MakeTransform", "Generators|Transform", "")
 
 public:
-	/** Material which will be outputed */
-	UPROPERTY(EditAnywhere, Category = "Material", meta = (DisplayName = "Material"))
-	TObjectPtr<UMaterial> InMaterial;
+	UPROPERTY(EditAnywhere, Category = "Transform", meta = (DisplayName = "Transform"));
+	FTransform InTransform = FTransform::Identity;
 
-	/** Output material */
 	UPROPERTY(meta = (DataflowOutput));
-	TObjectPtr<UMaterial> Material;
+	FTransform Transform = FTransform::Identity;
 
-	FMakeMaterialDataflowNode(const Dataflow::FNodeParameters& InParam, FGuid InGuid = FGuid::NewGuid())
+	FMakeTransformDataflowNode(const Dataflow::FNodeParameters& InParam, FGuid InGuid = FGuid::NewGuid())
 		: FDataflowNode(InParam, InGuid)
 	{
-		RegisterOutputConnection(&Material);
+		RegisterOutputConnection(&Transform);
 	}
 
 	virtual void Evaluate(Dataflow::FContext& Context, const FDataflowOutput* Out) const override;
 
 };
 
-/**
- *
- * Makes an empty Materials array
- *
- */
-USTRUCT()
-struct FMakeMaterialsArrayDataflowNode : public FDataflowNode
-{
-	GENERATED_USTRUCT_BODY()
-	DATAFLOW_NODE_DEFINE_INTERNAL(FMakeMaterialsArrayDataflowNode, "MakeMaterialsArray", "Generators|Material", "")
-
-public:
-	/** Output Matarials array */
-	UPROPERTY(meta = (DataflowOutput));
-	TArray<TObjectPtr<UMaterial>> Materials;
-
-	FMakeMaterialsArrayDataflowNode(const Dataflow::FNodeParameters& InParam, FGuid InGuid = FGuid::NewGuid())
-		: FDataflowNode(InParam, InGuid)
-	{
-		RegisterOutputConnection(&Materials);
-	}
-
-	virtual void Evaluate(Dataflow::FContext& Context, const FDataflowOutput* Out) const override;
-
-};
 
 
 namespace Dataflow
diff --git a/Engine/Source/Runtime/Experimental/Chaos/Private/GeometryCollection/Facades/CollectionRenderingFacade.cpp b/Engine/Source/Runtime/Experimental/Chaos/Private/GeometryCollection/Facades/CollectionRenderingFacade.cpp
--- a/Engine/Source/Runtime/Experimental/Chaos/Private/GeometryCollection/Facades/CollectionRenderingFacade.cpp
+++ b/Engine/Source/Runtime/Experimental/Chaos/Private/GeometryCollection/Facades/CollectionRenderingFacade.cpp
@@ -27,7 +27,7 @@ namespace GeometryCollection::Facades
 		, IndicesStartAttribute(InCollection, "IndicesStart", FGeometryCollection::GeometryGroup)
 		, IndicesCountAttribute(InCollection, "IndicesCount", FGeometryCollection::GeometryGroup)
 		, GeometrySelectionAttribute(InCollection, "SelectionState", FGeometryCollection::GeometryGroup)
-
+		, VertexColorAttribute(InCollection, "Color", FGeometryCollection::VerticesGroup)
 	{}
 
 	FRenderingFacade::FRenderingFacade(const FManagedArrayCollection& InCollection)
@@ -47,6 +47,7 @@ namespace GeometryCollection::Facades
 		, IndicesStartAttribute(InCollection, "IndicesStart", FGeometryCollection::GeometryGroup, FGeometryCollection::FacesGroup)
 		, IndicesCountAttribute(InCollection, "IndicesCount", FGeometryCollection::GeometryGroup)
 		, GeometrySelectionAttribute(InCollection, "SelectionState", FGeometryCollection::GeometryGroup)
+		, VertexColorAttribute(InCollection, "Color", FGeometryCollection::VerticesGroup)
 	{}
 
 	//
@@ -70,6 +71,7 @@ namespace GeometryCollection::Facades
 		IndicesStartAttribute.Add();
 		IndicesCountAttribute.Add();
 		GeometrySelectionAttribute.Add();
+		VertexColorAttribute.Add();
 	}
 
 	bool FRenderingFacade::CanRenderSurface( ) const
@@ -86,7 +88,7 @@ namespace GeometryCollection::Facades
 			GeometryNameAttribute.IsValid() && GeometryHitProxyIndexAttribute.IsValid() &&
 			VertexStartAttribute.IsValid() && VertexCountAttribute.IsValid() &&
 			IndicesStartAttribute.IsValid() && IndicesCountAttribute.IsValid() &&
-			GeometrySelectionAttribute.IsValid();
+			GeometrySelectionAttribute.IsValid() && VertexColorAttribute.IsValid();
 	}
 
 	int32 FRenderingFacade::NumTriangles() const
diff --git a/Engine/Source/Runtime/Experimental/Chaos/Public/GeometryCollection/Facades/CollectionRenderingFacade.h b/Engine/Source/Runtime/Experimental/Chaos/Public/GeometryCollection/Facades/CollectionRenderingFacade.h
--- a/Engine/Source/Runtime/Experimental/Chaos/Public/GeometryCollection/Facades/CollectionRenderingFacade.h
+++ b/Engine/Source/Runtime/Experimental/Chaos/Public/GeometryCollection/Facades/CollectionRenderingFacade.h
@@ -89,6 +89,9 @@ namespace GeometryCollection::Facades
 		/** NumVertices */
 		int32 NumVertices() const { return VertexAttribute.Num(); }
 
+		/** GetVertexColorAttribute */
+		const TManagedArray<FLinearColor>& GetVertexColor() const { return VertexColorAttribute.Get(); }
+		TManagedArray<FLinearColor>& ModifyVertexColor() { check(!IsConst()); return VertexColorAttribute.Modify(); }
 
 		//
 		// Geometry Group Attributes
@@ -151,6 +154,7 @@ namespace GeometryCollection::Facades
 		TManagedArrayAccessor<int32> IndicesStartAttribute;
 		TManagedArrayAccessor<int32> IndicesCountAttribute;
 		TManagedArrayAccessor<int32> GeometrySelectionAttribute;
+		TManagedArrayAccessor<FLinearColor> VertexColorAttribute;
 
 
 	};
