Dataflow:
- Fixed Issue 257111

diff --git a/Engine/Plugins/Experimental/GeometryCollectionPlugin/Source/GeometryCollectionNodes/Private/Dataflow/GeometryCollectionFieldNodes.cpp b/Engine/Plugins/Experimental/GeometryCollectionPlugin/Source/GeometryCollectionNodes/Private/Dataflow/GeometryCollectionFieldNodes.cpp
--- a/Engine/Plugins/Experimental/GeometryCollectionPlugin/Source/GeometryCollectionNodes/Private/Dataflow/GeometryCollectionFieldNodes.cpp
+++ b/Engine/Plugins/Experimental/GeometryCollectionPlugin/Source/GeometryCollectionNodes/Private/Dataflow/GeometryCollectionFieldNodes.cpp
@@ -90,53 +90,59 @@ void FRadialFalloffFieldDataflowNode::Evaluate(Dataflow::FContext& Context, cons
 			ExecutionDatas.SamplePositions[Index] = (FVector)InVertexArray[Index];
 		}
 
-		FRadialFalloff* RadialFalloffField = new FRadialFalloff();
-		RadialFalloffField->Position = InSphere.Center + InTranslation;
-		RadialFalloffField->Radius = InSphere.W;
-		RadialFalloffField->Magnitude = Magnitude;
-		RadialFalloffField->MinRange = MinRange;
-		RadialFalloffField->MaxRange = MaxRange;
-		RadialFalloffField->Default = Default;
-		RadialFalloffField->Falloff = (EFieldFalloffType)FalloffType;
-
-		FFieldContext FieldContext{
-			ExecutionDatas,
-			FFieldContext::UniquePointerMap(),
-			0.0
-		};
-
-		TArray<float> ResultsArray;
-		ResultsArray.Init(false, NumVertices);
-		TFieldArrayView<float> ResultsView(ResultsArray, 0, ResultsArray.Num());
-		RadialFalloffField->Evaluate(FieldContext, ResultsView);
-
-		TArray<float> NewWeightArray;
-		NewWeightArray.Init(0.f, NumVertices);
-
-		if (IsConnected<FDataflowVertexSelection>(&VertexSelection))
-		{
-			const FDataflowVertexSelection& InVertexSelection = GetValue<FDataflowVertexSelection>(Context, &VertexSelection);
-
-			if (InVertexArray.Num() == InVertexSelection.Num())
+		if (TUniquePtr<FRadialFalloff> RadialFalloffField = TUniquePtr<FRadialFalloff>(new FRadialFalloff()))
+		{
+			RadialFalloffField->Position = InSphere.Center + InTranslation;
+			RadialFalloffField->Radius = InSphere.W;
+			RadialFalloffField->Magnitude = Magnitude;
+			RadialFalloffField->MinRange = MinRange;
+			RadialFalloffField->MaxRange = MaxRange;
+			RadialFalloffField->Default = Default;
+			RadialFalloffField->Falloff = (EFieldFalloffType)FalloffType;
+
+			FFieldContext FieldContext{
+				ExecutionDatas,
+				FFieldContext::UniquePointerMap(),
+				0.0
+			};
+
+			TArray<float> ResultsArray;
+			ResultsArray.Init(false, NumVertices);
+			TFieldArrayView<float> ResultsView(ResultsArray, 0, ResultsArray.Num());
+			RadialFalloffField->Evaluate(FieldContext, ResultsView);
+
+			TArray<float> NewWeightArray;
+			NewWeightArray.Init(0.f, NumVertices);
+
+			if (IsConnected<FDataflowVertexSelection>(&VertexSelection))
 			{
-				for (int32 Idx = 0; Idx < NumVertices; ++Idx)
+				const FDataflowVertexSelection& InVertexSelection = GetValue<FDataflowVertexSelection>(Context, &VertexSelection);
+
+				if (InVertexArray.Num() == InVertexSelection.Num())
 				{
-					if (InVertexSelection.IsSelected(Idx))
+					for (int32 Idx = 0; Idx < NumVertices; ++Idx)
 					{
-						NewWeightArray[Idx] = ResultsView[Idx];
+						if (InVertexSelection.IsSelected(Idx))
+						{
+							NewWeightArray[Idx] = ResultsView[Idx];
+						}
 					}
 				}
 			}
-		}
-		else
-		{
-			for (int32 Idx = 0; Idx < NumVertices; ++Idx)
+			else
 			{
-				NewWeightArray[Idx] = ResultsView[Idx];
+				for (int32 Idx = 0; Idx < NumVertices; ++Idx)
+				{
+					NewWeightArray[Idx] = ResultsView[Idx];
+				}
 			}
+
+			SetValue<TArray<float>>(Context, NewWeightArray, &WeightArray);
+
+			return;
 		}
 
-		SetValue<TArray<float>>(Context, NewWeightArray, &WeightArray);
+		SetValue<TArray<float>>(Context, TArray<float>(), &WeightArray);
 	}
 }
 
@@ -162,54 +168,60 @@ void FPlaneFalloffFieldDataflowNode::Evaluate(Dataflow::FContext& Context, const
 			ExecutionDatas.SamplePositions[Index] = (FVector)InVertexArray[Index];
 		}
 
-		FPlaneFalloff* PlaneFalloffField = new FPlaneFalloff();
-		PlaneFalloffField->Position = InPosition + InTranslation;
-		PlaneFalloffField->Normal = InNormal;
-		PlaneFalloffField->Magnitude = Magnitude;
-		PlaneFalloffField->MinRange = MinRange;
-		PlaneFalloffField->MaxRange = MaxRange;
-		PlaneFalloffField->Distance = Distance;
-		PlaneFalloffField->Default = Default;
-		PlaneFalloffField->Falloff = (EFieldFalloffType)FalloffType;
-
-		FFieldContext FieldContext{
-			ExecutionDatas,
-			FFieldContext::UniquePointerMap(),
-			0.0
-		};
-
-		TArray<float> ResultsArray;
-		ResultsArray.Init(false, NumVertices);
-		TFieldArrayView<float> ResultsView(ResultsArray, 0, ResultsArray.Num());
-		PlaneFalloffField->Evaluate(FieldContext, ResultsView);
-
-		TArray<float> NewWeightArray;
-		NewWeightArray.Init(0.f, NumVertices);
-
-		if (IsConnected<FDataflowVertexSelection>(&VertexSelection))
+		if (TUniquePtr<FPlaneFalloff> PlaneFalloffField = TUniquePtr<FPlaneFalloff>(new FPlaneFalloff()))
 		{
-			const FDataflowVertexSelection& InVertexSelection = GetValue<FDataflowVertexSelection>(Context, &VertexSelection);
+			PlaneFalloffField->Position = InPosition + InTranslation;
+			PlaneFalloffField->Normal = InNormal;
+			PlaneFalloffField->Magnitude = Magnitude;
+			PlaneFalloffField->MinRange = MinRange;
+			PlaneFalloffField->MaxRange = MaxRange;
+			PlaneFalloffField->Distance = Distance;
+			PlaneFalloffField->Default = Default;
+			PlaneFalloffField->Falloff = (EFieldFalloffType)FalloffType;
+
+			FFieldContext FieldContext{
+				ExecutionDatas,
+				FFieldContext::UniquePointerMap(),
+				0.0
+			};
+
+			TArray<float> ResultsArray;
+			ResultsArray.Init(false, NumVertices);
+			TFieldArrayView<float> ResultsView(ResultsArray, 0, ResultsArray.Num());
+			PlaneFalloffField->Evaluate(FieldContext, ResultsView);
 
-			if (InVertexArray.Num() == InVertexSelection.Num())
+			TArray<float> NewWeightArray;
+			NewWeightArray.Init(0.f, NumVertices);
+
+			if (IsConnected<FDataflowVertexSelection>(&VertexSelection))
 			{
-				for (int32 Idx = 0; Idx < NumVertices; ++Idx)
+				const FDataflowVertexSelection& InVertexSelection = GetValue<FDataflowVertexSelection>(Context, &VertexSelection);
+
+				if (InVertexArray.Num() == InVertexSelection.Num())
 				{
-					if (InVertexSelection.IsSelected(Idx))
+					for (int32 Idx = 0; Idx < NumVertices; ++Idx)
 					{
-						NewWeightArray[Idx] = ResultsView[Idx];
+						if (InVertexSelection.IsSelected(Idx))
+						{
+							NewWeightArray[Idx] = ResultsView[Idx];
+						}
 					}
 				}
 			}
-		}
-		else
-		{
-			for (int32 Idx = 0; Idx < NumVertices; ++Idx)
+			else
 			{
-				NewWeightArray[Idx] = ResultsView[Idx];
+				for (int32 Idx = 0; Idx < NumVertices; ++Idx)
+				{
+					NewWeightArray[Idx] = ResultsView[Idx];
+				}
 			}
+
+			SetValue<TArray<float>>(Context, NewWeightArray, &WeightArray);
+
+			return;
 		}
 
-		SetValue<TArray<float>>(Context, NewWeightArray, &WeightArray);
+		SetValue<TArray<float>>(Context, TArray<float>(), &WeightArray);
 	}
 }
 
@@ -234,51 +246,57 @@ void FRadialIntMaskFieldDataflowNode::Evaluate(Dataflow::FContext& Context, cons
 			ExecutionDatas.SamplePositions[Index] = (FVector)InVertexArray[Index];
 		}
 
-		FRadialIntMask* RadialIntMaskField = new FRadialIntMask();
-		RadialIntMaskField->Position = InSphere.Center + InTranslation;
-		RadialIntMaskField->Radius = InSphere.W;
-		RadialIntMaskField->InteriorValue = InteriorValue;
-		RadialIntMaskField->ExteriorValue = ExteriorValue;
-		RadialIntMaskField->SetMaskCondition = (ESetMaskConditionType)SetMaskCondition;
-
-		FFieldContext FieldContext{
-			ExecutionDatas,
-			FFieldContext::UniquePointerMap(),
-			0.0
-		};
+		if (TUniquePtr<FRadialIntMask> RadialIntMaskField = TUniquePtr<FRadialIntMask>(new FRadialIntMask()))
+		{
+			RadialIntMaskField->Position = InSphere.Center + InTranslation;
+			RadialIntMaskField->Radius = InSphere.W;
+			RadialIntMaskField->InteriorValue = InteriorValue;
+			RadialIntMaskField->ExteriorValue = ExteriorValue;
+			RadialIntMaskField->SetMaskCondition = (ESetMaskConditionType)SetMaskCondition;
 
-		TArray<int32> ResultsArray;
-		ResultsArray.Init(false, NumVertices);
-		TFieldArrayView<int32> ResultsView(ResultsArray, 0, ResultsArray.Num());
-		RadialIntMaskField->Evaluate(FieldContext, ResultsView);
+			FFieldContext FieldContext{
+				ExecutionDatas,
+				FFieldContext::UniquePointerMap(),
+				0.0
+			};
 
-		TArray<int32> NewMaskArray;
-		NewMaskArray.Init(0, NumVertices);
+			TArray<int32> ResultsArray;
+			ResultsArray.Init(false, NumVertices);
+			TFieldArrayView<int32> ResultsView(ResultsArray, 0, ResultsArray.Num());
+			RadialIntMaskField->Evaluate(FieldContext, ResultsView);
 
-		if (IsConnected<FDataflowVertexSelection>(&VertexSelection))
-		{
-			const FDataflowVertexSelection& InVertexSelection = GetValue<FDataflowVertexSelection>(Context, &VertexSelection);
+			TArray<int32> NewMaskArray;
+			NewMaskArray.Init(0, NumVertices);
 
-			if (InVertexArray.Num() == InVertexSelection.Num())
+			if (IsConnected<FDataflowVertexSelection>(&VertexSelection))
 			{
-				for (int32 Idx = 0; Idx < NumVertices; ++Idx)
+				const FDataflowVertexSelection& InVertexSelection = GetValue<FDataflowVertexSelection>(Context, &VertexSelection);
+
+				if (InVertexArray.Num() == InVertexSelection.Num())
 				{
-					if (InVertexSelection.IsSelected(Idx))
+					for (int32 Idx = 0; Idx < NumVertices; ++Idx)
 					{
-						NewMaskArray[Idx] = ResultsView[Idx];
+						if (InVertexSelection.IsSelected(Idx))
+						{
+							NewMaskArray[Idx] = ResultsView[Idx];
+						}
 					}
 				}
 			}
-		}
-		else
-		{
-			for (int32 Idx = 0; Idx < NumVertices; ++Idx)
+			else
 			{
-				NewMaskArray[Idx] = ResultsView[Idx];
+				for (int32 Idx = 0; Idx < NumVertices; ++Idx)
+				{
+					NewMaskArray[Idx] = ResultsView[Idx];
+				}
 			}
+
+			SetValue<TArray<int32>>(Context, NewMaskArray, &MaskArray);
+
+			return;
 		}
 
-		SetValue<TArray<int32>>(Context, NewMaskArray, &MaskArray);
+		SetValue<TArray<int32>>(Context, TArray<int32>(), &MaskArray);
 	}
 }
 
@@ -301,47 +319,53 @@ void FUniformScalarFieldDataflowNode::Evaluate(Dataflow::FContext& Context, cons
 			ExecutionDatas.SamplePositions[Index] = (FVector)InVertexArray[Index];
 		}
 
-		FUniformScalar* UniformScalarField = new FUniformScalar();
-		UniformScalarField->Magnitude = Magnitude;
-
-		FFieldContext FieldContext{
-			ExecutionDatas,
-			FFieldContext::UniquePointerMap(),
-			0.0
-		};
+		if (TUniquePtr<FUniformScalar> UniformScalarField = TUniquePtr<FUniformScalar>(new FUniformScalar()))
+		{
+			UniformScalarField->Magnitude = Magnitude;
 
-		TArray<float> ResultsArray;
-		ResultsArray.Init(false, NumVertices);
-		TFieldArrayView<float> ResultsView(ResultsArray, 0, ResultsArray.Num());
-		UniformScalarField->Evaluate(FieldContext, ResultsView);
+			FFieldContext FieldContext{
+				ExecutionDatas,
+				FFieldContext::UniquePointerMap(),
+				0.0
+			};
 
-		TArray<float> NewWeightArray;
-		NewWeightArray.Init(0.f, NumVertices);
+			TArray<float> ResultsArray;
+			ResultsArray.Init(false, NumVertices);
+			TFieldArrayView<float> ResultsView(ResultsArray, 0, ResultsArray.Num());
+			UniformScalarField->Evaluate(FieldContext, ResultsView);
 
-		if (IsConnected<FDataflowVertexSelection>(&VertexSelection))
-		{
-			const FDataflowVertexSelection& InVertexSelection = GetValue<FDataflowVertexSelection>(Context, &VertexSelection);
+			TArray<float> NewWeightArray;
+			NewWeightArray.Init(0.f, NumVertices);
 
-			if (InVertexArray.Num() == InVertexSelection.Num())
+			if (IsConnected<FDataflowVertexSelection>(&VertexSelection))
 			{
-				for (int32 Idx = 0; Idx < NumVertices; ++Idx)
+				const FDataflowVertexSelection& InVertexSelection = GetValue<FDataflowVertexSelection>(Context, &VertexSelection);
+
+				if (InVertexArray.Num() == InVertexSelection.Num())
 				{
-					if (InVertexSelection.IsSelected(Idx))
+					for (int32 Idx = 0; Idx < NumVertices; ++Idx)
 					{
-						NewWeightArray[Idx] = ResultsView[Idx];
+						if (InVertexSelection.IsSelected(Idx))
+						{
+							NewWeightArray[Idx] = ResultsView[Idx];
+						}
 					}
 				}
 			}
-		}
-		else
-		{
-			for (int32 Idx = 0; Idx < NumVertices; ++Idx)
+			else
 			{
-				NewWeightArray[Idx] = ResultsView[Idx];
+				for (int32 Idx = 0; Idx < NumVertices; ++Idx)
+				{
+					NewWeightArray[Idx] = ResultsView[Idx];
+				}
 			}
+
+			SetValue<TArray<float>>(Context, NewWeightArray, &WeightArray);
+
+			return;
 		}
 
-		SetValue<TArray<float>>(Context, NewWeightArray, &WeightArray);
+		SetValue<TArray<float>>(Context, TArray<float>(), &WeightArray);
 	}
 }
 
@@ -364,48 +388,54 @@ void FUniformVectorFieldDataflowNode::Evaluate(Dataflow::FContext& Context, cons
 			ExecutionDatas.SamplePositions[Index] = (FVector)InVertexArray[Index];
 		}
 
-		FUniformVector* UniformVectorField = new FUniformVector();
-		UniformVectorField->Magnitude = Magnitude;
-		UniformVectorField->Direction = Direction;
-
-		FFieldContext FieldContext{
-			ExecutionDatas,
-			FFieldContext::UniquePointerMap(),
-			0.0
-		};
+		if (TUniquePtr<FUniformVector> UniformVectorField = TUniquePtr<FUniformVector>(new FUniformVector()))
+		{
+			UniformVectorField->Magnitude = Magnitude;
+			UniformVectorField->Direction = Direction;
 
-		TArray<FVector> ResultsArray;
-		ResultsArray.Init(FVector(0.f), NumVertices);
-		TFieldArrayView<FVector> ResultsView(ResultsArray, 0, ResultsArray.Num());
-		UniformVectorField->Evaluate(FieldContext, ResultsView);
+			FFieldContext FieldContext{
+				ExecutionDatas,
+				FFieldContext::UniquePointerMap(),
+				0.0
+			};
 
-		TArray<FVector> NewVectorArray;
-		NewVectorArray.Init(FVector(0.f), NumVertices);
+			TArray<FVector> ResultsArray;
+			ResultsArray.Init(FVector(0.f), NumVertices);
+			TFieldArrayView<FVector> ResultsView(ResultsArray, 0, ResultsArray.Num());
+			UniformVectorField->Evaluate(FieldContext, ResultsView);
 
-		if (IsConnected<FDataflowVertexSelection>(&VertexSelection))
-		{
-			const FDataflowVertexSelection& InVertexSelection = GetValue<FDataflowVertexSelection>(Context, &VertexSelection);
+			TArray<FVector> NewVectorArray;
+			NewVectorArray.Init(FVector(0.f), NumVertices);
 
-			if (InVertexArray.Num() == InVertexSelection.Num())
+			if (IsConnected<FDataflowVertexSelection>(&VertexSelection))
 			{
-				for (int32 Idx = 0; Idx < NumVertices; ++Idx)
+				const FDataflowVertexSelection& InVertexSelection = GetValue<FDataflowVertexSelection>(Context, &VertexSelection);
+
+				if (InVertexArray.Num() == InVertexSelection.Num())
 				{
-					if (InVertexSelection.IsSelected(Idx))
+					for (int32 Idx = 0; Idx < NumVertices; ++Idx)
 					{
-						NewVectorArray[Idx] = ResultsView[Idx];
+						if (InVertexSelection.IsSelected(Idx))
+						{
+							NewVectorArray[Idx] = ResultsView[Idx];
+						}
 					}
 				}
 			}
-		}
-		else
-		{
-			for (int32 Idx = 0; Idx < NumVertices; ++Idx)
+			else
 			{
-				NewVectorArray[Idx] = ResultsView[Idx];
+				for (int32 Idx = 0; Idx < NumVertices; ++Idx)
+				{
+					NewVectorArray[Idx] = ResultsView[Idx];
+				}
 			}
+
+			SetValue<TArray<FVector>>(Context, NewVectorArray, &VectorArray);
+
+			return;
 		}
 
-		SetValue<TArray<FVector>>(Context, NewVectorArray, &VectorArray);
+		SetValue<TArray<FVector>>(Context, TArray<FVector>(), &VectorArray);
 	}
 }
 
@@ -428,48 +458,54 @@ void FRadialVectorFieldDataflowNode::Evaluate(Dataflow::FContext& Context, const
 			ExecutionDatas.SamplePositions[Index] = (FVector)InVertexArray[Index];
 		}
 
-		FRadialVector* RadialVectorField = new FRadialVector();
-		RadialVectorField->Magnitude = Magnitude;
-		RadialVectorField->Position = Position;
-
-		FFieldContext FieldContext{
-			ExecutionDatas,
-			FFieldContext::UniquePointerMap(),
-			0.0
-		};
+		if (TUniquePtr<FRadialVector> RadialVectorField = TUniquePtr<FRadialVector>(new FRadialVector()))
+		{
+			RadialVectorField->Magnitude = Magnitude;
+			RadialVectorField->Position = Position;
 
-		TArray<FVector> ResultsArray;
-		ResultsArray.Init(FVector(0.f), NumVertices);
-		TFieldArrayView<FVector> ResultsView(ResultsArray, 0, ResultsArray.Num());
-		RadialVectorField->Evaluate(FieldContext, ResultsView);
+			FFieldContext FieldContext{
+				ExecutionDatas,
+				FFieldContext::UniquePointerMap(),
+				0.0
+			};
 
-		TArray<FVector> NewVectorArray;
-		NewVectorArray.Init(FVector(0.f), NumVertices);
+			TArray<FVector> ResultsArray;
+			ResultsArray.Init(FVector(0.f), NumVertices);
+			TFieldArrayView<FVector> ResultsView(ResultsArray, 0, ResultsArray.Num());
+			RadialVectorField->Evaluate(FieldContext, ResultsView);
 
-		if (IsConnected<FDataflowVertexSelection>(&VertexSelection))
-		{
-			const FDataflowVertexSelection& InVertexSelection = GetValue<FDataflowVertexSelection>(Context, &VertexSelection);
+			TArray<FVector> NewVectorArray;
+			NewVectorArray.Init(FVector(0.f), NumVertices);
 
-			if (InVertexArray.Num() == InVertexSelection.Num())
+			if (IsConnected<FDataflowVertexSelection>(&VertexSelection))
 			{
-				for (int32 Idx = 0; Idx < NumVertices; ++Idx)
+				const FDataflowVertexSelection& InVertexSelection = GetValue<FDataflowVertexSelection>(Context, &VertexSelection);
+
+				if (InVertexArray.Num() == InVertexSelection.Num())
 				{
-					if (InVertexSelection.IsSelected(Idx))
+					for (int32 Idx = 0; Idx < NumVertices; ++Idx)
 					{
-						NewVectorArray[Idx] = ResultsView[Idx];
+						if (InVertexSelection.IsSelected(Idx))
+						{
+							NewVectorArray[Idx] = ResultsView[Idx];
+						}
 					}
 				}
 			}
-		}
-		else
-		{
-			for (int32 Idx = 0; Idx < NumVertices; ++Idx)
+			else
 			{
-				NewVectorArray[Idx] = ResultsView[Idx];
+				for (int32 Idx = 0; Idx < NumVertices; ++Idx)
+				{
+					NewVectorArray[Idx] = ResultsView[Idx];
+				}
 			}
+
+			SetValue<TArray<FVector>>(Context, NewVectorArray, &VectorArray);
+
+			return;
 		}
 
-		SetValue<TArray<FVector>>(Context, NewVectorArray, &VectorArray);
+		SetValue<TArray<FVector>>(Context, TArray<FVector>(), &VectorArray);
 	}
 }
 
@@ -492,47 +528,53 @@ void FRandomVectorFieldDataflowNode::Evaluate(Dataflow::FContext& Context, const
 			ExecutionDatas.SamplePositions[Index] = (FVector)InVertexArray[Index];
 		}
 
-		FRandomVector* RandomVectorField = new FRandomVector();
-		RandomVectorField->Magnitude = Magnitude;
-
-		FFieldContext FieldContext{
-			ExecutionDatas,
-			FFieldContext::UniquePointerMap(),
-			0.0
-		};
+		if (TUniquePtr<FRandomVector> RandomVectorField = TUniquePtr<FRandomVector>(new FRandomVector()))
+		{
+			RandomVectorField->Magnitude = Magnitude;
 
-		TArray<FVector> ResultsArray;
-		ResultsArray.Init(FVector(0.f), NumVertices);
-		TFieldArrayView<FVector> ResultsView(ResultsArray, 0, ResultsArray.Num());
-		RandomVectorField->Evaluate(FieldContext, ResultsView);
+			FFieldContext FieldContext{
+				ExecutionDatas,
+				FFieldContext::UniquePointerMap(),
+				0.0
+			};
 
-		TArray<FVector> NewVectorArray;
-		NewVectorArray.Init(FVector(0.f), NumVertices);
+			TArray<FVector> ResultsArray;
+			ResultsArray.Init(FVector(0.f), NumVertices);
+			TFieldArrayView<FVector> ResultsView(ResultsArray, 0, ResultsArray.Num());
+			RandomVectorField->Evaluate(FieldContext, ResultsView);
 
-		if (IsConnected<FDataflowVertexSelection>(&VertexSelection))
-		{
-			const FDataflowVertexSelection& InVertexSelection = GetValue<FDataflowVertexSelection>(Context, &VertexSelection);
+			TArray<FVector> NewVectorArray;
+			NewVectorArray.Init(FVector(0.f), NumVertices);
 
-			if (InVertexArray.Num() == InVertexSelection.Num())
+			if (IsConnected<FDataflowVertexSelection>(&VertexSelection))
 			{
-				for (int32 Idx = 0; Idx < NumVertices; ++Idx)
+				const FDataflowVertexSelection& InVertexSelection = GetValue<FDataflowVertexSelection>(Context, &VertexSelection);
+
+				if (InVertexArray.Num() == InVertexSelection.Num())
 				{
-					if (InVertexSelection.IsSelected(Idx))
+					for (int32 Idx = 0; Idx < NumVertices; ++Idx)
 					{
-						NewVectorArray[Idx] = ResultsView[Idx];
+						if (InVertexSelection.IsSelected(Idx))
+						{
+							NewVectorArray[Idx] = ResultsView[Idx];
+						}
 					}
 				}
 			}
-		}
-		else
-		{
-			for (int32 Idx = 0; Idx < NumVertices; ++Idx)
+			else
 			{
-				NewVectorArray[Idx] = ResultsView[Idx];
+				for (int32 Idx = 0; Idx < NumVertices; ++Idx)
+				{
+					NewVectorArray[Idx] = ResultsView[Idx];
+				}
 			}
+
+			SetValue<TArray<FVector>>(Context, NewVectorArray, &VectorArray);
+
+			return;
 		}
 
-		SetValue<TArray<FVector>>(Context, NewVectorArray, &VectorArray);
+		SetValue<TArray<FVector>>(Context, TArray<FVector>(), &VectorArray);
 	}
 }
 
@@ -556,49 +598,55 @@ void FNoiseFieldDataflowNode::Evaluate(Dataflow::FContext& Context, const FDataf
 			ExecutionDatas.SamplePositions[Index] = (FVector)InVertexArray[Index];
 		}
 
-		FNoiseField* NoiseField = new FNoiseField();
-		NoiseField->MinRange = MinRange;
-		NoiseField->MaxRange = MaxRange;
-		NoiseField->Transform = Transform;
-
-		FFieldContext FieldContext{
-			ExecutionDatas,
-			FFieldContext::UniquePointerMap(),
-			0.0
-		};
+		if (TUniquePtr<FNoiseField> NoiseField = TUniquePtr<FNoiseField>(new FNoiseField()))
+		{
+			NoiseField->MinRange = MinRange;
+			NoiseField->MaxRange = MaxRange;
+			NoiseField->Transform = Transform;
 
-		TArray<float> ResultsArray;
-		ResultsArray.Init(false, NumVertices);
-		TFieldArrayView<float> ResultsView(ResultsArray, 0, ResultsArray.Num());
-		NoiseField->Evaluate(FieldContext, ResultsView);
+			FFieldContext FieldContext{
+				ExecutionDatas,
+				FFieldContext::UniquePointerMap(),
+				0.0
+			};
 
-		TArray<float> NewWeightArray;
-		NewWeightArray.Init(0.f, NumVertices);
+			TArray<float> ResultsArray;
+			ResultsArray.Init(false, NumVertices);
+			TFieldArrayView<float> ResultsView(ResultsArray, 0, ResultsArray.Num());
+			NoiseField->Evaluate(FieldContext, ResultsView);
 
-		if (IsConnected<FDataflowVertexSelection>(&VertexSelection))
-		{
-			const FDataflowVertexSelection& InVertexSelection = GetValue<FDataflowVertexSelection>(Context, &VertexSelection);
+			TArray<float> NewWeightArray;
+			NewWeightArray.Init(0.f, NumVertices);
 
-			if (InVertexArray.Num() == InVertexSelection.Num())
+			if (IsConnected<FDataflowVertexSelection>(&VertexSelection))
 			{
-				for (int32 Idx = 0; Idx < NumVertices; ++Idx)
+				const FDataflowVertexSelection& InVertexSelection = GetValue<FDataflowVertexSelection>(Context, &VertexSelection);
+
+				if (InVertexArray.Num() == InVertexSelection.Num())
 				{
-					if (InVertexSelection.IsSelected(Idx))
+					for (int32 Idx = 0; Idx < NumVertices; ++Idx)
 					{
-						NewWeightArray[Idx] = ResultsView[Idx];
+						if (InVertexSelection.IsSelected(Idx))
+						{
+							NewWeightArray[Idx] = ResultsView[Idx];
+						}
 					}
 				}
 			}
-		}
-		else
-		{
-			for (int32 Idx = 0; Idx < NumVertices; ++Idx)
+			else
 			{
-				NewWeightArray[Idx] = ResultsView[Idx];
+				for (int32 Idx = 0; Idx < NumVertices; ++Idx)
+				{
+					NewWeightArray[Idx] = ResultsView[Idx];
+				}
 			}
+
+			SetValue<TArray<float>>(Context, NewWeightArray, &WeightArray);
+
+			return;
 		}
 
-		SetValue<TArray<float>>(Context, NewWeightArray, &WeightArray);
+		SetValue<TArray<float>>(Context, TArray<float>(), &WeightArray);
 	}
 }
 
@@ -621,47 +669,53 @@ void FUniformIntegerFieldDataflowNode::Evaluate(Dataflow::FContext& Context, con
 			ExecutionDatas.SamplePositions[Index] = (FVector)InVertexArray[Index];
 		}
 
-		FUniformInteger* UniformIntegerField = new FUniformInteger();
-		UniformIntegerField->Magnitude = Magnitude;
-
-		FFieldContext FieldContext{
-			ExecutionDatas,
-			FFieldContext::UniquePointerMap(),
-			0.0
-		};
+		if (TUniquePtr<FUniformInteger> UniformIntegerField = TUniquePtr<FUniformInteger>(new FUniformInteger()))
+		{
+			UniformIntegerField->Magnitude = Magnitude;
 
-		TArray<int32> ResultsArray;
-		ResultsArray.Init(false, NumVertices);
-		TFieldArrayView<int32> ResultsView(ResultsArray, 0, ResultsArray.Num());
-		UniformIntegerField->Evaluate(FieldContext, ResultsView);
+			FFieldContext FieldContext{
+				ExecutionDatas,
+				FFieldContext::UniquePointerMap(),
+				0.0
+			};
 
-		TArray<int32> NewMaskArray;
-		NewMaskArray.Init(0, NumVertices);
+			TArray<int32> ResultsArray;
+			ResultsArray.Init(false, NumVertices);
+			TFieldArrayView<int32> ResultsView(ResultsArray, 0, ResultsArray.Num());
+			UniformIntegerField->Evaluate(FieldContext, ResultsView);
 
-		if (IsConnected<FDataflowVertexSelection>(&VertexSelection))
-		{
-			const FDataflowVertexSelection& InVertexSelection = GetValue<FDataflowVertexSelection>(Context, &VertexSelection);
+			TArray<int32> NewMaskArray;
+			NewMaskArray.Init(0, NumVertices);
 
-			if (InVertexArray.Num() == InVertexSelection.Num())
+			if (IsConnected<FDataflowVertexSelection>(&VertexSelection))
 			{
-				for (int32 Idx = 0; Idx < NumVertices; ++Idx)
+				const FDataflowVertexSelection& InVertexSelection = GetValue<FDataflowVertexSelection>(Context, &VertexSelection);
+
+				if (InVertexArray.Num() == InVertexSelection.Num())
 				{
-					if (InVertexSelection.IsSelected(Idx))
+					for (int32 Idx = 0; Idx < NumVertices; ++Idx)
 					{
-						NewMaskArray[Idx] = ResultsView[Idx];
+						if (InVertexSelection.IsSelected(Idx))
+						{
+							NewMaskArray[Idx] = ResultsView[Idx];
+						}
 					}
 				}
 			}
-		}
-		else
-		{
-			for (int32 Idx = 0; Idx < NumVertices; ++Idx)
+			else
 			{
-				NewMaskArray[Idx] = ResultsView[Idx];
+				for (int32 Idx = 0; Idx < NumVertices; ++Idx)
+				{
+					NewMaskArray[Idx] = ResultsView[Idx];
+				}
 			}
+
+			SetValue<TArray<int32>>(Context, NewMaskArray, &MaskArray);
+
+			return;
 		}
 
-		SetValue<TArray<int32>>(Context, NewMaskArray, &MaskArray);
+		SetValue<TArray<int32>>(Context, TArray<int32>(), &MaskArray);
 	}
 }
 
@@ -685,52 +739,58 @@ void FWaveScalarFieldDataflowNode::Evaluate(Dataflow::FContext& Context, const F
 			ExecutionDatas.SamplePositions[Index] = (FVector)InVertexArray[Index];
 		}
 
-		FWaveScalar* WaveScalarField = new FWaveScalar();
-		WaveScalarField->Magnitude = Magnitude;
-		WaveScalarField->Position = Position + InTranslation;
-		WaveScalarField->Wavelength = Wavelength;
-		WaveScalarField->Period = Period;
-		WaveScalarField->Function = (EWaveFunctionType)FunctionType;
-		WaveScalarField->Falloff = (EFieldFalloffType)FalloffType;
-
-		FFieldContext FieldContext{
-			ExecutionDatas,
-			FFieldContext::UniquePointerMap(),
-			0.0
-		};
-
-		TArray<float> ResultsArray;
-		ResultsArray.Init(false, NumVertices);
-		TFieldArrayView<float> ResultsView(ResultsArray, 0, ResultsArray.Num());
-		WaveScalarField->Evaluate(FieldContext, ResultsView);
-
-		TArray<float> NewWeightArray;
-		NewWeightArray.Init(0.f, NumVertices);
-
-		if (IsConnected<FDataflowVertexSelection>(&VertexSelection))
+		if (TUniquePtr<FWaveScalar> WaveScalarField = TUniquePtr<FWaveScalar>(new FWaveScalar()))
 		{
-			const FDataflowVertexSelection& InVertexSelection = GetValue<FDataflowVertexSelection>(Context, &VertexSelection);
+			WaveScalarField->Magnitude = Magnitude;
+			WaveScalarField->Position = Position + InTranslation;
+			WaveScalarField->Wavelength = Wavelength;
+			WaveScalarField->Period = Period;
+			WaveScalarField->Function = (EWaveFunctionType)FunctionType;
+			WaveScalarField->Falloff = (EFieldFalloffType)FalloffType;
 
-			if (InVertexArray.Num() == InVertexSelection.Num())
+			FFieldContext FieldContext{
+				ExecutionDatas,
+				FFieldContext::UniquePointerMap(),
+				0.0
+			};
+
+			TArray<float> ResultsArray;
+			ResultsArray.Init(false, NumVertices);
+			TFieldArrayView<float> ResultsView(ResultsArray, 0, ResultsArray.Num());
+			WaveScalarField->Evaluate(FieldContext, ResultsView);
+
+			TArray<float> NewWeightArray;
+			NewWeightArray.Init(0.f, NumVertices);
+
+			if (IsConnected<FDataflowVertexSelection>(&VertexSelection))
 			{
-				for (int32 Idx = 0; Idx < NumVertices; ++Idx)
+				const FDataflowVertexSelection& InVertexSelection = GetValue<FDataflowVertexSelection>(Context, &VertexSelection);
+
+				if (InVertexArray.Num() == InVertexSelection.Num())
 				{
-					if (InVertexSelection.IsSelected(Idx))
+					for (int32 Idx = 0; Idx < NumVertices; ++Idx)
 					{
-						NewWeightArray[Idx] = ResultsView[Idx];
+						if (InVertexSelection.IsSelected(Idx))
+						{
+							NewWeightArray[Idx] = ResultsView[Idx];
+						}
 					}
 				}
 			}
-		}
-		else
-		{
-			for (int32 Idx = 0; Idx < NumVertices; ++Idx)
+			else
 			{
-				NewWeightArray[Idx] = ResultsView[Idx];
+				for (int32 Idx = 0; Idx < NumVertices; ++Idx)
+				{
+					NewWeightArray[Idx] = ResultsView[Idx];
+				}
 			}
+
+			SetValue<TArray<float>>(Context, NewWeightArray, &WeightArray);
+
+			return;
 		}
 
-		SetValue<TArray<float>>(Context, NewWeightArray, &WeightArray);
+		SetValue<TArray<float>>(Context, TArray<float>(), &WeightArray);
 	}
 }
 
@@ -755,52 +815,58 @@ void FBoxFalloffFieldDataflowNode::Evaluate(Dataflow::FContext& Context, const F
 			ExecutionDatas.SamplePositions[Index] = (FVector)InVertexArray[Index];
 		}
 
-		FBoxFalloff* BoxFalloffField = new FBoxFalloff();
-		BoxFalloffField->Magnitude = Magnitude;
-		BoxFalloffField->MinRange = MinRange;
-		BoxFalloffField->MaxRange = MaxRange;
-		BoxFalloffField->Default = Default;
-		BoxFalloffField->Transform = InTransform;
-		BoxFalloffField->Falloff = (EFieldFalloffType)FalloffType;
-
-		FFieldContext FieldContext{
-			ExecutionDatas,
-			FFieldContext::UniquePointerMap(),
-			0.0
-		};
-
-		TArray<float> ResultsArray;
-		ResultsArray.Init(false, NumVertices);
-		TFieldArrayView<float> ResultsView(ResultsArray, 0, ResultsArray.Num());
-		BoxFalloffField->Evaluate(FieldContext, ResultsView);
-
-		TArray<float> NewWeightArray;
-		NewWeightArray.Init(0.f, NumVertices);
-
-		if (IsConnected<FDataflowVertexSelection>(&VertexSelection))
+		if (TUniquePtr<FBoxFalloff> BoxFalloffField = TUniquePtr<FBoxFalloff>(new FBoxFalloff()))
 		{
-			const FDataflowVertexSelection& InVertexSelection = GetValue<FDataflowVertexSelection>(Context, &VertexSelection);
+			BoxFalloffField->Magnitude = Magnitude;
+			BoxFalloffField->MinRange = MinRange;
+			BoxFalloffField->MaxRange = MaxRange;
+			BoxFalloffField->Default = Default;
+			BoxFalloffField->Transform = InTransform;
+			BoxFalloffField->Falloff = (EFieldFalloffType)FalloffType;
 
-			if (InVertexArray.Num() == InVertexSelection.Num())
+			FFieldContext FieldContext{
+				ExecutionDatas,
+				FFieldContext::UniquePointerMap(),
+				0.0
+			};
+
+			TArray<float> ResultsArray;
+			ResultsArray.Init(false, NumVertices);
+			TFieldArrayView<float> ResultsView(ResultsArray, 0, ResultsArray.Num());
+			BoxFalloffField->Evaluate(FieldContext, ResultsView);
+
+			TArray<float> NewWeightArray;
+			NewWeightArray.Init(0.f, NumVertices);
+
+			if (IsConnected<FDataflowVertexSelection>(&VertexSelection))
 			{
-				for (int32 Idx = 0; Idx < NumVertices; ++Idx)
+				const FDataflowVertexSelection& InVertexSelection = GetValue<FDataflowVertexSelection>(Context, &VertexSelection);
+
+				if (InVertexArray.Num() == InVertexSelection.Num())
 				{
-					if (InVertexSelection.IsSelected(Idx))
+					for (int32 Idx = 0; Idx < NumVertices; ++Idx)
 					{
-						NewWeightArray[Idx] = ResultsView[Idx];
+						if (InVertexSelection.IsSelected(Idx))
+						{
+							NewWeightArray[Idx] = ResultsView[Idx];
+						}
 					}
 				}
 			}
-		}
-		else
-		{
-			for (int32 Idx = 0; Idx < NumVertices; ++Idx)
+			else
 			{
-				NewWeightArray[Idx] = ResultsView[Idx];
+				for (int32 Idx = 0; Idx < NumVertices; ++Idx)
+				{
+					NewWeightArray[Idx] = ResultsView[Idx];
+				}
 			}
+
+			SetValue<TArray<float>>(Context, NewWeightArray, &WeightArray);
+
+			return;
 		}
 
-		SetValue<TArray<float>>(Context, NewWeightArray, &WeightArray);
+		SetValue<TArray<float>>(Context, TArray<float>(), &WeightArray);
 	}
 }
 
