Dataflow:
- bug fixes

diff --git a/Engine/Plugins/Experimental/GeometryCollectionPlugin/Source/GeometryCollectionNodes/Private/Dataflow/GeometryCollectionNodes.cpp b/Engine/Plugins/Experimental/GeometryCollectionPlugin/Source/GeometryCollectionNodes/Private/Dataflow/GeometryCollectionNodes.cpp
--- a/Engine/Plugins/Experimental/GeometryCollectionPlugin/Source/GeometryCollectionNodes/Private/Dataflow/GeometryCollectionNodes.cpp
+++ b/Engine/Plugins/Experimental/GeometryCollectionPlugin/Source/GeometryCollectionNodes/Private/Dataflow/GeometryCollectionNodes.cpp
@@ -578,7 +578,10 @@ void FGetArrayElementDataflowNode::Evaluate(Dataflow::FContext& Context, const F
 		if (Array.Num() > 0 && Index >= 0 && Index < Array.Num())
 		{
 			SetValue<FVector>(Context, Array[Index], &Point);
+			return;
 		}
+
+		SetValue<FVector>(Context, FVector(0.f), &Point);
 	}
 }
 
@@ -589,11 +592,15 @@ void FGetNumArrayElementsDataflowNode::Evaluate(Dataflow::FContext& Context, con
 		if (IsConnected<TArray<FVector>>(&Points))
 		{
 			SetValue<int32>(Context, GetValue<TArray<FVector>>(Context, &Points).Num(), &NumElements);
+			return;
 		}
 		else if (IsConnected<TArray<FVector3f>>(&Vector3fArray))
 		{
 			SetValue<int32>(Context, GetValue<TArray<FVector3f>>(Context, &Vector3fArray).Num(), &NumElements);
+			return;
 		}
+
+		SetValue<int32>(Context, 0, &NumElements);
 	}
 }
 
