ChaosFlesh : Skeletal tranform access nodes
- Nodes for the ability to translate and place geometry within a collections skeletal hierarchy.

diff --git a/Engine/Plugins/Experimental/ChaosFlesh/Source/ChaosFleshEngine/Public/ChaosFlesh/ChaosDeformableTypes.h b/Engine/Plugins/Experimental/ChaosFlesh/Source/ChaosFleshEngine/Public/ChaosFlesh/ChaosDeformableTypes.h
--- a/Engine/Plugins/Experimental/ChaosFlesh/Source/ChaosFleshEngine/Public/ChaosFlesh/ChaosDeformableTypes.h
+++ b/Engine/Plugins/Experimental/ChaosFlesh/Source/ChaosFleshEngine/Public/ChaosFlesh/ChaosDeformableTypes.h
@@ -18,4 +18,9 @@ enum class EDeformableExecutionModel : uint8
 struct FChaosEngineDeformableCVarParams
 {
 	bool bEnableDeformableSolver = true;
+#ifdef WITH_ENGINE
+	bool bDoDrawSimulationMesh = true;
+#else
+	bool bDoDrawSimulationMesh = false;
+#endif
 };
\ No newline at end of file
diff --git a/Engine/Plugins/Experimental/ChaosFlesh/Source/ChaosFleshNodes/Private/Dataflow/ChaosFleshRadialTetrahedronNodes.cpp b/Engine/Plugins/Experimental/ChaosFlesh/Source/ChaosFleshNodes/Private/Dataflow/ChaosFleshRadialTetrahedronNodes.cpp
--- a/Engine/Plugins/Experimental/ChaosFlesh/Source/ChaosFleshNodes/Private/Dataflow/ChaosFleshRadialTetrahedronNodes.cpp
+++ b/Engine/Plugins/Experimental/ChaosFlesh/Source/ChaosFleshNodes/Private/Dataflow/ChaosFleshRadialTetrahedronNodes.cpp
@@ -45,7 +45,7 @@ void FRadialTetrahedronDataflowNodes::Evaluate(Dataflow::FContext& Context, cons
 		TArray<FVector> TetVertices;
 		RadialTetMesh(InnerRadius, OuterRadius, Height, RadialSample, AngularSample, VerticalSample, TetElements, TetVertices);
 		TArray<FIntVector3> SurfaceElements = Dataflow::GetSurfaceTriangles(TetElements);;
-		TUniquePtr<FTetrahedralCollection> RadialCollection(FTetrahedralCollection::NewTetrahedralCollection(TetVertices, SurfaceElements, TetElements));
+		TUniquePtr<FTetrahedralCollection> RadialCollection(FTetrahedralCollection::NewTetrahedralCollection(TetVertices, SurfaceElements, TetElements, false));
 		InCollection->AppendGeometry(*RadialCollection.Get());
 
 		SetValue<DataType>(Context, *(FManagedArrayCollection*)InCollection.Get(), &Collection);
diff --git a/Engine/Plugins/Experimental/ChaosFlesh/Source/ChaosFleshNodes/Public/Dataflow/ChaosFleshRadialTetrahedronNodes.h b/Engine/Plugins/Experimental/ChaosFlesh/Source/ChaosFleshNodes/Public/Dataflow/ChaosFleshRadialTetrahedronNodes.h
--- a/Engine/Plugins/Experimental/ChaosFlesh/Source/ChaosFleshNodes/Public/Dataflow/ChaosFleshRadialTetrahedronNodes.h
+++ b/Engine/Plugins/Experimental/ChaosFlesh/Source/ChaosFleshNodes/Public/Dataflow/ChaosFleshRadialTetrahedronNodes.h
@@ -33,13 +33,13 @@ public:
 	UPROPERTY(EditAnywhere, Category = "Dataflow", meta = (ClampMin = "0.0"))
 	double Height = double(1.0);
 
-	UPROPERTY(EditAnywhere, Category = "Dataflow", meta = (ClampMin = "0"))
+	UPROPERTY(EditAnywhere, Category = "Dataflow", meta = (ClampMin = "2"))
 	int32 RadialSample = 2;
 
-	UPROPERTY(EditAnywhere, Category = "Dataflow", meta = (ClampMin = "0"))
+	UPROPERTY(EditAnywhere, Category = "Dataflow", meta = (ClampMin = "3"))
 	int32 AngularSample = 4;
 
-	UPROPERTY(EditAnywhere, Category = "Dataflow", meta = (ClampMin = "0"))
+	UPROPERTY(EditAnywhere, Category = "Dataflow", meta = (ClampMin = "2"))
 	int32 VerticalSample = 2;
 
 	FRadialTetrahedronDataflowNodes(const Dataflow::FNodeParameters& InParam, FGuid InGuid = FGuid::NewGuid())
diff --git a/Engine/Plugins/Experimental/Dataflow/Source/DataflowNodes/Private/Dataflow/DataflowContextOverridesNodes.cpp b/Engine/Plugins/Experimental/Dataflow/Source/DataflowNodes/Private/Dataflow/DataflowContextOverridesNodes.cpp
new file mode 100644
--- /dev/null
+++ b/Engine/Plugins/Experimental/Dataflow/Source/DataflowNodes/Private/Dataflow/DataflowContextOverridesNodes.cpp
@@ -0,0 +1,31 @@
+// Copyright Epic Games, Inc. All Rights Reserved.
+
+#include "Dataflow/DataflowContextOverridesNodes.h"
+
+#include "Dataflow/DataflowEngineUtil.h"
+
+#include UE_INLINE_GENERATED_CPP_BY_NAME(DataflowContextOverridesNodes)
+
+namespace Dataflow
+{
+	void RegisterContextOverridesNodes()
+	{
+		DATAFLOW_NODE_REGISTER_CREATION_FACTORY(FFloatOverrideDataflowNode);
+	}
+}
+
+void FFloatOverrideDataflowNode::Evaluate(Dataflow::FContext& Context, const FDataflowOutput* Out) const
+{
+	if (Out->IsA<float>(&ValueOut))
+	{
+		float Value = 0.f;
+		if (const Dataflow::FEngineContext* EngineContext = Context.AsType<Dataflow::FEngineContext>())
+		{
+			FString Result = Dataflow::Reflection::FindOverrideProperty< FString >(EngineContext->Owner, PropertyName, KeyName);
+			Value = FCString::Atof(*Result);
+		}
+		SetValue<float>(Context, Value, &ValueOut);
+	}
+}
+
+
diff --git a/Engine/Plugins/Experimental/Dataflow/Source/DataflowNodes/Private/Dataflow/DataflowNodesPlugin.cpp b/Engine/Plugins/Experimental/Dataflow/Source/DataflowNodes/Private/Dataflow/DataflowNodesPlugin.cpp
--- a/Engine/Plugins/Experimental/Dataflow/Source/DataflowNodes/Private/Dataflow/DataflowNodesPlugin.cpp
+++ b/Engine/Plugins/Experimental/Dataflow/Source/DataflowNodes/Private/Dataflow/DataflowNodesPlugin.cpp
@@ -7,6 +7,7 @@
 #include "Dataflow/DataflowStaticMeshNodes.h"
 #include "Dataflow/DataflowNodeFactory.h"
 #include "Dataflow/DataflowSelectionNodes.h"
+#include "Dataflow/DataflowContextOverridesNodes.h"
 
 #define LOCTEXT_NAMESPACE "DataflowNodes"
 
@@ -16,6 +17,7 @@ void IDataflowNodesPlugin::StartupModule()
 	Dataflow::RegisterSkeletalMeshNodes();
 	Dataflow::RegisterStaticMeshNodes();
 	Dataflow::RegisterSelectionNodes();
+	Dataflow::RegisterContextOverridesNodes();
 }
 
 void IDataflowNodesPlugin::ShutdownModule()
diff --git a/Engine/Plugins/Experimental/Dataflow/Source/DataflowNodes/Private/Dataflow/DataflowSkeletalMeshNodes.cpp b/Engine/Plugins/Experimental/Dataflow/Source/DataflowNodes/Private/Dataflow/DataflowSkeletalMeshNodes.cpp
--- a/Engine/Plugins/Experimental/Dataflow/Source/DataflowNodes/Private/Dataflow/DataflowSkeletalMeshNodes.cpp
+++ b/Engine/Plugins/Experimental/Dataflow/Source/DataflowNodes/Private/Dataflow/DataflowSkeletalMeshNodes.cpp
@@ -19,6 +19,7 @@ namespace Dataflow
 		DATAFLOW_NODE_REGISTER_CREATION_FACTORY(FGetSkeletalMeshDataflowNode);
 		DATAFLOW_NODE_REGISTER_CREATION_FACTORY(FGetSkeletonDataflowNode);
 		DATAFLOW_NODE_REGISTER_CREATION_FACTORY(FSkeletalMeshBoneDataflowNode);
+		DATAFLOW_NODE_REGISTER_CREATION_FACTORY(FSkeletalMeshReferenceTransformDataflowNode);
 	}
 }
 
@@ -79,7 +80,7 @@ void FSkeletalMeshBoneDataflowNode::Evaluate(Dataflow::FContext& Context, const
 			{
 				if (const Dataflow::FEngineContext* EngineContext = Context.AsType<Dataflow::FEngineContext>())
 				{
-					LocalBoneName = FName(Dataflow::Reflection::FindOverrideProperty< FString >(EngineContext->Owner, PropertyName, FString("BoneName")));
+					LocalBoneName = FName(Dataflow::Reflection::FindOverrideProperty< FString >(EngineContext->Owner, PropertyName, FName("BoneName")));
 				}
 			}
 
@@ -91,4 +92,31 @@ void FSkeletalMeshBoneDataflowNode::Evaluate(Dataflow::FContext& Context, const
 }
 
 
+void FSkeletalMeshReferenceTransformDataflowNode::Evaluate(Dataflow::FContext& Context, const FDataflowOutput* Out) const
+{
+	typedef TObjectPtr<const USkeletalMesh> InDataType;
+	if (Out->IsA<FTransform>(&TransformOut))
+	{
+		SetValue<FTransform>(Context, FTransform::Identity, &TransformOut);
+		
+		int32 BoneIndex = GetValue<int32>(Context, &BoneIndexIn);
+		if (0 <= BoneIndex)
+		{
+			if (InDataType SkeletalMesh = GetValue<InDataType>(Context, &SkeletalMeshIn))
+			{
+				TArray<FTransform> ComponentPose;
+				Dataflow::Animation::GlobalTransforms(SkeletalMesh->GetRefSkeleton(), ComponentPose);
+				if (BoneIndex < ComponentPose.Num())
+				{
+					SetValue<FTransform>(Context, ComponentPose[BoneIndex], &TransformOut);
+
+				}
+			}
+		}
+	}
+}
+
+
+
+
 
diff --git a/Engine/Plugins/Experimental/Dataflow/Source/DataflowNodes/Public/Dataflow/DataflowContextOverridesNodes.h b/Engine/Plugins/Experimental/Dataflow/Source/DataflowNodes/Public/Dataflow/DataflowContextOverridesNodes.h
new file mode 100644
--- /dev/null
+++ b/Engine/Plugins/Experimental/Dataflow/Source/DataflowNodes/Public/Dataflow/DataflowContextOverridesNodes.h
@@ -0,0 +1,41 @@
+// Copyright Epic Games, Inc. All Rights Reserved.
+
+#pragma once 
+
+#include "CoreMinimal.h"
+#include "Dataflow/DataflowCore.h"
+#include "Dataflow/DataflowEngine.h"
+
+#include "DataflowContextOverridesNodes.generated.h"
+
+USTRUCT(meta = (DataflowFlesh))
+struct DATAFLOWNODES_API FFloatOverrideDataflowNode : public FDataflowNode
+{
+	GENERATED_USTRUCT_BODY()
+	DATAFLOW_NODE_DEFINE_INTERNAL(FFloatOverrideDataflowNode, "FloatOverride", "Dataflow", "")
+
+public:
+	UPROPERTY(EditAnywhere, Category = "Overrides")
+	FName PropertyName = "Overrides";
+
+	UPROPERTY(EditAnywhere, Category = "Overrides")
+	FName KeyName = "FloatAttr";
+
+	UPROPERTY(meta = (DataflowOutput, DisplayName = "Float"))
+	float ValueOut = 0.f;
+
+	FFloatOverrideDataflowNode(const Dataflow::FNodeParameters& InParam, FGuid InGuid = FGuid::NewGuid())
+		: FDataflowNode(InParam, InGuid)
+	{
+		RegisterOutputConnection(&ValueOut);
+	}
+
+	virtual void Evaluate(Dataflow::FContext& Context, const FDataflowOutput* Out) const override;
+};
+
+
+namespace Dataflow
+{
+	void RegisterContextOverridesNodes();
+}
+
diff --git a/Engine/Plugins/Experimental/Dataflow/Source/DataflowNodes/Public/Dataflow/DataflowSkeletalMeshNodes.h b/Engine/Plugins/Experimental/Dataflow/Source/DataflowNodes/Public/Dataflow/DataflowSkeletalMeshNodes.h
--- a/Engine/Plugins/Experimental/Dataflow/Source/DataflowNodes/Public/Dataflow/DataflowSkeletalMeshNodes.h
+++ b/Engine/Plugins/Experimental/Dataflow/Source/DataflowNodes/Public/Dataflow/DataflowSkeletalMeshNodes.h
@@ -91,6 +91,36 @@ public:
 	virtual void Evaluate(Dataflow::FContext& Context, const FDataflowOutput* Out) const override;
 };
 
+
+USTRUCT()
+struct DATAFLOWNODES_API FSkeletalMeshReferenceTransformDataflowNode : public FDataflowNode
+{
+	GENERATED_USTRUCT_BODY()
+	DATAFLOW_NODE_DEFINE_INTERNAL(FSkeletalMeshReferenceTransformDataflowNode, "SkeletalMeshReferenceTransform", "General", "Skeletal Mesh")
+
+public:
+
+	UPROPERTY(meta = (DataflowInput, DisplayName = "SkeletalMesh"))
+	TObjectPtr<const USkeletalMesh> SkeletalMeshIn = nullptr;
+
+	UPROPERTY(meta = (DataflowInput, DisplayName = "Index"))
+	int32 BoneIndexIn = INDEX_NONE;
+
+	UPROPERTY(meta = (DataflowOutput, DisplayName = "Transform"))
+	FTransform TransformOut = FTransform::Identity;
+
+	FSkeletalMeshReferenceTransformDataflowNode(const Dataflow::FNodeParameters& InParam, FGuid InGuid = FGuid::NewGuid())
+		: FDataflowNode(InParam, InGuid)
+	{
+		RegisterInputConnection(&SkeletalMeshIn);
+		RegisterInputConnection(&BoneIndexIn);
+		RegisterOutputConnection(&TransformOut);
+	}
+
+
+	virtual void Evaluate(Dataflow::FContext& Context, const FDataflowOutput* Out) const override;
+};
+
 namespace Dataflow
 {
 	void RegisterSkeletalMeshNodes();
diff --git a/Engine/Plugins/Experimental/Dataflow/Source/DataflowNodes/Public/Dataflow/DataflowStaticMeshNodes.h b/Engine/Plugins/Experimental/Dataflow/Source/DataflowNodes/Public/Dataflow/DataflowStaticMeshNodes.h
--- a/Engine/Plugins/Experimental/Dataflow/Source/DataflowNodes/Public/Dataflow/DataflowStaticMeshNodes.h
+++ b/Engine/Plugins/Experimental/Dataflow/Source/DataflowNodes/Public/Dataflow/DataflowStaticMeshNodes.h
@@ -23,7 +23,7 @@ public:
 	UPROPERTY(EditAnywhere, Category = "Dataflow", meta = (DataflowOutput, DisplayName = "StaticMesh"))
 	TObjectPtr<const UStaticMesh> StaticMesh = nullptr;
 
-	UPROPERTY(EditAnywhere, Category = "Dataflow")
+	UPROPERTY(EditAnywhere, Category = "Overrides")
 	FName PropertyName = "StaticMesh";
 
 	FGetStaticMeshDataflowNode(const Dataflow::FNodeParameters& InParam, FGuid InGuid = FGuid::NewGuid())
diff --git a/Engine/Plugins/Experimental/GeometryCollectionPlugin/Source/GeometryCollectionNodes/Private/Dataflow/GeometryCollectionNodes.cpp b/Engine/Plugins/Experimental/GeometryCollectionPlugin/Source/GeometryCollectionNodes/Private/Dataflow/GeometryCollectionNodes.cpp
--- a/Engine/Plugins/Experimental/GeometryCollectionPlugin/Source/GeometryCollectionNodes/Private/Dataflow/GeometryCollectionNodes.cpp
+++ b/Engine/Plugins/Experimental/GeometryCollectionPlugin/Source/GeometryCollectionNodes/Private/Dataflow/GeometryCollectionNodes.cpp
@@ -99,6 +99,8 @@ namespace Dataflow
 		DATAFLOW_NODE_REGISTER_CREATION_FACTORY(FFloatArrayToVertexSelectionDataflowNode);
 		DATAFLOW_NODE_REGISTER_CREATION_FACTORY(FSetVertexColorInCollectionDataflowNode);
 		DATAFLOW_NODE_REGISTER_CREATION_FACTORY(FMakeTransformDataflowNode);
+		DATAFLOW_NODE_REGISTER_CREATION_FACTORY(FMultiplyTransformDataflowNode);
+		
 
 		// GeometryCollection
 		DATAFLOW_NODE_REGISTER_CREATION_FACTORY_NODE_COLORS_BY_CATEGORY("GeometryCollection", FLinearColor(0.55f, 0.45f, 1.0f), CDefaultNodeBodyTintColor);
@@ -313,6 +315,7 @@ void FMakeLiteralVectorDataflowNode::Evaluate(Dataflow::FContext& Context, const
 {
 	if (Out->IsA<FVector>(&Vector))
 	{
+		FVector Value(GetValue<float>(Context, &X, X), GetValue<float>(Context, &Y, Y), GetValue<float>(Context, &Z, Z));
 		SetValue<FVector>(Context, Value, &Vector);
 	}
 }
@@ -1424,11 +1427,28 @@ void FSetVertexColorInCollectionDataflowNode::Evaluate(Dataflow::FContext& Conte
 
 void FMakeTransformDataflowNode::Evaluate(Dataflow::FContext& Context, const FDataflowOutput* Out) const
 {
-	if (Out->IsA<FTransform>(&Transform))
+	if (Out->IsA<FTransform>(&OutTransform))
+	{
+		SetValue<FTransform>(Context, 
+							 FTransform(FQuat::MakeFromEuler(GetValue<FVector>(Context, &InRotation))
+							 , GetValue<FVector>(Context, &InTranslation)
+							 , GetValue<FVector>(Context, &InScale))
+							 ,&OutTransform);
+	}
+}
+
+void FMultiplyTransformDataflowNode::Evaluate(Dataflow::FContext& Context, const FDataflowOutput* Out) const
+{
+	if (Out->IsA<FTransform>(&OutTransform))
 	{
-		SetValue<FTransform>(Context, FTransform(InTransform), &Transform);
+		SetValue<FTransform>(Context,
+			GetValue<FTransform>(Context, &InLeftTransform, FTransform::Identity)
+			*GetValue<FTransform>(Context, &InRightTransform, FTransform::Identity)
+			, &OutTransform);
 	}
 }
 
 
 
+
+
diff --git a/Engine/Plugins/Experimental/GeometryCollectionPlugin/Source/GeometryCollectionNodes/Private/Dataflow/GeometryCollectionNodesPlugin.cpp b/Engine/Plugins/Experimental/GeometryCollectionPlugin/Source/GeometryCollectionNodes/Private/Dataflow/GeometryCollectionNodesPlugin.cpp
--- a/Engine/Plugins/Experimental/GeometryCollectionPlugin/Source/GeometryCollectionNodes/Private/Dataflow/GeometryCollectionNodesPlugin.cpp
+++ b/Engine/Plugins/Experimental/GeometryCollectionPlugin/Source/GeometryCollectionNodes/Private/Dataflow/GeometryCollectionNodesPlugin.cpp
@@ -17,6 +17,7 @@
 #include "Dataflow/GeometryCollectionUtilityNodes.h"
 #include "Dataflow/GeometryCollectionMaterialNodes.h"
 #include "Dataflow/GeometryCollectionFieldNodes.h"
+#include "Dataflow/GeometryCollectionVerticesNodes.h"
 
 #define LOCTEXT_NAMESPACE "DataflowNodes"
 
@@ -35,6 +36,7 @@ void IGeometryCollectionNodesPlugin::StartupModule()
 	Dataflow::GeometryCollectionUtilityNodes();
 	Dataflow::GeometryCollectionMaterialNodes();
 	Dataflow::GeometryCollectionFieldNodes();
+	Dataflow::GeometryCollectionVerticesNodes();
 }
 
 void IGeometryCollectionNodesPlugin::ShutdownModule()
diff --git a/Engine/Plugins/Experimental/GeometryCollectionPlugin/Source/GeometryCollectionNodes/Private/Dataflow/GeometryCollectionVerticesNodes.cpp b/Engine/Plugins/Experimental/GeometryCollectionPlugin/Source/GeometryCollectionNodes/Private/Dataflow/GeometryCollectionVerticesNodes.cpp
new file mode 100644
--- /dev/null
+++ b/Engine/Plugins/Experimental/GeometryCollectionPlugin/Source/GeometryCollectionNodes/Private/Dataflow/GeometryCollectionVerticesNodes.cpp
@@ -0,0 +1,43 @@
+// Copyright Epic Games, Inc. All Rights Reserved.
+
+#include "Dataflow/GeometryCollectionVerticesNodes.h"
+#include "GeometryCollection/GeometryCollection.h"
+
+#include "Dataflow/DataflowCore.h"
+
+#include UE_INLINE_GENERATED_CPP_BY_NAME(GeometryCollectionVerticesNodes)
+
+namespace Dataflow
+{
+	void GeometryCollectionVerticesNodes()
+	{
+		DATAFLOW_NODE_REGISTER_CREATION_FACTORY(FTransformCollectionAttributeDataflowNode);
+	}
+}
+
+
+void FTransformCollectionAttributeDataflowNode::Evaluate(Dataflow::FContext& Context, const FDataflowOutput* Out) const
+{
+	if (Out->IsA<DataType>(&Collection))
+	{
+		const FTransform& Transform = GetValue<FTransform>(Context, &TransformIn);
+		FManagedArrayCollection CollectionValue = GetValue<FManagedArrayCollection>(Context, &Collection);
+		if (CollectionValue.FindAttributeTyped<FVector3f>("Vertex", FGeometryCollection::VerticesGroup))
+		{
+			auto UEVertd = [](FVector3f V) { return FVector3d(V.X, V.Y, V.Z); };
+			auto UEVertf = [](FVector3d V) { return FVector3f(V.X, V.Y, V.Z); };
+
+			FTransform FinalTransform = LocalTransform * Transform;
+			//FMatrix44d Matrix = (LocalTransform * Transform).ToMatrixWithScale();
+			TManagedArray<FVector3f>& Vertices = CollectionValue.ModifyAttribute<FVector3f>("Vertex", FGeometryCollection::VerticesGroup);
+			for (int32 i = 0; i < Vertices.Num(); i++)
+			{
+				//Vertices[i] = UEVertf(Matrix.TransformPosition(UEVertd(Vertices[i])));
+				Vertices[i] = UEVertf(FinalTransform.TransformPosition(UEVertd(Vertices[i])));
+			}
+		}
+		SetValue<DataType>(Context, CollectionValue, &Collection);
+
+	}
+}
+
diff --git a/Engine/Plugins/Experimental/GeometryCollectionPlugin/Source/GeometryCollectionNodes/Public/Dataflow/GeometryCollectionNodes.h b/Engine/Plugins/Experimental/GeometryCollectionPlugin/Source/GeometryCollectionNodes/Public/Dataflow/GeometryCollectionNodes.h
--- a/Engine/Plugins/Experimental/GeometryCollectionPlugin/Source/GeometryCollectionNodes/Public/Dataflow/GeometryCollectionNodes.h
+++ b/Engine/Plugins/Experimental/GeometryCollectionPlugin/Source/GeometryCollectionNodes/Public/Dataflow/GeometryCollectionNodes.h
@@ -534,8 +534,14 @@ struct FMakeLiteralVectorDataflowNode : public FDataflowNode
 	DATAFLOW_NODE_DEFINE_INTERNAL(FMakeLiteralVectorDataflowNode, "MakeLiteralVector", "Math|Vector", "")
 
 public:
-	UPROPERTY(EditAnywhere, Category = "Vector");
-	FVector Value = FVector(0.0);
+	UPROPERTY(EditAnywhere, Category = "Vector", meta = (DataflowInput));
+	float X = float(0.0);
+
+	UPROPERTY(EditAnywhere, Category = "Vector", meta = (DataflowInput));
+	float Y = float(0.0);
+
+	UPROPERTY(EditAnywhere, Category = "Vector", meta = (DataflowInput));
+	float Z = float(0.0);
 
 	UPROPERTY(meta = (DataflowOutput, DisplayName = "Vector"))
 	FVector Vector = FVector(0.0);
@@ -543,6 +549,9 @@ public:
 	FMakeLiteralVectorDataflowNode(const Dataflow::FNodeParameters& InParam, FGuid InGuid = FGuid::NewGuid())
 		: FDataflowNode(InParam, InGuid)
 	{
+		RegisterInputConnection(&X);
+		RegisterInputConnection(&Y);
+		RegisterInputConnection(&Z);
 		RegisterOutputConnection(&Vector);
 	}
 
@@ -2098,22 +2107,64 @@ struct FMakeTransformDataflowNode : public FDataflowNode
 	DATAFLOW_NODE_DEFINE_INTERNAL(FMakeTransformDataflowNode, "MakeTransform", "Generators|Transform", "")
 
 public:
-	UPROPERTY(EditAnywhere, Category = "Transform", meta = (DisplayName = "Transform"));
-	FTransform InTransform = FTransform::Identity;
 
-	UPROPERTY(meta = (DataflowOutput));
-	FTransform Transform = FTransform::Identity;
+	UPROPERTY(EditAnywhere, Category = "Transform", meta = (DataflowInput, DisplayName = "Translation"));
+	FVector InTranslation = FVector(0,0,0);
+
+	UPROPERTY(EditAnywhere, Category = "Transform", meta = (DataflowInput, DisplayName = "Rotation"));
+	FVector InRotation = FVector(0, 0, 0);
+
+	UPROPERTY(EditAnywhere, Category = "Transform", meta = (DataflowInput, DisplayName = "Scale"));
+	FVector InScale = FVector(1,1,1);
+
+	UPROPERTY(meta = (DataflowOutput, DisplayName = "Transform"));
+	FTransform OutTransform = FTransform::Identity;
 
 	FMakeTransformDataflowNode(const Dataflow::FNodeParameters& InParam, FGuid InGuid = FGuid::NewGuid())
 		: FDataflowNode(InParam, InGuid)
 	{
-		RegisterOutputConnection(&Transform);
+		RegisterInputConnection(&InTranslation);
+		RegisterInputConnection(&InRotation);
+		RegisterInputConnection(&InScale);
+		RegisterOutputConnection(&OutTransform);
 	}
 
 	virtual void Evaluate(Dataflow::FContext& Context, const FDataflowOutput* Out) const override;
 
 };
 
+/**
+ *
+ * Description for this node
+ *
+ */
+USTRUCT()
+struct FMultiplyTransformDataflowNode : public FDataflowNode
+{
+	GENERATED_USTRUCT_BODY()
+	DATAFLOW_NODE_DEFINE_INTERNAL(FMultiplyTransformDataflowNode, "MultiplyTransform", "Math|Transform", "")
+
+public:
+	UPROPERTY(meta = (DataflowInput, DisplayName = "Left Transform"));
+	FTransform InLeftTransform = FTransform::Identity;
+
+	UPROPERTY(meta = (DataflowInput, DisplayName = "Right Transform"));
+	FTransform InRightTransform = FTransform::Identity;
+
+	UPROPERTY(meta = ( DataflowOutput, DisplayName = "Out Transform"));
+	FTransform OutTransform = FTransform::Identity;
+
+	FMultiplyTransformDataflowNode(const Dataflow::FNodeParameters& InParam, FGuid InGuid = FGuid::NewGuid())
+		: FDataflowNode(InParam, InGuid)
+	{
+		RegisterInputConnection(&InLeftTransform);
+		RegisterInputConnection(&InRightTransform);
+		RegisterOutputConnection(&OutTransform);
+	}
+
+	virtual void Evaluate(Dataflow::FContext& Context, const FDataflowOutput* Out) const override;
+
+};
 
 
 namespace Dataflow
diff --git a/Engine/Plugins/Experimental/GeometryCollectionPlugin/Source/GeometryCollectionNodes/Public/Dataflow/GeometryCollectionVerticesNodes.h b/Engine/Plugins/Experimental/GeometryCollectionPlugin/Source/GeometryCollectionNodes/Public/Dataflow/GeometryCollectionVerticesNodes.h
new file mode 100644
--- /dev/null
+++ b/Engine/Plugins/Experimental/GeometryCollectionPlugin/Source/GeometryCollectionNodes/Public/Dataflow/GeometryCollectionVerticesNodes.h
@@ -0,0 +1,55 @@
+// Copyright Epic Games, Inc. All Rights Reserved.
+
+#pragma once 
+
+#include "CoreMinimal.h"
+#include "Dataflow/DataflowEngine.h"
+#include "GeometryCollection/ManagedArrayCollection.h"
+
+#include "GeometryCollectionVerticesNodes.generated.h"
+
+class USkeletalMesh;
+
+USTRUCT(meta = (DataflowGeometryCollection))
+struct FTransformCollectionAttributeDataflowNode : public FDataflowNode
+{
+	GENERATED_USTRUCT_BODY()
+	DATAFLOW_NODE_DEFINE_INTERNAL(FTransformCollectionAttributeDataflowNode, "TransformCollectionAttribute", "GeometryCollection", "")
+	DATAFLOW_NODE_RENDER_TYPE(FGeometryCollection::StaticType(), "Collection")
+
+public:
+	typedef FManagedArrayCollection DataType;
+
+
+	UPROPERTY(meta = (DataflowInput, DataflowOutput, DisplayName = "Collection"))
+	FManagedArrayCollection Collection;
+
+	UPROPERTY(meta = (DataflowInput, DisplayName = "Transform"))
+	FTransform TransformIn = FTransform::Identity;
+
+	UPROPERTY(EditAnywhere, Category = "Dataflow")
+	FTransform LocalTransform = FTransform::Identity;
+
+	UPROPERTY(EditAnywhere, Category = "Dataflow")
+	FString GroupName = FString("Vertices");
+
+	UPROPERTY(EditAnywhere, Category = "Dataflow")
+	FString AttributeName = FString("Vertex");
+
+
+	FTransformCollectionAttributeDataflowNode(const Dataflow::FNodeParameters& InParam, FGuid InGuid = FGuid::NewGuid())
+		: FDataflowNode(InParam, InGuid)
+	{
+		RegisterInputConnection(&TransformIn);
+		RegisterInputConnection(&Collection);
+		RegisterOutputConnection(&Collection);
+	}
+
+	virtual void Evaluate(Dataflow::FContext& Context, const FDataflowOutput* Out) const override;
+};
+
+namespace Dataflow
+{
+	void GeometryCollectionVerticesNodes();
+}
+
diff --git a/Engine/Source/Runtime/Experimental/Dataflow/Engine/Public/Dataflow/DataflowEngineUtil.h b/Engine/Source/Runtime/Experimental/Dataflow/Engine/Public/Dataflow/DataflowEngineUtil.h
--- a/Engine/Source/Runtime/Experimental/Dataflow/Engine/Public/Dataflow/DataflowEngineUtil.h
+++ b/Engine/Source/Runtime/Experimental/Dataflow/Engine/Public/Dataflow/DataflowEngineUtil.h
@@ -42,7 +42,7 @@ namespace Dataflow
 		}
 
 		template<class T>
-		const T FindOverrideProperty(const UObject* Owner, FName PropertyName, FString ArrayKey, const T & Default = T())
+		const T FindOverrideProperty(const UObject* Owner, FName PropertyName, FName ArrayKey, const T & Default = T())
 		{
 			if (Owner && Owner->GetClass())
 			{
@@ -56,7 +56,7 @@ namespace Dataflow
 							{
 								for (const FStringValuePair& Pair : *ContainerData)
 								{
-									if (Pair.Key.Equals(ArrayKey))
+									if (Pair.Key.Equals(ArrayKey.ToString()))
 									{
 										return Pair.Value;
 									}
