Dataflow:
- Updated field nodes (FRadialFalloffFieldDataflowNode, FBoxFalloffFieldDataflowNode, FPlaneFalloffFieldDataflowNode, FRadialIntMaskFieldDataflowNode, FUniformScalarFieldDataflowNode, FUniformVectorFieldDataflowNode, FRadialVectorFieldDataflowNode,
- Added FSetVertexColorInCollectionFromFloatArrayDataflowNode, FFloatArrayNormalizeDataflowNode, FVectorArrayNormalizeDataflowNode, FMakeQuaternionDataflowNode, FFieldMakeDenseFloatArrayDataflowNode nodes
FRandomVectorFieldDataflowNode, FNoiseFieldDataflowNode, FUniformIntegerFieldDataflowNode, FWaveScalarFieldDataflowNode, FSumScalarFieldDataflowNode, FSumVectorFieldDataflowNode)
-Renamed FSetVertexColorInCollectionDataflowNode to FSetVertexColorInCollectionFromVertexSelectionDataflowNode
- Assigned color to FSphere data type
- Updated FAddCustomCollectionAttributeDataflowNode, FAddCustomCollectionAttributeDataflowNode, FSetCollectionAttributeDataTypedDataflowNode nodes to work with all the non-pointer type data

diff --git a/Engine/Plugins/Experimental/Dataflow/Source/DataflowEditor/Private/Dataflow/DataflowSchema.cpp b/Engine/Plugins/Experimental/Dataflow/Source/DataflowEditor/Private/Dataflow/DataflowSchema.cpp
--- a/Engine/Plugins/Experimental/Dataflow/Source/DataflowEditor/Private/Dataflow/DataflowSchema.cpp
+++ b/Engine/Plugins/Experimental/Dataflow/Source/DataflowEditor/Private/Dataflow/DataflowSchema.cpp
@@ -178,6 +178,10 @@ FLinearColor UDataflowSchema::GetTypeColor(const FName& Type)
 	{
 		return DataflowSettings->BoxPinTypeColor;
 	}
+	else if (Type == FName("FSphere"))
+	{
+		return DataflowSettings->SpherePinTypeColor;
+	}
 
 	return Settings->DefaultPinTypeColor;
 }
diff --git a/Engine/Plugins/Experimental/GeometryCollectionPlugin/Source/GeometryCollectionNodes/Private/Dataflow/GeometryCollectionFieldNodes.cpp b/Engine/Plugins/Experimental/GeometryCollectionPlugin/Source/GeometryCollectionNodes/Private/Dataflow/GeometryCollectionFieldNodes.cpp
--- a/Engine/Plugins/Experimental/GeometryCollectionPlugin/Source/GeometryCollectionNodes/Private/Dataflow/GeometryCollectionFieldNodes.cpp
+++ b/Engine/Plugins/Experimental/GeometryCollectionPlugin/Source/GeometryCollectionNodes/Private/Dataflow/GeometryCollectionFieldNodes.cpp
@@ -54,6 +54,7 @@ namespace Dataflow
 		static const FLinearColor CDefaultNodeBodyTintColor = FLinearColor(0.f, 0.f, 0.f, 0.5f);
 
 		DATAFLOW_NODE_REGISTER_CREATION_FACTORY(FRadialFalloffFieldDataflowNode);
+		DATAFLOW_NODE_REGISTER_CREATION_FACTORY(FBoxFalloffFieldDataflowNode);
 		DATAFLOW_NODE_REGISTER_CREATION_FACTORY(FPlaneFalloffFieldDataflowNode);
 		DATAFLOW_NODE_REGISTER_CREATION_FACTORY(FRadialIntMaskFieldDataflowNode);
 		DATAFLOW_NODE_REGISTER_CREATION_FACTORY(FUniformScalarFieldDataflowNode);
@@ -62,820 +63,2430 @@ namespace Dataflow
 		DATAFLOW_NODE_REGISTER_CREATION_FACTORY(FRandomVectorFieldDataflowNode);
 		DATAFLOW_NODE_REGISTER_CREATION_FACTORY(FNoiseFieldDataflowNode);
 		DATAFLOW_NODE_REGISTER_CREATION_FACTORY(FUniformIntegerFieldDataflowNode);
-		DATAFLOW_NODE_REGISTER_CREATION_FACTORY(FBoxFalloffFieldDataflowNode);
+		DATAFLOW_NODE_REGISTER_CREATION_FACTORY(FWaveScalarFieldDataflowNode);
+		DATAFLOW_NODE_REGISTER_CREATION_FACTORY(FSumScalarFieldDataflowNode);
+		DATAFLOW_NODE_REGISTER_CREATION_FACTORY(FSumVectorFieldDataflowNode);
+		DATAFLOW_NODE_REGISTER_CREATION_FACTORY(FFieldMakeDenseFloatArrayDataflowNode);
 
 		// Field nodes
 		DATAFLOW_NODE_REGISTER_CREATION_FACTORY_NODE_COLORS_BY_CATEGORY("Fields", FLinearColor(.0f, 0.8f, 1.f), CDefaultNodeBodyTintColor);
 	}
 }
 
+// --------------------------------------------------------------------------------------------------------------------
+
+static void FieldComputeRemap(const TArray<FVector3f>& InSamplePositions,
+	const FDataflowVertexSelection& InSampleIndices,
+	TArray<FVector3f>& OutNewSamplePositions,
+	TArray<int32>& OutRemapArray)
+{
+	const int32 NumInSamplePositions = InSamplePositions.Num();
+
+	//
+	// Build the SamplePositions array with only the selected indices
+	// and build the remap array
+	//
+	OutNewSamplePositions.Init(FVector3f(0.f), NumInSamplePositions);
+	OutRemapArray.Init(0, NumInSamplePositions);
+
+	int32 NewIdx = 0;
+	for (int32 Idx = 0; Idx < NumInSamplePositions; ++Idx)
+	{
+		if (InSampleIndices.IsSelected(Idx))
+		{
+			OutNewSamplePositions[NewIdx] = InSamplePositions[Idx];
+			OutRemapArray[NewIdx] = Idx;
+
+			NewIdx++;
+		}
+	}
+
+	OutNewSamplePositions.SetNum(NewIdx, true);
+	OutRemapArray.SetNum(NewIdx, true);
+}
+
+// --------------------------------------------------------------------------------------------------------------------
+
+static void RadialFalloffFieldEvaluate(const TArray<FVector3f>& InSamplePositions,
+	const FSphere& InSphere,
+	const FVector& InTranslation,
+	const float InMagnitude,
+	const float InMinRange,
+	const float InMaxRange,
+	const float InDefault,
+	const EDataflowFieldFalloffType InFalloffType,
+	TFieldArrayView<float>& OutResultsView)
+{
+	const int32 NumInSamplePositions = InSamplePositions.Num();
+
+	// Setup the data for the field 
+	FFieldExecutionDatas ExecutionDatas;
+
+	FFieldContextIndex::ContiguousIndices(ExecutionDatas.SampleIndices, NumInSamplePositions);
+	ExecutionDatas.SamplePositions.Init(FVector(0.f), NumInSamplePositions);
+
+	for (int32 Idx = 0; Idx < NumInSamplePositions; Idx++)
+	{
+		ExecutionDatas.SamplePositions[Idx] = (FVector)InSamplePositions[Idx];
+	}
+
+	FRadialFalloff RadialFalloffField(InMagnitude, 
+		InMinRange, 
+		InMaxRange, 
+		InDefault, 
+		InSphere.W, 
+		InSphere.Center + InTranslation, 
+		(EFieldFalloffType)InFalloffType);
+
+	FFieldContext FieldContext{
+		ExecutionDatas,
+		FFieldContext::UniquePointerMap(),
+		0.0
+	};
+
+	// Evalute the field
+	RadialFalloffField.Evaluate(FieldContext, OutResultsView);
+}
+
+static void RadialFalloffFieldProcess(const TArray<FVector3f>& InSamplePositions,
+	const FSphere& InSphere,
+	const FVector& InTranslation,
+	const float InMagnitude,
+	const float InMinRange,
+	const float InMaxRange,
+	const float InDefault,
+	const EDataflowFieldFalloffType InFalloffType,
+	FDataflowVertexSelection& OutFieldSelectionMask,
+	TArray<float>& OutFieldFloatOutput)
+{
+	const int32 NumInSamplePositions = InSamplePositions.Num();
+
+	TArray<float> ResultsArray;
+	ResultsArray.Init(false, NumInSamplePositions);
+	TFieldArrayView<float> ResultsViewTest(ResultsArray, 0, ResultsArray.Num());
+
+	//
+	// Get the samples which effected by the field by evaluating the field with a default 
+	// value outside of the (MinRange, MaxRange)
+	//
+	const float NewDefault = InMinRange - 1000.f;
+
+	RadialFalloffFieldEvaluate(InSamplePositions,
+		InSphere,
+		InTranslation,
+		InMagnitude,
+		InMinRange,
+		InMaxRange,
+		NewDefault,
+		InFalloffType,
+		ResultsViewTest);
+
+	for (int32 Idx = 0; Idx < NumInSamplePositions; ++Idx)
+	{
+		if (ResultsViewTest[Idx] != NewDefault)
+		{
+			OutFieldSelectionMask.SetSelected(Idx);
+		}
+	}
+
+	// Compute the field
+	ResultsArray.Init(false, NumInSamplePositions);
+	TFieldArrayView<float> ResultsView(ResultsArray, 0, ResultsArray.Num());
+
+	RadialFalloffFieldEvaluate(InSamplePositions,
+		InSphere,
+		InTranslation,
+		InMagnitude,
+		InMinRange,
+		InMaxRange,
+		InDefault,
+		InFalloffType,
+		ResultsView);
+
+	for (int32 Idx = 0; Idx < NumInSamplePositions; ++Idx)
+	{
+		OutFieldFloatOutput[Idx] = ResultsView[Idx];
+	}
+}
 
 void FRadialFalloffFieldDataflowNode::Evaluate(Dataflow::FContext& Context, const FDataflowOutput* Out) const
 {
-	if (Out->IsA<TArray<float>>(&WeightArray))
+	if (Out->IsA<TArray<float>>(&FieldFloatResult) ||
+		Out->IsA<TArray<int32>>(&FieldRemap) ||
+		Out->IsA<FDataflowVertexSelection>(&FieldSelectionMask) ||
+		Out->IsA<int32>(&NumSamplePositions))
 	{
-		const TArray<FVector3f>& InVertexArray = GetValue<TArray<FVector3f>>(Context, &VertexArray);
+		const TArray<FVector3f>& InSamplePositions = GetValue<TArray<FVector3f>>(Context, &SamplePositions);
+		const FDataflowVertexSelection& InSampleIndices = GetValue<FDataflowVertexSelection>(Context, &SampleIndices);
 		const FSphere& InSphere = GetValue<FSphere>(Context, &Sphere);
 		const FVector& InTranslation = GetValue<FVector>(Context, &Translation);
+		const float InMagnitude = GetValue<float>(Context, &Magnitude);
+		const float InMinRange = GetValue<float>(Context, &MinRange);
+		const float InMaxRange = GetValue<float>(Context, &MaxRange);
+		const float InDefault = GetValue<float>(Context, &Default);
+
+		//
+		// SampleIndices input not connected, all the elements from SamplePositions array
+		// will be processed and FieldRemap output will be set empty
+		//
+		if (!IsConnected<FDataflowVertexSelection>(&SampleIndices))
+		{
+			const int32 NumInSamplePositions = InSamplePositions.Num();
+
+			//
+			// Process the field
+			//
+			FDataflowVertexSelection NewFieldSelectionMask;
+			NewFieldSelectionMask.Initialize(NumInSamplePositions, false);
+
+			TArray<float> NewFieldFloatResult;
+			NewFieldFloatResult.Init(false, NumInSamplePositions);
+
+			RadialFalloffFieldProcess(InSamplePositions,
+				InSphere,
+				InTranslation,
+				InMagnitude,
+				InMinRange,
+				InMaxRange,
+				InDefault,
+				FalloffType,
+				NewFieldSelectionMask,
+				NewFieldFloatResult);
+
+			// Set the outputs
+			SetValue<TArray<float>>(Context, NewFieldFloatResult, &FieldFloatResult);
+			SetValue<TArray<int32>>(Context, TArray<int32>(), &FieldRemap);
+			SetValue<FDataflowVertexSelection>(Context, NewFieldSelectionMask, &FieldSelectionMask);
+			SetValue<int32>(Context, InSamplePositions.Num(), &NumSamplePositions);
+
+			return;
+		}
+		//
+		// SampleIndices input connected, the selected elements from SamplePositions array
+		// will be processed and FieldRemap output will contain the remap info
+		// IMPORTANT: Number of elements in SamplePositions and SampleIndices must be the same
+		//
+		else
+		{
+			if (InSamplePositions.Num() == InSampleIndices.Num())
+			{
+				TArray<FVector3f> NewSamplePositions;
+				TArray<int32> NewRemapArray;
+
+				FieldComputeRemap(InSamplePositions, InSampleIndices, NewSamplePositions, NewRemapArray);
+
+				const int32 NumNewSamplePositions = NewSamplePositions.Num();
+
+				//
+				// Process the field
+				//
+				FDataflowVertexSelection NewFieldSelectionMask;
+				NewFieldSelectionMask.Initialize(NumNewSamplePositions, false);
+
+				TArray<float> NewFieldFloatResult;
+				NewFieldFloatResult.Init(false, NumNewSamplePositions);
+
+				RadialFalloffFieldProcess(NewSamplePositions,
+					InSphere,
+					InTranslation,
+					InMagnitude,
+					InMinRange,
+					InMaxRange,
+					InDefault,
+					FalloffType,
+					NewFieldSelectionMask,
+					NewFieldFloatResult);
+
+				// Set the outputs
+				SetValue<TArray<float>>(Context, NewFieldFloatResult, &FieldFloatResult);
+				SetValue<TArray<int32>>(Context, NewRemapArray, &FieldRemap);
+				SetValue<FDataflowVertexSelection>(Context, NewFieldSelectionMask, &FieldSelectionMask);
+				SetValue<int32>(Context, InSamplePositions.Num(), &NumSamplePositions);
+
+				return;
+			}
+		}
+
+		// Set the outputs
+		SetValue<TArray<float>>(Context, TArray<float>(), &FieldFloatResult);
+		SetValue<TArray<int32>>(Context, TArray<int32>(), &FieldRemap);
+		SetValue<FDataflowVertexSelection>(Context, FDataflowVertexSelection(), &FieldSelectionMask);
+		SetValue<int32>(Context, 0, &NumSamplePositions);
+	}
+}
+
+
+// --------------------------------------------------------------------------------------------------------------------
+
+static void BoxFalloffFieldEvaluate(const TArray<FVector3f>& InSamplePositions,
+	const FBox& InBox,
+	FTransform InTransform,
+	const float InMagnitude,
+	const float InMinRange,
+	const float InMaxRange,
+	const float InDefault,
+	const EDataflowFieldFalloffType InFalloffType,
+	TFieldArrayView<float>& OutResultsView)
+{
+	const int32 NumInSamplePositions = InSamplePositions.Num();
+
+	// Setup the data for the field 
+	FFieldExecutionDatas ExecutionDatas;
+
+	FFieldContextIndex::ContiguousIndices(ExecutionDatas.SampleIndices, NumInSamplePositions);
+	ExecutionDatas.SamplePositions.Init(FVector(0.f), NumInSamplePositions);
+
+	for (int32 Idx = 0; Idx < NumInSamplePositions; Idx++)
+	{
+		ExecutionDatas.SamplePositions[Idx] = (FVector)InSamplePositions[Idx];
+	}
+
+	//
+	// Translation and Scale needs to be adjusted based on the InBox
+	//
+	FVector Center = InBox.GetCenter();
+	FVector Translation = InTransform.GetTranslation();
+	Translation -= Center;
 
-		const int32 NumVertices = InVertexArray.Num();
+	InTransform.SetTranslation(Translation);
 
-		FFieldExecutionDatas ExecutionDatas;
+	FVector Extent(InBox.Max - InBox.Min);
+	FVector Scale = InTransform.GetScale3D();
 
-		FFieldContextIndex::ContiguousIndices(ExecutionDatas.SampleIndices, NumVertices);
-		ExecutionDatas.SamplePositions.Init(FVector(0.f), NumVertices);
+	Scale.X *= 100.f / Extent.X;
+	Scale.Y *= 100.f / Extent.Y;
+	Scale.Z *= 100.f / Extent.Z;
 
-		for (int32 Index = 0; Index < NumVertices; Index++)
+	InTransform.SetScale3D(Scale);
+
+	FBoxFalloff BoxFalloffField(InMagnitude,
+		InMinRange,
+		InMaxRange,
+		InDefault,
+		InTransform,
+		(EFieldFalloffType)InFalloffType);
+
+	FFieldContext FieldContext{
+		ExecutionDatas,
+		FFieldContext::UniquePointerMap(),
+		0.0
+	};
+
+	// Evalute the field
+	BoxFalloffField.Evaluate(FieldContext, OutResultsView);
+}
+
+static void BoxFalloffFieldProcess(const TArray<FVector3f>& InSamplePositions,
+	const FBox& InBox,
+	const FTransform& InTransform,
+	const float InMagnitude,
+	const float InMinRange,
+	const float InMaxRange,
+	const float InDefault,
+	const EDataflowFieldFalloffType InFalloffType,
+	FDataflowVertexSelection& OutFieldSelectionMask,
+	TArray<float>& OutFieldFloatOutput)
+{
+	const int32 NumInSamplePositions = InSamplePositions.Num();
+
+	TArray<float> ResultsArray;
+	ResultsArray.Init(false, NumInSamplePositions);
+	TFieldArrayView<float> ResultsViewTest(ResultsArray, 0, ResultsArray.Num());
+
+	//
+	// Get the samples which effected by the field by evaluating the field with a default 
+	// value outside of the (MinRange, MaxRange)
+	//
+	const float NewDefault = InMinRange - 1000.f;
+
+	BoxFalloffFieldEvaluate(InSamplePositions,
+		InBox,
+		InTransform,
+		InMagnitude,
+		InMinRange,
+		InMaxRange,
+		NewDefault,
+		InFalloffType,
+		ResultsViewTest);
+
+	for (int32 Idx = 0; Idx < NumInSamplePositions; ++Idx)
+	{
+		if (ResultsViewTest[Idx] != NewDefault)
 		{
-			ExecutionDatas.SamplePositions[Index] = (FVector)InVertexArray[Index];
+			OutFieldSelectionMask.SetSelected(Idx);
 		}
+	}
+
+	// Compute the field
+	ResultsArray.Init(false, NumInSamplePositions);
+	TFieldArrayView<float> ResultsView(ResultsArray, 0, ResultsArray.Num());
+
+	BoxFalloffFieldEvaluate(InSamplePositions,
+		InBox,
+		InTransform,
+		InMagnitude,
+		InMinRange,
+		InMaxRange,
+		InDefault,
+		InFalloffType,
+		ResultsView);
+
+	for (int32 Idx = 0; Idx < NumInSamplePositions; ++Idx)
+	{
+		OutFieldFloatOutput[Idx] = ResultsView[Idx];
+	}
+}
+
+void FBoxFalloffFieldDataflowNode::Evaluate(Dataflow::FContext& Context, const FDataflowOutput* Out) const
+{
+	if (Out->IsA<TArray<float>>(&FieldFloatResult) ||
+		Out->IsA<TArray<int32>>(&FieldRemap) ||
+		Out->IsA<FDataflowVertexSelection>(&FieldSelectionMask) ||
+		Out->IsA<int32>(&NumSamplePositions))
+	{
+		const TArray<FVector3f>& InSamplePositions = GetValue<TArray<FVector3f>>(Context, &SamplePositions);
+		const FDataflowVertexSelection& InSampleIndices = GetValue<FDataflowVertexSelection>(Context, &SampleIndices);
+		const FBox& InBox = GetValue<FBox>(Context, &Box);
+		const FTransform& InTransform = GetValue<FTransform>(Context, &Transform);
+		const float InMagnitude = GetValue<float>(Context, &Magnitude);
+		const float InMinRange = GetValue<float>(Context, &MinRange);
+		const float InMaxRange = GetValue<float>(Context, &MaxRange);
+		const float InDefault = GetValue<float>(Context, &Default);
+
+		//
+		// SampleIndices input not connected, all the elements from SamplePositions array
+		// will be processed and FieldRemap output will be set empty
+		//
+		if (!IsConnected<FDataflowVertexSelection>(&SampleIndices))
+		{
+			const int32 NumInSamplePositions = InSamplePositions.Num();
+
+			//
+			// Process the field
+			//
+			FDataflowVertexSelection NewFieldSelectionMask;
+			NewFieldSelectionMask.Initialize(NumInSamplePositions, false);
+
+			TArray<float> NewFieldFloatResult;
+			NewFieldFloatResult.Init(false, NumInSamplePositions);
+
+			BoxFalloffFieldProcess(InSamplePositions,
+				InBox,
+				InTransform,
+				InMagnitude,
+				InMinRange,
+				InMaxRange,
+				InDefault,
+				FalloffType,
+				NewFieldSelectionMask,
+				NewFieldFloatResult);
+
+			// Set the outputs
+			SetValue<TArray<float>>(Context, NewFieldFloatResult, &FieldFloatResult);
+			SetValue<TArray<int32>>(Context, TArray<int32>(), &FieldRemap);
+			SetValue<FDataflowVertexSelection>(Context, NewFieldSelectionMask, &FieldSelectionMask);
+			SetValue<int32>(Context, InSamplePositions.Num(), &NumSamplePositions);
 
-		if (TUniquePtr<FRadialFalloff> RadialFalloffField = TUniquePtr<FRadialFalloff>(new FRadialFalloff()))
+			return;
+		}
+		//
+		// SampleIndices input connected, the selected elements from SamplePositions array
+		// will be processed and FieldRemap output will contain the remap info
+		// IMPORTANT: Number of elements in SamplePositions and SampleIndices must be the same
+		//
+		else
 		{
-			RadialFalloffField->Position = InSphere.Center + InTranslation;
-			RadialFalloffField->Radius = InSphere.W;
-			RadialFalloffField->Magnitude = Magnitude;
-			RadialFalloffField->MinRange = MinRange;
-			RadialFalloffField->MaxRange = MaxRange;
-			RadialFalloffField->Default = Default;
-			RadialFalloffField->Falloff = (EFieldFalloffType)FalloffType;
+			if (InSamplePositions.Num() == InSampleIndices.Num())
+			{
+				TArray<FVector3f> NewSamplePositions;
+				TArray<int32> NewRemapArray;
+
+				FieldComputeRemap(InSamplePositions, InSampleIndices, NewSamplePositions, NewRemapArray);
+
+				const int32 NumNewSamplePositions = NewSamplePositions.Num();
+
+				//
+				// Process the field
+				//
+				FDataflowVertexSelection NewFieldSelectionMask;
+				NewFieldSelectionMask.Initialize(NumNewSamplePositions, false);
+
+				TArray<float> NewFieldFloatResult;
+				NewFieldFloatResult.Init(false, NumNewSamplePositions);
+
+				BoxFalloffFieldProcess(NewSamplePositions,
+					InBox,
+					InTransform,
+					InMagnitude,
+					InMinRange,
+					InMaxRange,
+					InDefault,
+					FalloffType,
+					NewFieldSelectionMask,
+					NewFieldFloatResult);
+
+				// Set the outputs
+				SetValue<TArray<float>>(Context, NewFieldFloatResult, &FieldFloatResult);
+				SetValue<TArray<int32>>(Context, NewRemapArray, &FieldRemap);
+				SetValue<FDataflowVertexSelection>(Context, NewFieldSelectionMask, &FieldSelectionMask);
+				SetValue<int32>(Context, InSamplePositions.Num(), &NumSamplePositions);
+
+				return;
+			}
+		}
 
-			FFieldContext FieldContext{
-				ExecutionDatas,
-				FFieldContext::UniquePointerMap(),
-				0.0
-			};
+		// Set the outputs
+		SetValue<TArray<float>>(Context, TArray<float>(), &FieldFloatResult);
+		SetValue<TArray<int32>>(Context, TArray<int32>(), &FieldRemap);
+		SetValue<FDataflowVertexSelection>(Context, FDataflowVertexSelection(), &FieldSelectionMask);
+		SetValue<int32>(Context, 0, &NumSamplePositions);
+	}
+}
 
-			TArray<float> ResultsArray;
-			ResultsArray.Init(false, NumVertices);
-			TFieldArrayView<float> ResultsView(ResultsArray, 0, ResultsArray.Num());
-			RadialFalloffField->Evaluate(FieldContext, ResultsView);
+// --------------------------------------------------------------------------------------------------------------------
+
+static void PlaneFalloffFieldEvaluate(const TArray<FVector3f>& InSamplePositions,
+	const float InMagnitude,
+	const float InMinRange,
+	const float InMaxRange,
+	const float InDefault,
+	const float InDistance,
+	const FVector& InPosition,
+	const FVector& InNormal,
+	const EDataflowFieldFalloffType InFalloffType,
+	TFieldArrayView<float>& OutResultsView)
+{
+	const int32 NumInSamplePositions = InSamplePositions.Num();
 
-			TArray<float> NewWeightArray;
-			NewWeightArray.Init(0.f, NumVertices);
+	// Setup the data for the field 
+	FFieldExecutionDatas ExecutionDatas;
 
-			if (IsConnected<FDataflowVertexSelection>(&VertexSelection))
-			{
-				const FDataflowVertexSelection& InVertexSelection = GetValue<FDataflowVertexSelection>(Context, &VertexSelection);
+	FFieldContextIndex::ContiguousIndices(ExecutionDatas.SampleIndices, NumInSamplePositions);
+	ExecutionDatas.SamplePositions.Init(FVector(0.f), NumInSamplePositions);
 
-				if (InVertexArray.Num() == InVertexSelection.Num())
-				{
-					for (int32 Idx = 0; Idx < NumVertices; ++Idx)
-					{
-						if (InVertexSelection.IsSelected(Idx))
-						{
-							NewWeightArray[Idx] = ResultsView[Idx];
-						}
-					}
-				}
-			}
-			else
-			{
-				for (int32 Idx = 0; Idx < NumVertices; ++Idx)
-				{
-					NewWeightArray[Idx] = ResultsView[Idx];
-				}
-			}
+	for (int32 Idx = 0; Idx < NumInSamplePositions; Idx++)
+	{
+		ExecutionDatas.SamplePositions[Idx] = (FVector)InSamplePositions[Idx];
+	}
 
-			SetValue<TArray<float>>(Context, NewWeightArray, &WeightArray);
+	FPlaneFalloff PlaneFalloffField(InMagnitude,
+		InMinRange,
+		InMaxRange,
+		InDefault,
+		InDistance,
+		InPosition,
+		InNormal,
+		(EFieldFalloffType)InFalloffType);
+
+	FFieldContext FieldContext{
+		ExecutionDatas,
+		FFieldContext::UniquePointerMap(),
+		0.0
+	};
+
+	// Evalute the field
+	PlaneFalloffField.Evaluate(FieldContext, OutResultsView);
+}
 
-			return;
+static void PlaneFalloffFieldProcess(const TArray<FVector3f>& InSamplePositions,
+	const float InMagnitude,
+	const float InMinRange,
+	const float InMaxRange,
+	const float InDefault,
+	const float InDistance,
+	const FVector& InPosition,
+	const FVector& InTranslation,
+	const FVector& InNormal,
+	const EDataflowFieldFalloffType InFalloffType,
+	FDataflowVertexSelection& OutFieldSelectionMask,
+	TArray<float>& OutFieldFloatOutput)
+{
+	const int32 NumInSamplePositions = InSamplePositions.Num();
+
+	TArray<float> ResultsArray;
+	ResultsArray.Init(false, NumInSamplePositions);
+	TFieldArrayView<float> ResultsViewTest(ResultsArray, 0, ResultsArray.Num());
+
+	//
+	// Get the samples which effected by the field by evaluating the field with a default 
+	// value outside of the (MinRange, MaxRange)
+	//
+	const float NewDefault = InMinRange - 1000.f;
+
+	PlaneFalloffFieldEvaluate(InSamplePositions,
+		InMagnitude,
+		InMinRange,
+		InMaxRange,
+		NewDefault,
+		InDistance,
+		InPosition + InTranslation,
+		InNormal,
+		InFalloffType,
+		ResultsViewTest);
+
+	for (int32 Idx = 0; Idx < NumInSamplePositions; ++Idx)
+	{
+		if (ResultsViewTest[Idx] != NewDefault)
+		{
+			OutFieldSelectionMask.SetSelected(Idx);
 		}
+	}
 
-		SetValue<TArray<float>>(Context, TArray<float>(), &WeightArray);
+	// Compute the field
+	ResultsArray.Init(false, NumInSamplePositions);
+	TFieldArrayView<float> ResultsView(ResultsArray, 0, ResultsArray.Num());
+
+	PlaneFalloffFieldEvaluate(InSamplePositions,
+		InMagnitude,
+		InMinRange,
+		InMaxRange,
+		InDefault,
+		InDistance,
+		InPosition + InTranslation,
+		InNormal,
+		InFalloffType,
+		ResultsViewTest);
+
+	for (int32 Idx = 0; Idx < NumInSamplePositions; ++Idx)
+	{
+		OutFieldFloatOutput[Idx] = ResultsView[Idx];
 	}
 }
 
-
 void FPlaneFalloffFieldDataflowNode::Evaluate(Dataflow::FContext& Context, const FDataflowOutput* Out) const
 {
-	if (Out->IsA<TArray<float>>(&WeightArray))
+	if (Out->IsA<TArray<float>>(&FieldFloatResult) ||
+		Out->IsA<TArray<int32>>(&FieldRemap) ||
+		Out->IsA<FDataflowVertexSelection>(&FieldSelectionMask) ||
+		Out->IsA<int32>(&NumSamplePositions))
 	{
-		const TArray<FVector3f>& InVertexArray = GetValue<TArray<FVector3f>>(Context, &VertexArray);
-		const FVector& InTranslation = GetValue<FVector>(Context, &Translation);
+		const TArray<FVector3f>& InSamplePositions = GetValue<TArray<FVector3f>>(Context, &SamplePositions);
+		const FDataflowVertexSelection& InSampleIndices = GetValue<FDataflowVertexSelection>(Context, &SampleIndices);
 		const FVector& InPosition = GetValue<FVector>(Context, &Position);
 		const FVector& InNormal = GetValue<FVector>(Context, &Normal);
+		const float InDistance = GetValue<float>(Context, &Distance);
+		const FVector& InTranslation = GetValue<FVector>(Context, &Translation);
+		const float InMagnitude = GetValue<float>(Context, &Magnitude);
+		const float InMinRange = GetValue<float>(Context, &MinRange);
+		const float InMaxRange = GetValue<float>(Context, &MaxRange);
+		const float InDefault = GetValue<float>(Context, &Default);
+
+		//
+		// SampleIndices input not connected, all the elements from SamplePositions array
+		// will be processed and FieldRemap output will be set empty
+		//
+		if (!IsConnected<FDataflowVertexSelection>(&SampleIndices))
+		{
+			const int32 NumInSamplePositions = InSamplePositions.Num();
+
+			//
+			// Process the field
+			//
+			FDataflowVertexSelection NewFieldSelectionMask;
+			NewFieldSelectionMask.Initialize(NumInSamplePositions, false);
+
+			TArray<float> NewFieldFloatResult;
+			NewFieldFloatResult.Init(false, NumInSamplePositions);
+
+			PlaneFalloffFieldProcess(InSamplePositions,
+				InMagnitude,
+				InMinRange,
+				InMaxRange,
+				InDefault,
+				InDistance,
+				InPosition,
+				InTranslation,
+				InNormal,
+				FalloffType,
+				NewFieldSelectionMask,
+				NewFieldFloatResult);
+
+			// Set the outputs
+			SetValue<TArray<float>>(Context, NewFieldFloatResult, &FieldFloatResult);
+			SetValue<TArray<int32>>(Context, TArray<int32>(), &FieldRemap);
+			SetValue<FDataflowVertexSelection>(Context, NewFieldSelectionMask, &FieldSelectionMask);
+			SetValue<int32>(Context, InSamplePositions.Num(), &NumSamplePositions);
+
+			return;
+		}
+		//
+		// SampleIndices input connected, the selected elements from SamplePositions array
+		// will be processed and FieldRemap output will contain the remap info
+		// IMPORTANT: Number of elements in SamplePositions and SampleIndices must be the same
+		//
+		else
+		{
+			if (InSamplePositions.Num() == InSampleIndices.Num())
+			{
+				TArray<FVector3f> NewSamplePositions;
+				TArray<int32> NewRemapArray;
+
+				FieldComputeRemap(InSamplePositions, InSampleIndices, NewSamplePositions, NewRemapArray);
+
+				const int32 NumNewSamplePositions = NewSamplePositions.Num();
+
+				//
+				// Process the field
+				//
+				FDataflowVertexSelection NewFieldSelectionMask;
+				NewFieldSelectionMask.Initialize(NumNewSamplePositions, false);
+
+				TArray<float> NewFieldFloatResult;
+				NewFieldFloatResult.Init(false, NumNewSamplePositions);
+
+				PlaneFalloffFieldProcess(NewSamplePositions,
+					InMagnitude,
+					InMinRange,
+					InMaxRange,
+					InDefault,
+					InDistance,
+					InPosition,
+					InTranslation,
+					InNormal,
+					FalloffType,
+					NewFieldSelectionMask,
+					NewFieldFloatResult);
+
+				// Set the outputs
+				SetValue<TArray<float>>(Context, NewFieldFloatResult, &FieldFloatResult);
+				SetValue<TArray<int32>>(Context, NewRemapArray, &FieldRemap);
+				SetValue<FDataflowVertexSelection>(Context, NewFieldSelectionMask, &FieldSelectionMask);
+				SetValue<int32>(Context, InSamplePositions.Num(), &NumSamplePositions);
+
+				return;
+			}
+		}
 
-		const int32 NumVertices = InVertexArray.Num();
+		// Set the outputs
+		SetValue<TArray<float>>(Context, TArray<float>(), &FieldFloatResult);
+		SetValue<TArray<int32>>(Context, TArray<int32>(), &FieldRemap);
+		SetValue<FDataflowVertexSelection>(Context, FDataflowVertexSelection(), &FieldSelectionMask);
+		SetValue<int32>(Context, 0, &NumSamplePositions);
+	}
+}
 
-		FFieldExecutionDatas ExecutionDatas;
+// --------------------------------------------------------------------------------------------------------------------
 
-		FFieldContextIndex::ContiguousIndices(ExecutionDatas.SampleIndices, NumVertices);
-		ExecutionDatas.SamplePositions.Init(FVector(0.f), NumVertices);
+static void RadialIntMaskFieldEvaluate(const TArray<FVector3f>& InSamplePositions,
+	float InRadius,
+	const FVector InPosition,
+	int32 InInteriorValue,
+	int32 InExteriorValue,
+	EDataflowSetMaskConditionType InSetMaskConditionType,
+	TFieldArrayView<int32>& OutResultsView)
+{
+	const int32 NumInSamplePositions = InSamplePositions.Num();
 
-		for (int32 Index = 0; Index < NumVertices; Index++)
+	// Setup the data for the field 
+	FFieldExecutionDatas ExecutionDatas;
+
+	FFieldContextIndex::ContiguousIndices(ExecutionDatas.SampleIndices, NumInSamplePositions);
+	ExecutionDatas.SamplePositions.Init(FVector(0.f), NumInSamplePositions);
+
+	for (int32 Idx = 0; Idx < NumInSamplePositions; Idx++)
+	{
+		ExecutionDatas.SamplePositions[Idx] = (FVector)InSamplePositions[Idx];
+	}
+
+	FRadialIntMask RadialIntMaskField(InRadius, InPosition, InInteriorValue, InExteriorValue, (ESetMaskConditionType)InSetMaskConditionType);
+
+	FFieldContext FieldContext{
+		ExecutionDatas,
+		FFieldContext::UniquePointerMap(),
+		0.0
+	};
+
+	// Evalute the field
+	RadialIntMaskField.Evaluate(FieldContext, OutResultsView);
+}
+
+static void RadialIntMaskFieldProcess(const TArray<FVector3f>& InSamplePositions,
+	float InRadius,
+	const FVector InPosition,
+	const FVector InTranslation,
+	int32 InInteriorValue,
+	int32 InExteriorValue,
+	EDataflowSetMaskConditionType InSetMaskConditionType,
+	TArray<int32>& OutFieldFloatResult)
+{
+	const int32 NumInSamplePositions = InSamplePositions.Num();
+
+	// Compute the field
+	TArray<int32> ResultsArray;
+	ResultsArray.Init(0, NumInSamplePositions);
+	TFieldArrayView<int32> ResultsView(ResultsArray, 0, ResultsArray.Num());
+
+	RadialIntMaskFieldEvaluate(InSamplePositions, 
+		InRadius, 
+		InPosition + InTranslation,
+		InInteriorValue, 
+		InExteriorValue, 
+		InSetMaskConditionType, 
+		ResultsView);
+
+	// Set the outputs
+	OutFieldFloatResult.Init(0, NumInSamplePositions);
+
+	for (int32 Idx = 0; Idx < NumInSamplePositions; ++Idx)
+	{
+		OutFieldFloatResult[Idx] = ResultsView[Idx];
+	}
+}
+
+void FRadialIntMaskFieldDataflowNode::Evaluate(Dataflow::FContext& Context, const FDataflowOutput* Out) const
+{
+	if (Out->IsA<TArray<int32>>(&FieldIntResult) ||
+		Out->IsA<TArray<int32>>(&FieldRemap) ||
+		Out->IsA<int32>(&NumSamplePositions))
+	{
+		const TArray<FVector3f>& InSamplePositions = GetValue<TArray<FVector3f>>(Context, &SamplePositions);
+		const FDataflowVertexSelection& InSampleIndices = GetValue<FDataflowVertexSelection>(Context, &SampleIndices);
+		const FSphere& InSphere = GetValue<FSphere>(Context, &Sphere);
+		const FVector InTranslation = GetValue<FVector>(Context, &Translation);
+		const int32 InInteriorValue = GetValue<int32>(Context, &InteriorValue);
+		const int32 InExteriorValue = GetValue<int32>(Context, &ExteriorValue);
+
+		//
+		// SampleIndices input not connected, all the elements from SamplePositions array
+		// will be processed and FieldRemapOutput output will be set empty
+		//
+		if (!IsConnected<FDataflowVertexSelection>(&SampleIndices))
 		{
-			ExecutionDatas.SamplePositions[Index] = (FVector)InVertexArray[Index];
+			// Compute the field
+			TArray<int32> NewFieldIntResult;
+			RadialIntMaskFieldProcess(InSamplePositions,
+				InSphere.W,
+				InSphere.Center,
+				InTranslation,
+				InInteriorValue,
+				InExteriorValue,
+				SetMaskConditionType,
+				NewFieldIntResult);
+
+			// Set the outputs
+			SetValue<TArray<int32>>(Context, NewFieldIntResult, &FieldIntResult);
+			SetValue<TArray<int32>>(Context, TArray<int32>(), &FieldRemap);
+			SetValue<int32>(Context, InSamplePositions.Num(), &NumSamplePositions);
+
+			return;
+		}
+		//
+		// SampleIndices input connected, the selected elements from SamplePositions array
+		// will be processed and FieldRemapOutput output will contain the remap info
+		// IMPORTANT: Number of elements in SamplePositions and SampleIndices must be the same
+		//
+		else
+		{
+			if (InSamplePositions.Num() == InSampleIndices.Num())
+			{
+				TArray<FVector3f> NewSamplePositions;
+				TArray<int32> NewRemapArray;
+
+				FieldComputeRemap(InSamplePositions, InSampleIndices, NewSamplePositions, NewRemapArray);
+
+				// Compute the field
+				TArray<int32> NewFieldIntResult;
+				RadialIntMaskFieldProcess(NewSamplePositions,
+					InSphere.W,
+					InSphere.Center,
+					InTranslation,
+					InInteriorValue,
+					InExteriorValue,
+					SetMaskConditionType,
+					NewFieldIntResult);
+
+				// Set the outputs
+				SetValue<TArray<int32>>(Context, NewFieldIntResult, &FieldIntResult);
+				SetValue<TArray<int32>>(Context, NewRemapArray, &FieldRemap);
+				SetValue<int32>(Context, InSamplePositions.Num(), &NumSamplePositions);
+
+				return;
+			}
 		}
 
-		if (TUniquePtr<FPlaneFalloff> PlaneFalloffField = TUniquePtr<FPlaneFalloff>(new FPlaneFalloff()))
+		SetValue<TArray<int32>>(Context, TArray<int32>(), &FieldIntResult);
+		SetValue<TArray<int32>>(Context, TArray<int32>(), &FieldRemap);
+		SetValue<int32>(Context, 0, &NumSamplePositions);
+	}
+}
+
+// --------------------------------------------------------------------------------------------------------------------
+
+static void UniformScalarFieldEvaluate(const TArray<FVector3f>& InSamplePositions,
+	float InMagnitude,
+	TFieldArrayView<float>& OutResultsView)
+{
+	const int32 NumInSamplePositions = InSamplePositions.Num();
+
+	// Setup the data for the field 
+	FFieldExecutionDatas ExecutionDatas;
+
+	FFieldContextIndex::ContiguousIndices(ExecutionDatas.SampleIndices, NumInSamplePositions);
+	ExecutionDatas.SamplePositions.Init(FVector(0.f), NumInSamplePositions);
+
+	for (int32 Idx = 0; Idx < NumInSamplePositions; Idx++)
+	{
+		ExecutionDatas.SamplePositions[Idx] = (FVector)InSamplePositions[Idx];
+	}
+
+	FUniformScalar UniformScalarField(InMagnitude);
+
+	FFieldContext FieldContext{
+		ExecutionDatas,
+		FFieldContext::UniquePointerMap(),
+		0.0
+	};
+
+	// Evalute the field
+	UniformScalarField.Evaluate(FieldContext, OutResultsView);
+}
+
+static void UniformScalarFieldProcess(const TArray<FVector3f>& InSamplePositions,
+	float InMagnitude,
+	TArray<float>& OutFieldFloatResult)
+{
+	const int32 NumInSamplePositions = InSamplePositions.Num();
+
+	// Compute the field
+	TArray<float> ResultsArray;
+	ResultsArray.Init(0.f, NumInSamplePositions);
+	TFieldArrayView<float> ResultsView(ResultsArray, 0, ResultsArray.Num());
+
+	UniformScalarFieldEvaluate(InSamplePositions, InMagnitude, ResultsView);
+
+	// Set the outputs
+	OutFieldFloatResult.Init(0.f, NumInSamplePositions);
+
+	for (int32 Idx = 0; Idx < NumInSamplePositions; ++Idx)
+	{
+		OutFieldFloatResult[Idx] = ResultsView[Idx];
+	}
+}
+
+void FUniformScalarFieldDataflowNode::Evaluate(Dataflow::FContext& Context, const FDataflowOutput* Out) const
+{
+	if (Out->IsA<TArray<float>>(&FieldFloatResult) ||
+		Out->IsA<TArray<int32>>(&FieldRemap) ||
+		Out->IsA<int32>(&NumSamplePositions))
+	{
+		const TArray<FVector3f>& InSamplePositions = GetValue<TArray<FVector3f>>(Context, &SamplePositions);
+		const FDataflowVertexSelection& InSampleIndices = GetValue<FDataflowVertexSelection>(Context, &SampleIndices);
+		const float InMagnitude = GetValue<float>(Context, &Magnitude);
+
+		//
+		// SampleIndices input not connected, all the elements from SamplePositions array
+		// will be processed and FieldRemapOutput output will be set empty
+		//
+		if (!IsConnected<FDataflowVertexSelection>(&SampleIndices))
 		{
-			PlaneFalloffField->Position = InPosition + InTranslation;
-			PlaneFalloffField->Normal = InNormal;
-			PlaneFalloffField->Magnitude = Magnitude;
-			PlaneFalloffField->MinRange = MinRange;
-			PlaneFalloffField->MaxRange = MaxRange;
-			PlaneFalloffField->Distance = Distance;
-			PlaneFalloffField->Default = Default;
-			PlaneFalloffField->Falloff = (EFieldFalloffType)FalloffType;
+			// Compute the field
+			TArray<float> NewFieldFloatResult;
+			UniformScalarFieldProcess(InSamplePositions, InMagnitude, NewFieldFloatResult);
+
+			// Set the outputs
+			SetValue<TArray<float>>(Context, NewFieldFloatResult, &FieldFloatResult);
+			SetValue<TArray<int32>>(Context, TArray<int32>(), &FieldRemap);
+			SetValue<int32>(Context, InSamplePositions.Num(), &NumSamplePositions);
+
+			return;
+		}
+		//
+		// SampleIndices input connected, the selected elements from SamplePositions array
+		// will be processed and FieldRemapOutput output will contain the remap info
+		// IMPORTANT: Number of elements in SamplePositions and SampleIndices must be the same
+		//
+		else
+		{
+			if (InSamplePositions.Num() == InSampleIndices.Num())
+			{
+				TArray<FVector3f> NewSamplePositions;
+				TArray<int32> NewRemapArray;
+
+				FieldComputeRemap(InSamplePositions, InSampleIndices, NewSamplePositions, NewRemapArray);
+
+				// Compute the field
+				TArray<float> NewFieldFloatResult;
+				UniformScalarFieldProcess(NewSamplePositions, InMagnitude, NewFieldFloatResult);
+
+				// Set the outputs
+				SetValue<TArray<float>>(Context, NewFieldFloatResult, &FieldFloatResult);
+				SetValue<TArray<int32>>(Context, NewRemapArray, &FieldRemap);
+				SetValue<int32>(Context, InSamplePositions.Num(), &NumSamplePositions);
+
+				return;
+			}
+		}
+
+		SetValue<TArray<float>>(Context, TArray<float>(), &FieldFloatResult);
+		SetValue<TArray<int32>>(Context, TArray<int32>(), &FieldRemap);
+		SetValue<int32>(Context, 0, &NumSamplePositions);
+	}
+}
+
+// --------------------------------------------------------------------------------------------------------------------
+
+static void UniformVectorFieldEvaluate(const TArray<FVector3f>& InSamplePositions,
+	const float InMagnitude,
+	const FVector InDirection,
+	TFieldArrayView<FVector>& OutResultsView)
+{
+	const int32 NumInSamplePositions = InSamplePositions.Num();
+
+	// Setup the data for the field 
+	FFieldExecutionDatas ExecutionDatas;
+
+	FFieldContextIndex::ContiguousIndices(ExecutionDatas.SampleIndices, NumInSamplePositions);
+	ExecutionDatas.SamplePositions.Init(FVector(0.f), NumInSamplePositions);
+
+	for (int32 Idx = 0; Idx < NumInSamplePositions; Idx++)
+	{
+		ExecutionDatas.SamplePositions[Idx] = (FVector)InSamplePositions[Idx];
+	}
+
+	FUniformVector UniformVectorField(InMagnitude, InDirection);
+
+	FFieldContext FieldContext{
+		ExecutionDatas,
+		FFieldContext::UniquePointerMap(),
+		0.0
+	};
+
+	// Evalute the field
+	UniformVectorField.Evaluate(FieldContext, OutResultsView);
+}
+
+static void UniformVectorFieldProcess(const TArray<FVector3f>& InSamplePositions,
+	float InMagnitude,
+	const FVector InDirection,
+	TArray<FVector>& OutFieldFloatResult)
+{
+	const int32 NumInSamplePositions = InSamplePositions.Num();
+
+	// Compute the field
+	TArray<FVector> ResultsArray;
+	ResultsArray.Init(FVector(0.f), NumInSamplePositions);
+	TFieldArrayView<FVector> ResultsView(ResultsArray, 0, ResultsArray.Num());
 
-			FFieldContext FieldContext{
-				ExecutionDatas,
-				FFieldContext::UniquePointerMap(),
-				0.0
-			};
+	UniformVectorFieldEvaluate(InSamplePositions, InMagnitude, InDirection, ResultsView);
 
-			TArray<float> ResultsArray;
-			ResultsArray.Init(false, NumVertices);
-			TFieldArrayView<float> ResultsView(ResultsArray, 0, ResultsArray.Num());
-			PlaneFalloffField->Evaluate(FieldContext, ResultsView);
+	// Set the outputs
+	OutFieldFloatResult.Init(FVector(0.f), NumInSamplePositions);
+
+	for (int32 Idx = 0; Idx < NumInSamplePositions; ++Idx)
+	{
+		OutFieldFloatResult[Idx] = ResultsView[Idx];
+	}
+}
+
+void FUniformVectorFieldDataflowNode::Evaluate(Dataflow::FContext& Context, const FDataflowOutput* Out) const
+{
+	if (Out->IsA<TArray<FVector>>(&FieldVectorResult) ||
+		Out->IsA<TArray<int32>>(&FieldRemap) ||
+		Out->IsA<int32>(&NumSamplePositions))
+	{
+		const TArray<FVector3f>& InSamplePositions = GetValue<TArray<FVector3f>>(Context, &SamplePositions);
+		const FDataflowVertexSelection& InSampleIndices = GetValue<FDataflowVertexSelection>(Context, &SampleIndices);
+		const float InMagnitude = GetValue<float>(Context, &Magnitude);
+		const FVector InDirection = GetValue<FVector>(Context, &Direction);
+
+		//
+		// SampleIndices input not connected, all the elements from SamplePositions array
+		// will be processed and FieldRemapOutput output will be set empty
+		//
+		if (!IsConnected<FDataflowVertexSelection>(&SampleIndices))
+		{
+			// Compute the field
+			TArray<FVector> NewFieldVectorResult;
+			UniformVectorFieldProcess(InSamplePositions, InMagnitude, InDirection, NewFieldVectorResult);
 
-			TArray<float> NewWeightArray;
-			NewWeightArray.Init(0.f, NumVertices);
+			// Set the outputs
+			SetValue<TArray<FVector>>(Context, NewFieldVectorResult, &FieldVectorResult);
+			SetValue<TArray<int32>>(Context, TArray<int32>(), &FieldRemap);
+			SetValue<int32>(Context, InSamplePositions.Num(), &NumSamplePositions);
 
-			if (IsConnected<FDataflowVertexSelection>(&VertexSelection))
+			return;
+		}
+		//
+		// SampleIndices input connected, the selected elements from SamplePositions array
+		// will be processed and FieldRemapOutput output will contain the remap info
+		// IMPORTANT: Number of elements in SamplePositions and SampleIndices must be the same
+		//
+		else
+		{
+			if (InSamplePositions.Num() == InSampleIndices.Num())
 			{
-				const FDataflowVertexSelection& InVertexSelection = GetValue<FDataflowVertexSelection>(Context, &VertexSelection);
+				TArray<FVector3f> NewSamplePositions;
+				TArray<int32> NewRemapArray;
 
-				if (InVertexArray.Num() == InVertexSelection.Num())
-				{
-					for (int32 Idx = 0; Idx < NumVertices; ++Idx)
-					{
-						if (InVertexSelection.IsSelected(Idx))
-						{
-							NewWeightArray[Idx] = ResultsView[Idx];
-						}
-					}
-				}
+				FieldComputeRemap(InSamplePositions, InSampleIndices, NewSamplePositions, NewRemapArray);
+
+				// Compute the field
+				TArray<FVector> NewFieldVectorResult;
+				UniformVectorFieldProcess(NewSamplePositions, InMagnitude, InDirection, NewFieldVectorResult);
+
+				// Set the outputs
+				SetValue<TArray<FVector>>(Context, NewFieldVectorResult, &FieldVectorResult);
+				SetValue<TArray<int32>>(Context, NewRemapArray, &FieldRemap);
+				SetValue<int32>(Context, InSamplePositions.Num(), &NumSamplePositions);
+
+				return;
 			}
-			else
+		}
+
+		SetValue<TArray<FVector>>(Context, TArray<FVector>(), &FieldVectorResult);
+		SetValue<TArray<int32>>(Context, TArray<int32>(), &FieldRemap);
+		SetValue<int32>(Context, 0, &NumSamplePositions);
+	}
+
+}
+
+// --------------------------------------------------------------------------------------------------------------------
+
+static void RadialVectorFieldEvaluate(const TArray<FVector3f>& InSamplePositions,
+	const float InMagnitude,
+	const FVector InPosition,
+	TFieldArrayView<FVector>& OutResultsView)
+{
+	const int32 NumInSamplePositions = InSamplePositions.Num();
+
+	// Setup the data for the field 
+	FFieldExecutionDatas ExecutionDatas;
+
+	FFieldContextIndex::ContiguousIndices(ExecutionDatas.SampleIndices, NumInSamplePositions);
+	ExecutionDatas.SamplePositions.Init(FVector(0.f), NumInSamplePositions);
+
+	for (int32 Idx = 0; Idx < NumInSamplePositions; Idx++)
+	{
+		ExecutionDatas.SamplePositions[Idx] = (FVector)InSamplePositions[Idx];
+	}
+
+	FRadialVector RadialVectorField(InMagnitude, InPosition);
+
+	FFieldContext FieldContext{
+		ExecutionDatas,
+		FFieldContext::UniquePointerMap(),
+		0.0
+	};
+
+	// Evalute the field
+	RadialVectorField.Evaluate(FieldContext, OutResultsView);
+}
+
+static void RadialVectorFieldProcess(const TArray<FVector3f>& InSamplePositions,
+	float InMagnitude,
+	const FVector InPosition,
+	TArray<FVector>& OutFieldFloatResult)
+{
+	const int32 NumInSamplePositions = InSamplePositions.Num();
+
+	// Compute the field
+	TArray<FVector> ResultsArray;
+	ResultsArray.Init(FVector(0.f), NumInSamplePositions);
+	TFieldArrayView<FVector> ResultsView(ResultsArray, 0, ResultsArray.Num());
+
+	RadialVectorFieldEvaluate(InSamplePositions, InMagnitude, InPosition, ResultsView);
+
+	// Set the outputs
+	OutFieldFloatResult.Init(FVector(0.f), NumInSamplePositions);
+
+	for (int32 Idx = 0; Idx < NumInSamplePositions; ++Idx)
+	{
+		OutFieldFloatResult[Idx] = ResultsView[Idx];
+	}
+}
+void FRadialVectorFieldDataflowNode::Evaluate(Dataflow::FContext& Context, const FDataflowOutput* Out) const
+{
+	if (Out->IsA<TArray<FVector>>(&FieldVectorResult) ||
+		Out->IsA<TArray<int32>>(&FieldRemap) ||
+		Out->IsA<int32>(&NumSamplePositions))
+	{
+		const TArray<FVector3f>& InSamplePositions = GetValue<TArray<FVector3f>>(Context, &SamplePositions);
+		const FDataflowVertexSelection& InSampleIndices = GetValue<FDataflowVertexSelection>(Context, &SampleIndices);
+		const float InMagnitude = GetValue<float>(Context, &Magnitude);
+		const FVector InPosition = GetValue<FVector>(Context, &Position);
+
+		//
+		// SampleIndices input not connected, all the elements from SamplePositions array
+		// will be processed and FieldRemapOutput output will be set empty
+		//
+		if (!IsConnected<FDataflowVertexSelection>(&SampleIndices))
+		{
+			// Compute the field
+			TArray<FVector> NewFieldVectorResult;
+			RadialVectorFieldProcess(InSamplePositions, InMagnitude, InPosition, NewFieldVectorResult);
+
+			// Set the outputs
+			SetValue<TArray<FVector>>(Context, NewFieldVectorResult, &FieldVectorResult);
+			SetValue<TArray<int32>>(Context, TArray<int32>(), &FieldRemap);
+			SetValue<int32>(Context, InSamplePositions.Num(), &NumSamplePositions);
+
+			return;
+		}
+		//
+		// SampleIndices input connected, the selected elements from SamplePositions array
+		// will be processed and FieldRemapOutput output will contain the remap info
+		// IMPORTANT: Number of elements in SamplePositions and SampleIndices must be the same
+		//
+		else
+		{
+			if (InSamplePositions.Num() == InSampleIndices.Num())
 			{
-				for (int32 Idx = 0; Idx < NumVertices; ++Idx)
-				{
-					NewWeightArray[Idx] = ResultsView[Idx];
-				}
+				TArray<FVector3f> NewSamplePositions;
+				TArray<int32> NewRemapArray;
+
+				FieldComputeRemap(InSamplePositions, InSampleIndices, NewSamplePositions, NewRemapArray);
+
+				// Compute the field
+				TArray<FVector> NewFieldVectorResult;
+				RadialVectorFieldProcess(NewSamplePositions, InMagnitude, InPosition, NewFieldVectorResult);
+
+				// Set the outputs
+				SetValue<TArray<FVector>>(Context, NewFieldVectorResult, &FieldVectorResult);
+				SetValue<TArray<int32>>(Context, NewRemapArray, &FieldRemap);
+				SetValue<int32>(Context, InSamplePositions.Num(), &NumSamplePositions);
+
+				return;
 			}
+		}
+
+		SetValue<TArray<FVector>>(Context, TArray<FVector>(), &FieldVectorResult);
+		SetValue<TArray<int32>>(Context, TArray<int32>(), &FieldRemap);
+		SetValue<int32>(Context, 0, &NumSamplePositions);
+	}
+}
+
+// --------------------------------------------------------------------------------------------------------------------
+
+static void RandomVectorFieldEvaluate(const TArray<FVector3f>& InSamplePositions,
+	const float InMagnitude,
+	TFieldArrayView<FVector>& OutResultsView)
+{
+	const int32 NumInSamplePositions = InSamplePositions.Num();
+
+	// Setup the data for the field 
+	FFieldExecutionDatas ExecutionDatas;
+
+	FFieldContextIndex::ContiguousIndices(ExecutionDatas.SampleIndices, NumInSamplePositions);
+	ExecutionDatas.SamplePositions.Init(FVector(0.f), NumInSamplePositions);
+
+	for (int32 Idx = 0; Idx < NumInSamplePositions; Idx++)
+	{
+		ExecutionDatas.SamplePositions[Idx] = (FVector)InSamplePositions[Idx];
+	}
+
+	FRandomVector RandomVectorField(InMagnitude);
+
+	FFieldContext FieldContext{
+		ExecutionDatas,
+		FFieldContext::UniquePointerMap(),
+		0.0
+	};
+
+	// Evalute the field
+	RandomVectorField.Evaluate(FieldContext, OutResultsView);
+}
+
+static void RandomVectorFieldProcess(const TArray<FVector3f>& InSamplePositions,
+	float InMagnitude,
+	TArray<FVector>& OutFieldFloatResult)
+{
+	const int32 NumInSamplePositions = InSamplePositions.Num();
+
+	// Compute the field
+	TArray<FVector> ResultsArray;
+	ResultsArray.Init(FVector(0.f), NumInSamplePositions);
+	TFieldArrayView<FVector> ResultsView(ResultsArray, 0, ResultsArray.Num());
+
+	RandomVectorFieldEvaluate(InSamplePositions, InMagnitude, ResultsView);
+
+	// Set the outputs
+	OutFieldFloatResult.Init(FVector(0.f), NumInSamplePositions);
+
+	for (int32 Idx = 0; Idx < NumInSamplePositions; ++Idx)
+	{
+		OutFieldFloatResult[Idx] = ResultsView[Idx];
+	}
+}
 
-			SetValue<TArray<float>>(Context, NewWeightArray, &WeightArray);
+void FRandomVectorFieldDataflowNode::Evaluate(Dataflow::FContext& Context, const FDataflowOutput* Out) const
+{
+	if (Out->IsA<TArray<FVector>>(&FieldVectorResult) ||
+		Out->IsA<TArray<int32>>(&FieldRemap) ||
+		Out->IsA<int32>(&NumSamplePositions))
+	{
+		const TArray<FVector3f>& InSamplePositions = GetValue<TArray<FVector3f>>(Context, &SamplePositions);
+		const FDataflowVertexSelection& InSampleIndices = GetValue<FDataflowVertexSelection>(Context, &SampleIndices);
+		const float InMagnitude = GetValue<float>(Context, &Magnitude);
+
+		//
+		// SampleIndices input not connected, all the elements from SamplePositions array
+		// will be processed and FieldRemapOutput output will be set empty
+		//
+		if (!IsConnected<FDataflowVertexSelection>(&SampleIndices))
+		{
+			// Compute the field
+			TArray<FVector> NewFieldVectorResult;
+			RandomVectorFieldProcess(InSamplePositions, InMagnitude, NewFieldVectorResult);
+
+			// Set the outputs
+			SetValue<TArray<FVector>>(Context, NewFieldVectorResult, &FieldVectorResult);
+			SetValue<TArray<int32>>(Context, TArray<int32>(), &FieldRemap);
+			SetValue<int32>(Context, InSamplePositions.Num(), &NumSamplePositions);
 
 			return;
 		}
+		//
+		// SampleIndices input connected, the selected elements from SamplePositions array
+		// will be processed and FieldRemapOutput output will contain the remap info
+		// IMPORTANT: Number of elements in SamplePositions and SampleIndices must be the same
+		//
+		else
+		{
+			if (InSamplePositions.Num() == InSampleIndices.Num())
+			{
+				TArray<FVector3f> NewSamplePositions;
+				TArray<int32> NewRemapArray;
+
+				FieldComputeRemap(InSamplePositions, InSampleIndices, NewSamplePositions, NewRemapArray);
+
+				// Compute the field
+				TArray<FVector> NewFieldVectorResult;
+				RandomVectorFieldProcess(NewSamplePositions, InMagnitude, NewFieldVectorResult);
+
+				// Set the outputs
+				SetValue<TArray<FVector>>(Context, NewFieldVectorResult, &FieldVectorResult);
+				SetValue<TArray<int32>>(Context, NewRemapArray, &FieldRemap);
+				SetValue<int32>(Context, InSamplePositions.Num(), &NumSamplePositions);
+
+				return;
+			}
+		}
 
-		SetValue<TArray<float>>(Context, TArray<float>(), &WeightArray);
+		SetValue<TArray<FVector>>(Context, TArray<FVector>(), &FieldVectorResult);
+		SetValue<TArray<int32>>(Context, TArray<int32>(), &FieldRemap);
+		SetValue<int32>(Context, 0, &NumSamplePositions);
 	}
 }
 
+// --------------------------------------------------------------------------------------------------------------------
 
-void FRadialIntMaskFieldDataflowNode::Evaluate(Dataflow::FContext& Context, const FDataflowOutput* Out) const
+static void NoiseFieldEvaluate(const TArray<FVector3f>& InSamplePositions,
+	float InMinRange,
+	float InMaxRange,
+	const FTransform& InTransform,
+	TFieldArrayView<float>& OutResultsView)
 {
-	if (Out->IsA<TArray<int32>>(&MaskArray))
+	const int32 NumInSamplePositions = InSamplePositions.Num();
+
+	// Setup the data for the field 
+	FFieldExecutionDatas ExecutionDatas;
+
+	FFieldContextIndex::ContiguousIndices(ExecutionDatas.SampleIndices, NumInSamplePositions);
+	ExecutionDatas.SamplePositions.Init(FVector(0.f), NumInSamplePositions);
+
+	for (int32 Idx = 0; Idx < NumInSamplePositions; Idx++)
 	{
-		const TArray<FVector3f>& InVertexArray = GetValue<TArray<FVector3f>>(Context, &VertexArray);
-		const FSphere& InSphere = GetValue<FSphere>(Context, &Sphere);
-		const FVector& InTranslation = GetValue<FVector>(Context, &Translation);
+		ExecutionDatas.SamplePositions[Idx] = (FVector)InSamplePositions[Idx];
+	}
+
+	FNoiseField NoiseField(InMinRange, InMaxRange, InTransform);
+
+	FFieldContext FieldContext{
+		ExecutionDatas,
+		FFieldContext::UniquePointerMap(),
+		0.0
+	};
+
+	// Evalute the field
+	NoiseField.Evaluate(FieldContext, OutResultsView);
+}
+
+static void NoiseFieldProcess(const TArray<FVector3f>& InSamplePositions,
+	float InMinRange,
+	float InMaxRange,
+	const FTransform& InTransform, 
+	TArray<float>& OutFieldFloatResult)
+{
+	const int32 NumInSamplePositions = InSamplePositions.Num();
 
-		const int32 NumVertices = InVertexArray.Num();
+	// Compute the field
+	TArray<float> ResultsArray;
+	ResultsArray.Init(0.f, NumInSamplePositions);
+	TFieldArrayView<float> ResultsView(ResultsArray, 0, ResultsArray.Num());
 
-		FFieldExecutionDatas ExecutionDatas;
+	NoiseFieldEvaluate(InSamplePositions, InMinRange, InMaxRange, InTransform, ResultsView);
 
-		FFieldContextIndex::ContiguousIndices(ExecutionDatas.SampleIndices, NumVertices);
-		ExecutionDatas.SamplePositions.Init(FVector(0.f), NumVertices);
+	// Set the outputs
+	OutFieldFloatResult.Init(0.f, NumInSamplePositions);
 
-		for (int32 Index = 0; Index < NumVertices; Index++)
+	for (int32 Idx = 0; Idx < NumInSamplePositions; ++Idx)
+	{
+		OutFieldFloatResult[Idx] = ResultsView[Idx];
+	}
+}
+
+void FNoiseFieldDataflowNode::Evaluate(Dataflow::FContext& Context, const FDataflowOutput* Out) const
+{
+	if (Out->IsA<TArray<float>>(&FieldFloatResult) ||
+		Out->IsA<TArray<int32>>(&FieldRemap) ||
+		Out->IsA<int32>(&NumSamplePositions))
+	{
+		const TArray<FVector3f>& InSamplePositions = GetValue<TArray<FVector3f>>(Context, &SamplePositions);
+		const FDataflowVertexSelection& InSampleIndices = GetValue<FDataflowVertexSelection>(Context, &SampleIndices);
+		const float InMinRange = GetValue<float>(Context, &MinRange);
+		const float InMaxRange = GetValue<float>(Context, &MaxRange);
+		const FTransform InTransform = GetValue<FTransform>(Context, &Transform);
+
+		//
+		// SampleIndices input not connected, all the elements from SamplePositions array
+		// will be processed and FieldRemapOutput output will be set empty
+		//
+		if (!IsConnected<FDataflowVertexSelection>(&SampleIndices))
 		{
-			ExecutionDatas.SamplePositions[Index] = (FVector)InVertexArray[Index];
+			// Compute the field
+			TArray<float> NewFieldFloatResult;
+			NoiseFieldProcess(InSamplePositions, InMinRange, InMaxRange, InTransform, NewFieldFloatResult);
+
+			// Set the outputs
+			SetValue<TArray<float>>(Context, NewFieldFloatResult, &FieldFloatResult);
+			SetValue<TArray<int32>>(Context, TArray<int32>(), &FieldRemap);
+			SetValue<int32>(Context, InSamplePositions.Num(), &NumSamplePositions);
+
+			return;
 		}
+		//
+		// SampleIndices input connected, the selected elements from SamplePositions array
+		// will be processed and FieldRemapOutput output will contain the remap info
+		// IMPORTANT: Number of elements in SamplePositions and SampleIndices must be the same
+		//
+		else
+		{
+			if (InSamplePositions.Num() == InSampleIndices.Num())
+			{
+				TArray<FVector3f> NewSamplePositions;
+				TArray<int32> NewRemapArray;
+
+				FieldComputeRemap(InSamplePositions, InSampleIndices, NewSamplePositions, NewRemapArray);
+
+				// Compute the field
+				TArray<float> NewFieldFloatResult;
+				NoiseFieldProcess(NewSamplePositions, InMinRange, InMaxRange, InTransform, NewFieldFloatResult);
 
-		if (TUniquePtr<FRadialIntMask> RadialIntMaskField = TUniquePtr<FRadialIntMask>(new FRadialIntMask()))
+				// Set the outputs
+				SetValue<TArray<float>>(Context, NewFieldFloatResult, &FieldFloatResult);
+				SetValue<TArray<int32>>(Context, NewRemapArray, &FieldRemap);
+				SetValue<int32>(Context, InSamplePositions.Num(), &NumSamplePositions);
+
+				return;
+			}
+		}
+
+		SetValue<TArray<float>>(Context, TArray<float>(), &FieldFloatResult);
+		SetValue<TArray<int32>>(Context, TArray<int32>(), &FieldRemap);
+		SetValue<int32>(Context, 0, &NumSamplePositions);
+	}
+}
+
+// --------------------------------------------------------------------------------------------------------------------
+
+static void UniformIntegerFieldEvaluate(const TArray<FVector3f>& InSamplePositions,
+	const int32 InMagnitude,
+	TFieldArrayView<int32>& OutResultsView)
+{
+	const int32 NumInSamplePositions = InSamplePositions.Num();
+
+	// Setup the data for the field 
+	FFieldExecutionDatas ExecutionDatas;
+
+	FFieldContextIndex::ContiguousIndices(ExecutionDatas.SampleIndices, NumInSamplePositions);
+	ExecutionDatas.SamplePositions.Init(FVector(0.f), NumInSamplePositions);
+
+	for (int32 Idx = 0; Idx < NumInSamplePositions; Idx++)
+	{
+		ExecutionDatas.SamplePositions[Idx] = (FVector)InSamplePositions[Idx];
+	}
+
+	FUniformInteger UniformIntegerField(InMagnitude);
+
+	FFieldContext FieldContext{
+		ExecutionDatas,
+		FFieldContext::UniquePointerMap(),
+		0.0
+	};
+
+	// Evalute the field
+	UniformIntegerField.Evaluate(FieldContext, OutResultsView);
+}
+
+static void UniformIntegerFieldProcess(const TArray<FVector3f>& InSamplePositions,
+	int32 InMagnitude,
+	TArray<int32>& OutFieldFloatResult)
+{
+	const int32 NumInSamplePositions = InSamplePositions.Num();
+
+	// Compute the field
+	TArray<int32> ResultsArray;
+	ResultsArray.Init(0, NumInSamplePositions);
+	TFieldArrayView<int32> ResultsView(ResultsArray, 0, ResultsArray.Num());
+
+	UniformIntegerFieldEvaluate(InSamplePositions, InMagnitude, ResultsView);
+
+	// Set the outputs
+	OutFieldFloatResult.Init(0, NumInSamplePositions);
+
+	for (int32 Idx = 0; Idx < NumInSamplePositions; ++Idx)
+	{
+		OutFieldFloatResult[Idx] = ResultsView[Idx];
+	}
+}
+
+void FUniformIntegerFieldDataflowNode::Evaluate(Dataflow::FContext& Context, const FDataflowOutput* Out) const
+{
+	if (Out->IsA<TArray<int32>>(&FieldIntResult) ||
+		Out->IsA<TArray<int32>>(&FieldRemap) ||
+		Out->IsA<int32>(&NumSamplePositions))
+	{
+		const TArray<FVector3f>& InSamplePositions = GetValue<TArray<FVector3f>>(Context, &SamplePositions);
+		const FDataflowVertexSelection& InSampleIndices = GetValue<FDataflowVertexSelection>(Context, &SampleIndices);
+		const int32 InMagnitude = GetValue<int32>(Context, &Magnitude);
+
+		//
+		// SampleIndices input not connected, all the elements from SamplePositions array
+		// will be processed and FieldRemapOutput output will be set empty
+		//
+		if (!IsConnected<FDataflowVertexSelection>(&SampleIndices))
 		{
-			RadialIntMaskField->Position = InSphere.Center + InTranslation;
-			RadialIntMaskField->Radius = InSphere.W;
-			RadialIntMaskField->InteriorValue = InteriorValue;
-			RadialIntMaskField->ExteriorValue = ExteriorValue;
-			RadialIntMaskField->SetMaskCondition = (ESetMaskConditionType)SetMaskCondition;
+			// Compute the field
+			TArray<int32> NewFieldIntResult;
+			UniformIntegerFieldProcess(InSamplePositions, InMagnitude, NewFieldIntResult);
 
-			FFieldContext FieldContext{
-				ExecutionDatas,
-				FFieldContext::UniquePointerMap(),
-				0.0
-			};
+			// Set the outputs
+			SetValue<TArray<int32>>(Context, NewFieldIntResult, &FieldIntResult);
+			SetValue<TArray<int32>>(Context, TArray<int32>(), &FieldRemap);
+			SetValue<int32>(Context, InSamplePositions.Num(), &NumSamplePositions);
 
-			TArray<int32> ResultsArray;
-			ResultsArray.Init(false, NumVertices);
-			TFieldArrayView<int32> ResultsView(ResultsArray, 0, ResultsArray.Num());
-			RadialIntMaskField->Evaluate(FieldContext, ResultsView);
+			return;
+		}
+		//
+		// SampleIndices input connected, the selected elements from SamplePositions array
+		// will be processed and FieldRemapOutput output will contain the remap info
+		// IMPORTANT: Number of elements in SamplePositions and SampleIndices must be the same
+		//
+		else
+		{
+			if (InSamplePositions.Num() == InSampleIndices.Num())
+			{
+				TArray<FVector3f> NewSamplePositions;
+				TArray<int32> NewRemapArray;
+
+				FieldComputeRemap(InSamplePositions, InSampleIndices, NewSamplePositions, NewRemapArray);
 
-			TArray<int32> NewMaskArray;
-			NewMaskArray.Init(0, NumVertices);
+				// Compute the field
+				TArray<int32> NewFieldIntResult;
+				UniformIntegerFieldProcess(NewSamplePositions, InMagnitude, NewFieldIntResult);
 
-			if (IsConnected<FDataflowVertexSelection>(&VertexSelection))
+				// Set the outputs
+				SetValue<TArray<int32>>(Context, NewFieldIntResult, &FieldIntResult);
+				SetValue<TArray<int32>>(Context, NewRemapArray, &FieldRemap);
+				SetValue<int32>(Context, InSamplePositions.Num(), &NumSamplePositions);
+
+				return;
+			}
+		}
+
+		SetValue<TArray<int32>>(Context, TArray<int32>(), &FieldIntResult);
+		SetValue<TArray<int32>>(Context, TArray<int32>(), &FieldRemap);
+		SetValue<int32>(Context, 0, &NumSamplePositions);
+	}
+}
+
+// --------------------------------------------------------------------------------------------------------------------
+
+static void WaveScalarFieldEvaluate(const TArray<FVector3f>& InSamplePositions,
+	float InMagnitude,
+	const FVector InPosition,
+	float InWavelength,
+	float InPeriod,
+	const EDataflowWaveFunctionType InFunctionType,
+	const EDataflowFieldFalloffType InFalloffType,
+	TFieldArrayView<float>& OutResultsView)
+{
+	const int32 NumInSamplePositions = InSamplePositions.Num();
+
+	// Setup the data for the field 
+	FFieldExecutionDatas ExecutionDatas;
+
+	FFieldContextIndex::ContiguousIndices(ExecutionDatas.SampleIndices, NumInSamplePositions);
+	ExecutionDatas.SamplePositions.Init(FVector(0.f), NumInSamplePositions);
+
+	for (int32 Idx = 0; Idx < NumInSamplePositions; Idx++)
+	{
+		ExecutionDatas.SamplePositions[Idx] = (FVector)InSamplePositions[Idx];
+	}
+
+	FWaveScalar WaveScalarField(InMagnitude, InPosition, InWavelength, InPeriod, (EWaveFunctionType)InFunctionType, (EFieldFalloffType)InFalloffType);
+
+	FFieldContext FieldContext{
+		ExecutionDatas,
+		FFieldContext::UniquePointerMap(),
+		0.0
+	};
+
+	// Evalute the field
+	WaveScalarField.Evaluate(FieldContext, OutResultsView);
+}
+
+static void WaveScalarFieldProcess(const TArray<FVector3f>& InSamplePositions,
+	float InMagnitude,
+	const FVector InPosition,
+	const FVector InTranslation,
+	float InWavelength,
+	float InPeriod,
+	const EDataflowWaveFunctionType InFunctionType,
+	const EDataflowFieldFalloffType InFalloffType,
+	TArray<float>& OutFieldFloatResult)
+{
+	const int32 NumInSamplePositions = InSamplePositions.Num();
+
+	// Compute the field
+	TArray<float> ResultsArray;
+	ResultsArray.Init(0.f, NumInSamplePositions);
+	TFieldArrayView<float> ResultsView(ResultsArray, 0, ResultsArray.Num());
+
+	WaveScalarFieldEvaluate(InSamplePositions,
+		InMagnitude,
+		InPosition + InTranslation,
+		InWavelength,
+		InPeriod,
+		InFunctionType,
+		InFalloffType,
+		ResultsView);
+
+	// Set the outputs
+	OutFieldFloatResult.Init(0.f, NumInSamplePositions);
+
+	for (int32 Idx = 0; Idx < NumInSamplePositions; ++Idx)
+	{
+		OutFieldFloatResult[Idx] = ResultsView[Idx];
+	}
+}
+
+void FWaveScalarFieldDataflowNode::Evaluate(Dataflow::FContext& Context, const FDataflowOutput* Out) const
+{
+	if (Out->IsA<TArray<float>>(&FieldFloatResult) ||
+		Out->IsA<TArray<int32>>(&FieldRemap) ||
+		Out->IsA<int32>(&NumSamplePositions))
+	{
+		const TArray<FVector3f>& InSamplePositions = GetValue<TArray<FVector3f>>(Context, &SamplePositions);
+		const FDataflowVertexSelection& InSampleIndices = GetValue<FDataflowVertexSelection>(Context, &SampleIndices);
+		const float InMagnitude = GetValue<float>(Context, &Magnitude);
+		const FVector InPosition = GetValue<FVector>(Context, &Position);
+		const FVector InTranslation = GetValue<FVector>(Context, &Translation);
+		const float InWavelength = GetValue<float>(Context, &Wavelength);
+		const float InPeriod = GetValue<float>(Context, &Period);
+
+		//
+		// SampleIndices input not connected, all the elements from SamplePositions array
+		// will be processed and FieldRemapOutput output will be set empty
+		//
+		if (!IsConnected<FDataflowVertexSelection>(&SampleIndices))
+		{
+			// Compute the field
+			TArray<float> NewFieldFloatResult;
+			WaveScalarFieldProcess(InSamplePositions, 
+				InMagnitude,
+				InPosition,
+				InTranslation, 
+				InWavelength,
+				InPeriod,
+				FunctionType, 
+				FalloffType, 
+				NewFieldFloatResult);
+
+			// Set the outputs
+			SetValue<TArray<float>>(Context, NewFieldFloatResult, &FieldFloatResult);
+			SetValue<TArray<int32>>(Context, TArray<int32>(), &FieldRemap);
+			SetValue<int32>(Context, InSamplePositions.Num(), &NumSamplePositions);
+
+			return;
+		}
+		//
+		// SampleIndices input connected, the selected elements from SamplePositions array
+		// will be processed and FieldRemapOutput output will contain the remap info
+		// IMPORTANT: Number of elements in SamplePositions and SampleIndices must be the same
+		//
+		else
+		{
+			if (InSamplePositions.Num() == InSampleIndices.Num())
 			{
-				const FDataflowVertexSelection& InVertexSelection = GetValue<FDataflowVertexSelection>(Context, &VertexSelection);
+				TArray<FVector3f> NewSamplePositions;
+				TArray<int32> NewRemapArray;
+
+				FieldComputeRemap(InSamplePositions, InSampleIndices, NewSamplePositions, NewRemapArray);
+
+				// Compute the field
+				TArray<float> NewFieldFloatResult;
+				WaveScalarFieldProcess(NewSamplePositions,
+					InMagnitude,
+					InPosition,
+					InTranslation,
+					InWavelength,
+					InPeriod,
+					FunctionType,
+					FalloffType,
+					NewFieldFloatResult);
+
+				// Set the outputs
+				SetValue<TArray<float>>(Context, NewFieldFloatResult, &FieldFloatResult);
+				SetValue<TArray<int32>>(Context, NewRemapArray, &FieldRemap);
+				SetValue<int32>(Context, InSamplePositions.Num(), &NumSamplePositions);
+
+				return;
+			}
+		}
 
-				if (InVertexArray.Num() == InVertexSelection.Num())
+		SetValue<TArray<float>>(Context, TArray<float>(), &FieldFloatResult);
+		SetValue<TArray<int32>>(Context, TArray<int32>(), &FieldRemap);
+		SetValue<int32>(Context, 0, &NumSamplePositions);
+	}
+}
+
+// --------------------------------------------------------------------------------------------------------------------
+
+static void SumScalarEvaluate(const TArray<float>& InFieldFloatLeft,
+	const TArray<float>& InFieldFloatRight,
+	TArray<float>& OutNewFieldFloatOutput,
+	const int32 InNumSamplePositions,
+	const EDataflowFloatFieldOperationType InOperation,
+	const float InMagnitude,
+	const bool bInSwapInputs)
+{
+	switch (InOperation)
+	{
+		case EDataflowFloatFieldOperationType::Dataflow_FloatFieldOperationType_Multiply:
+			for (int32 Idx = 0; Idx < InNumSamplePositions; ++Idx)
+			{
+				OutNewFieldFloatOutput[Idx] = InFieldFloatLeft[Idx] * InFieldFloatRight[Idx];
+			}
+			break;
+		case EDataflowFloatFieldOperationType::Dataflow_FloatFieldFalloffType_Divide:
+			if (!bInSwapInputs)
+			{
+				for (int32 Idx = 0; Idx < InNumSamplePositions; ++Idx)
 				{
-					for (int32 Idx = 0; Idx < NumVertices; ++Idx)
+					if (!FMath::IsNearlyZero(InFieldFloatRight[Idx]))
 					{
-						if (InVertexSelection.IsSelected(Idx))
-						{
-							NewMaskArray[Idx] = ResultsView[Idx];
-						}
+						OutNewFieldFloatOutput[Idx] = InFieldFloatLeft[Idx] / InFieldFloatRight[Idx];
 					}
 				}
 			}
 			else
 			{
-				for (int32 Idx = 0; Idx < NumVertices; ++Idx)
+				for (int32 Idx = 0; Idx < InNumSamplePositions; ++Idx)
+				{
+					if (!FMath::IsNearlyZero(InFieldFloatLeft[Idx]))
+					{
+						OutNewFieldFloatOutput[Idx] = InFieldFloatRight[Idx] / InFieldFloatLeft[Idx];
+					}
+				}
+			}
+			break;
+		case EDataflowFloatFieldOperationType::Dataflow_FloatFieldFalloffType_Add:
+			for (int32 Idx = 0; Idx < InNumSamplePositions; ++Idx)
+			{
+				OutNewFieldFloatOutput[Idx] = InFieldFloatLeft[Idx] + InFieldFloatRight[Idx];
+			}
+			break;
+		case EDataflowFloatFieldOperationType::Dataflow_FloatFieldFalloffType_Substract:
+			if (!bInSwapInputs)
+			{
+				for (int32 Idx = 0; Idx < InNumSamplePositions; ++Idx)
+				{
+					OutNewFieldFloatOutput[Idx] = InFieldFloatLeft[Idx] - InFieldFloatRight[Idx];
+				}
+			}
+			else
+			{
+				for (int32 Idx = 0; Idx < InNumSamplePositions; ++Idx)
 				{
-					NewMaskArray[Idx] = ResultsView[Idx];
+					OutNewFieldFloatOutput[Idx] = InFieldFloatRight[Idx] - InFieldFloatLeft[Idx];
 				}
 			}
+			break;
+		case EDataflowFloatFieldOperationType::Dataflow_FloatFieldFalloffType_Min:
+			for (int32 Idx = 0; Idx < InNumSamplePositions; ++Idx)
+			{
+				OutNewFieldFloatOutput[Idx] = FMath::Min(InFieldFloatLeft[Idx], InFieldFloatRight[Idx]);
+			}
+			break;
+		case EDataflowFloatFieldOperationType::Dataflow_FloatFieldFalloffType_Max:
 
-			SetValue<TArray<int32>>(Context, NewMaskArray, &MaskArray);
+			for (int32 Idx = 0; Idx < InNumSamplePositions; ++Idx)
+			{
+				OutNewFieldFloatOutput[Idx] = FMath::Max(InFieldFloatLeft[Idx], InFieldFloatRight[Idx]);
+			}
+			break;
+	}
 
-			return;
+	if (InMagnitude != 1.0)
+	{
+		for (int32 Idx = 0; Idx < InNumSamplePositions; ++Idx)
+		{
+			OutNewFieldFloatOutput[Idx] *= InMagnitude;
 		}
-
-		SetValue<TArray<int32>>(Context, TArray<int32>(), &MaskArray);
 	}
 }
 
 
-void FUniformScalarFieldDataflowNode::Evaluate(Dataflow::FContext& Context, const FDataflowOutput* Out) const
+void FSumScalarFieldDataflowNode::Evaluate(Dataflow::FContext& Context, const FDataflowOutput* Out) const
 {
-	if (Out->IsA<TArray<float>>(&WeightArray))
+	if (Out->IsA<TArray<float>>(&FieldFloatResult) ||
+		Out->IsA<TArray<int32>>(&FieldRemap))
 	{
-		const TArray<FVector3f>& InVertexArray = GetValue<TArray<FVector3f>>(Context, &VertexArray);
-
-		const int32 NumVertices = InVertexArray.Num();
+		const TArray<float>& InFieldFloatLeft = GetValue<TArray<float>>(Context, &FieldFloatLeft);
+		const TArray<int32>& InFieldRemapLeft = GetValue<TArray<int32>>(Context, &FieldRemapLeft);
+		const TArray<float>& InFieldFloatRight = GetValue<TArray<float>>(Context, &FieldFloatRight);
+		const TArray<int32>& InFieldRemapRight = GetValue<TArray<int32>>(Context, &FieldRemapRight);
 
-		FFieldExecutionDatas ExecutionDatas;
+		const int32 NumFieldFloatLeft = InFieldFloatLeft.Num();
+		const int32 NumFieldRemapLeft = InFieldRemapLeft.Num();
+		const int32 NumFieldFloatRight = InFieldFloatRight.Num();
+		const int32 NumFieldRemapRight = InFieldRemapRight.Num();
 
-		FFieldContextIndex::ContiguousIndices(ExecutionDatas.SampleIndices, NumVertices);
-		ExecutionDatas.SamplePositions.Init(FVector(0.f), NumVertices);
+		const bool IsFieldRemapLeftConnected = IsConnected<TArray<int32>>(&FieldRemapLeft);
+		const bool IsFieldRemapRightConnected = IsConnected<TArray<int32>>(&FieldRemapRight);
 
-		for (int32 Index = 0; Index < NumVertices; Index++)
-		{
-			ExecutionDatas.SamplePositions[Index] = (FVector)InVertexArray[Index];
-		}
+		TArray<float> NewFieldFloatOutput;
 
-		if (TUniquePtr<FUniformScalar> UniformScalarField = TUniquePtr<FUniformScalar>(new FUniformScalar()))
+		// Both Remap inputs are connected but they can be different sizes
+		// Process the two FloatArray inputs and generate a new output remap
+		if (IsFieldRemapLeftConnected && IsFieldRemapRightConnected)
 		{
-			UniformScalarField->Magnitude = Magnitude;
+			if (NumFieldFloatLeft == NumFieldRemapLeft &&
+				NumFieldFloatRight == NumFieldRemapRight)
+			{
+				int32 NewArraySize = FMath::Max(NumFieldFloatLeft, NumFieldFloatRight);
 
-			FFieldContext FieldContext{
-				ExecutionDatas,
-				FFieldContext::UniquePointerMap(),
-				0.0
-			};
+				NewFieldFloatOutput.Init(false, NewArraySize);
 
-			TArray<float> ResultsArray;
-			ResultsArray.Init(false, NumVertices);
-			TFieldArrayView<float> ResultsView(ResultsArray, 0, ResultsArray.Num());
-			UniformScalarField->Evaluate(FieldContext, ResultsView);
+				TArray<int32> NewRemapArray;
+				NewRemapArray.Init(false, NewArraySize);
 
-			TArray<float> NewWeightArray;
-			NewWeightArray.Init(0.f, NumVertices);
+				TArray<float> NewFloatArrayLeft, NewFloatArrayRight;
+				NewFloatArrayLeft.Init(false, NewArraySize);
+				NewFloatArrayRight.Init(false, NewArraySize);
 
-			if (IsConnected<FDataflowVertexSelection>(&VertexSelection))
-			{
-				const FDataflowVertexSelection& InVertexSelection = GetValue<FDataflowVertexSelection>(Context, &VertexSelection);
+				int32 NewSampleCount = 0;
 
-				if (InVertexArray.Num() == InVertexSelection.Num())
+				for (int32 IndexLeft = 0; IndexLeft < NumFieldFloatLeft; ++IndexLeft)
 				{
-					for (int32 Idx = 0; Idx < NumVertices; ++Idx)
+					int32 SampleIndex = InFieldRemapLeft[IndexLeft];
+
+					if (InFieldRemapRight.Contains(SampleIndex))
 					{
-						if (InVertexSelection.IsSelected(Idx))
-						{
-							NewWeightArray[Idx] = ResultsView[Idx];
-						}
+						int32 IndexRight = InFieldRemapRight.Find(SampleIndex);
+
+						NewRemapArray[NewSampleCount] = SampleIndex;
+
+						NewFloatArrayLeft[NewSampleCount] = InFieldFloatLeft[IndexLeft];
+						NewFloatArrayRight[NewSampleCount] = InFieldFloatRight[IndexRight];
+
+						NewSampleCount++;
 					}
 				}
+
+				NewRemapArray.SetNum(NewSampleCount, true);
+				NewFloatArrayLeft.SetNum(NewSampleCount, true);
+				NewFloatArrayRight.SetNum(NewSampleCount, true);
+				NewFieldFloatOutput.SetNum(NewSampleCount, true);
+
+				SumScalarEvaluate(NewFloatArrayLeft,
+					NewFloatArrayRight,
+					NewFieldFloatOutput,
+					NewSampleCount,
+					Operation,
+					Magnitude,
+					bSwapInputs);
+
+				SetValue<TArray<float>>(Context, NewFieldFloatOutput, &FieldFloatResult);
+				SetValue<TArray<int32>>(Context, NewRemapArray, &FieldRemap);
+
+				return;
 			}
-			else
+		}
+		// At least one of the Remap inputs are not connected
+		// There is no remap information, just process the two FloatArray inputs
+		// but they have to be the same size
+		else
+		{
+			if (NumFieldFloatLeft == NumFieldFloatRight)
 			{
-				for (int32 Idx = 0; Idx < NumVertices; ++Idx)
-				{
-					NewWeightArray[Idx] = ResultsView[Idx];
-				}
-			}
+				const int32 NumSamplePositions = NumFieldFloatLeft;
 
-			SetValue<TArray<float>>(Context, NewWeightArray, &WeightArray);
+				NewFieldFloatOutput.Init(false, NumSamplePositions);
 
-			return;
+				SumScalarEvaluate(InFieldFloatLeft,
+					InFieldFloatRight,
+					NewFieldFloatOutput,
+					NumSamplePositions,
+					Operation,
+					Magnitude,
+					bSwapInputs);
+
+				SetValue<TArray<float>>(Context, NewFieldFloatOutput, &FieldFloatResult);
+				SetValue<TArray<int32>>(Context, TArray<int32>(), &FieldRemap);
+
+				return;
+			}
 		}
 
-		SetValue<TArray<float>>(Context, TArray<float>(), &WeightArray);
+		SetValue<TArray<float>>(Context, TArray<float>(), &FieldFloatResult);
+		SetValue<TArray<int32>>(Context, TArray<int32>(), &FieldRemap);
 	}
 }
 
+// --------------------------------------------------------------------------------------------------------------------
 
-void FUniformVectorFieldDataflowNode::Evaluate(Dataflow::FContext& Context, const FDataflowOutput* Out) const
+static void SumVectorEvaluate(const TArray<float>* InFieldFloat,
+	const TArray<FVector>* InFieldVectorLeft,
+	const TArray<FVector>* InFieldVectorRight,
+	TArray<FVector>& OutNewFieldVectorResult,
+	const int32 InNumSamplePositions,
+	const EDataflowVectorFieldOperationType InOperation,
+	const float InMagnitude,
+	const bool bSwapVectorInputs)
 {
-	if (Out->IsA<TArray<FVector>>(&VectorArray))
+	if (InFieldVectorLeft != nullptr && InFieldVectorRight != nullptr)
 	{
-		const TArray<FVector3f>& InVertexArray = GetValue<TArray<FVector3f>>(Context, &VertexArray);
-
-		const int32 NumVertices = InVertexArray.Num();
-
-		FFieldExecutionDatas ExecutionDatas;
-
-		FFieldContextIndex::ContiguousIndices(ExecutionDatas.SampleIndices, NumVertices);
-		ExecutionDatas.SamplePositions.Init(FVector(0.f), NumVertices);
-
-		for (int32 Index = 0; Index < NumVertices; Index++)
-		{
-			ExecutionDatas.SamplePositions[Index] = (FVector)InVertexArray[Index];
-		}
-
-		if (TUniquePtr<FUniformVector> UniformVectorField = TUniquePtr<FUniformVector>(new FUniformVector()))
+		switch (InOperation)
 		{
-			UniformVectorField->Magnitude = Magnitude;
-			UniformVectorField->Direction = Direction;
-
-			FFieldContext FieldContext{
-				ExecutionDatas,
-				FFieldContext::UniquePointerMap(),
-				0.0
-			};
-
-			TArray<FVector> ResultsArray;
-			ResultsArray.Init(FVector(0.f), NumVertices);
-			TFieldArrayView<FVector> ResultsView(ResultsArray, 0, ResultsArray.Num());
-			UniformVectorField->Evaluate(FieldContext, ResultsView);
-
-			TArray<FVector> NewVectorArray;
-			NewVectorArray.Init(FVector(0.f), NumVertices);
-
-			if (IsConnected<FDataflowVertexSelection>(&VertexSelection))
+		case EDataflowVectorFieldOperationType::Dataflow_VectorFieldOperationType_Multiply:
+			for (int32 Idx = 0; Idx < InNumSamplePositions; ++Idx)
 			{
-				const FDataflowVertexSelection& InVertexSelection = GetValue<FDataflowVertexSelection>(Context, &VertexSelection);
-
-				if (InVertexArray.Num() == InVertexSelection.Num())
+				OutNewFieldVectorResult[Idx] = (*InFieldVectorLeft)[Idx] * (*InFieldVectorRight)[Idx];
+			}
+			break;
+		case EDataflowVectorFieldOperationType::Dataflow_VectorFieldFalloffType_Divide:
+			if (!bSwapVectorInputs)
+			{
+				for (int32 Idx = 0; Idx < InNumSamplePositions; ++Idx)
 				{
-					for (int32 Idx = 0; Idx < NumVertices; ++Idx)
+					if (!FMath::IsNearlyZero((*InFieldVectorRight)[Idx].X) &&
+						!FMath::IsNearlyZero((*InFieldVectorRight)[Idx].Y) &&
+						!FMath::IsNearlyZero((*InFieldVectorRight)[Idx].Z))
 					{
-						if (InVertexSelection.IsSelected(Idx))
-						{
-							NewVectorArray[Idx] = ResultsView[Idx];
-						}
+						OutNewFieldVectorResult[Idx] = (*InFieldVectorLeft)[Idx] / (*InFieldVectorRight)[Idx];
 					}
 				}
 			}
 			else
 			{
-				for (int32 Idx = 0; Idx < NumVertices; ++Idx)
+				for (int32 Idx = 0; Idx < InNumSamplePositions; ++Idx)
 				{
-					NewVectorArray[Idx] = ResultsView[Idx];
+					if (!FMath::IsNearlyZero((*InFieldVectorLeft)[Idx].X) &&
+						!FMath::IsNearlyZero((*InFieldVectorLeft)[Idx].Y) &&
+						!FMath::IsNearlyZero((*InFieldVectorLeft)[Idx].Z))
+					{
+						OutNewFieldVectorResult[Idx] = (*InFieldVectorRight)[Idx] / (*InFieldVectorLeft)[Idx];
+					}
 				}
-			}
-
-			SetValue<TArray<FVector>>(Context, NewVectorArray, &VectorArray);
-
-			return;
-		}
-
-		SetValue<TArray<FVector>>(Context, TArray<FVector>(), &VectorArray);
-	}
-}
-
-
-void FRadialVectorFieldDataflowNode::Evaluate(Dataflow::FContext& Context, const FDataflowOutput* Out) const
-{
-	if (Out->IsA<TArray<FVector>>(&VectorArray))
-	{
-		const TArray<FVector3f>& InVertexArray = GetValue<TArray<FVector3f>>(Context, &VertexArray);
-
-		const int32 NumVertices = InVertexArray.Num();
-
-		FFieldExecutionDatas ExecutionDatas;
-
-		FFieldContextIndex::ContiguousIndices(ExecutionDatas.SampleIndices, NumVertices);
-		ExecutionDatas.SamplePositions.Init(FVector(0.f), NumVertices);
-
-		for (int32 Index = 0; Index < NumVertices; Index++)
-		{
-			ExecutionDatas.SamplePositions[Index] = (FVector)InVertexArray[Index];
-		}
-
-		if (TUniquePtr<FRadialVector> RadialVectorField = TUniquePtr<FRadialVector>(new FRadialVector()))
-		{
-			RadialVectorField->Magnitude = Magnitude;
-			RadialVectorField->Position = Position;
-
-			FFieldContext FieldContext{
-				ExecutionDatas,
-				FFieldContext::UniquePointerMap(),
-				0.0
-			};
-
-			TArray<FVector> ResultsArray;
-			ResultsArray.Init(FVector(0.f), NumVertices);
-			TFieldArrayView<FVector> ResultsView(ResultsArray, 0, ResultsArray.Num());
-			RadialVectorField->Evaluate(FieldContext, ResultsView);
 
-			TArray<FVector> NewVectorArray;
-			NewVectorArray.Init(FVector(0.f), NumVertices);
-
-			if (IsConnected<FDataflowVertexSelection>(&VertexSelection))
+			}
+			break;
+		case EDataflowVectorFieldOperationType::Dataflow_VectorFieldFalloffType_Add:
+			for (int32 Idx = 0; Idx < InNumSamplePositions; ++Idx)
 			{
-				const FDataflowVertexSelection& InVertexSelection = GetValue<FDataflowVertexSelection>(Context, &VertexSelection);
-
-				if (InVertexArray.Num() == InVertexSelection.Num())
+				OutNewFieldVectorResult[Idx] = (*InFieldVectorLeft)[Idx] + (*InFieldVectorRight)[Idx];
+			}
+			break;
+		case EDataflowVectorFieldOperationType::Dataflow_VectorFieldFalloffType_Substract:
+			if (!bSwapVectorInputs)
+			{
+				for (int32 Idx = 0; Idx < InNumSamplePositions; ++Idx)
 				{
-					for (int32 Idx = 0; Idx < NumVertices; ++Idx)
-					{
-						if (InVertexSelection.IsSelected(Idx))
-						{
-							NewVectorArray[Idx] = ResultsView[Idx];
-						}
-					}
+					OutNewFieldVectorResult[Idx] = (*InFieldVectorLeft)[Idx] - (*InFieldVectorRight)[Idx];
 				}
 			}
 			else
 			{
-				for (int32 Idx = 0; Idx < NumVertices; ++Idx)
+				for (int32 Idx = 0; Idx < InNumSamplePositions; ++Idx)
+				{
+					OutNewFieldVectorResult[Idx] = (*InFieldVectorRight)[Idx] - (*InFieldVectorLeft)[Idx];
+				}
+
+			}
+			break;
+		case EDataflowVectorFieldOperationType::Dataflow_VectorFieldFalloffType_CrossProduct:
+			if (!bSwapVectorInputs)
+			{
+				for (int32 Idx = 0; Idx < InNumSamplePositions; ++Idx)
 				{
-					NewVectorArray[Idx] = ResultsView[Idx];
+					OutNewFieldVectorResult[Idx] = (*InFieldVectorLeft)[Idx] ^ (*InFieldVectorRight)[Idx];
 				}
 			}
+			else
+			{
+				for (int32 Idx = 0; Idx < InNumSamplePositions; ++Idx)
+				{
+					OutNewFieldVectorResult[Idx] = (*InFieldVectorRight)[Idx] ^ (*InFieldVectorLeft)[Idx];
+				}
 
-			SetValue<TArray<FVector>>(Context, NewVectorArray, &VectorArray);
+			}
+			break;
+		}
+	}
+	else if (InFieldVectorLeft != nullptr)
+	{
+		for (int32 Idx = 0; Idx < InNumSamplePositions; ++Idx)
+		{
+			OutNewFieldVectorResult[Idx] = (*InFieldVectorLeft)[Idx];
+		}
+	}
+	else if (InFieldVectorRight != nullptr)
+	{
+		for (int32 Idx = 0; Idx < InNumSamplePositions; ++Idx)
+		{
+			OutNewFieldVectorResult[Idx] = (*InFieldVectorRight)[Idx];
+		}
+	}
 
-			return;
+	if (InFieldFloat != nullptr)
+	{
+		for (int32 Idx = 0; Idx < InNumSamplePositions; ++Idx)
+		{
+			OutNewFieldVectorResult[Idx] *= (*InFieldFloat)[Idx];
 		}
+	}
 
-		SetValue<TArray<FVector>>(Context, TArray<FVector>(), &VectorArray);
+	if (InMagnitude != 1.0)
+	{
+		for (int32 Idx = 0; Idx < InNumSamplePositions; ++Idx)
+		{
+			OutNewFieldVectorResult[Idx] *= InMagnitude;
+		}
 	}
 }
 
-
-void FRandomVectorFieldDataflowNode::Evaluate(Dataflow::FContext& Context, const FDataflowOutput* Out) const
+void FSumVectorFieldDataflowNode::Evaluate(Dataflow::FContext& Context, const FDataflowOutput* Out) const
 {
-	if (Out->IsA<TArray<FVector>>(&VectorArray))
+	if (Out->IsA<TArray<FVector>>(&FieldVectorResult) ||
+		Out->IsA<TArray<int32>>(&FieldRemap))
 	{
-		const TArray<FVector3f>& InVertexArray = GetValue<TArray<FVector3f>>(Context, &VertexArray);
-
-		const int32 NumVertices = InVertexArray.Num();
+		const TArray<float>& InFieldFloat = GetValue<TArray<float>>(Context, &FieldFloat);
+		const TArray<int32>& InFieldFloatRemap = GetValue<TArray<int32>>(Context, &FieldFloatRemap);
+		const TArray<FVector>& InFieldVectorLeft = GetValue<TArray<FVector>>(Context, &FieldVectorLeft);
+		const TArray<int32>& InFieldRemapLeft = GetValue<TArray<int32>>(Context, &FieldRemapLeft);
+		const TArray<FVector>& InFieldVectorRight = GetValue<TArray<FVector>>(Context, &FieldVectorRight);
+		const TArray<int32>& InFieldRemapRight = GetValue<TArray<int32>>(Context, &FieldRemapRight);
+
+		const int32 NumFieldFloat = InFieldFloat.Num();
+		const int32 NumFieldFloatRemap = InFieldFloatRemap.Num();
+		const int32 NumFieldVectorLeft = InFieldVectorLeft.Num();
+		const int32 NumFieldRemapLeft = InFieldRemapLeft.Num();
+		const int32 NumFieldVectorRight = InFieldVectorRight.Num();
+		const int32 NumFieldRemapRight = InFieldRemapRight.Num();
+
+		const bool IsFieldFloatConnected = IsConnected<TArray<float>>(&FieldFloat);
+		const bool IsFieldFloatRemapConnected = IsConnected<TArray<int32>>(&FieldFloatRemap);
+		const bool IsFieldVectorLeftConnected = IsConnected<TArray<FVector>>(&FieldVectorLeft);
+		const bool IsFieldRemapLeftConnected = IsConnected<TArray<int32>>(&FieldRemapLeft);
+		const bool IsFieldVectorRightConnected = IsConnected<TArray<FVector>>(&FieldVectorRight);
+		const bool IsFieldRemapRightConnected = IsConnected<TArray<int32>>(&FieldRemapRight);
+
+		TArray<FVector> NewFieldVectorResult;
+
+		if (IsFieldFloatConnected && IsFieldVectorLeftConnected && !IsFieldVectorRightConnected)
+		{
+			// Both Remap inputs are connected but they can be different sizes
+			// Process the two array inputs and generate a new output remap
+			if (IsFieldFloatRemapConnected && IsFieldRemapLeftConnected)
+			{
+				if (NumFieldFloat == NumFieldFloatRemap &&
+					NumFieldVectorLeft == NumFieldRemapLeft)
+				{
+					int32 NewArraySize = FMath::Max(NumFieldFloat, NumFieldVectorLeft);
 
-		FFieldExecutionDatas ExecutionDatas;
+					NewFieldVectorResult.Init(FVector(0.f), NewArraySize);
 
-		FFieldContextIndex::ContiguousIndices(ExecutionDatas.SampleIndices, NumVertices);
-		ExecutionDatas.SamplePositions.Init(FVector(0.f), NumVertices);
+					TArray<int32> NewRemapArray;
+					NewRemapArray.Init(false, NewArraySize);
 
-		for (int32 Index = 0; Index < NumVertices; Index++)
-		{
-			ExecutionDatas.SamplePositions[Index] = (FVector)InVertexArray[Index];
-		}
+					TArray<float> NewFloatArray;
+					TArray<FVector> NewVectorArrayLeft;
+					NewFloatArray.Init(false, NewArraySize);
+					NewVectorArrayLeft.Init(FVector(0.f), NewArraySize);
 
-		if (TUniquePtr<FRandomVector> RandomVectorField = TUniquePtr<FRandomVector>(new FRandomVector()))
-		{
-			RandomVectorField->Magnitude = Magnitude;
+					int32 NewSampleCount = 0;
 
-			FFieldContext FieldContext{
-				ExecutionDatas,
-				FFieldContext::UniquePointerMap(),
-				0.0
-			};
+					for (int32 IndexLeft = 0; IndexLeft < NumFieldFloat; ++IndexLeft)
+					{
+						int32 SampleIndex = InFieldFloatRemap[IndexLeft];
 
-			TArray<FVector> ResultsArray;
-			ResultsArray.Init(FVector(0.f), NumVertices);
-			TFieldArrayView<FVector> ResultsView(ResultsArray, 0, ResultsArray.Num());
-			RandomVectorField->Evaluate(FieldContext, ResultsView);
+						if (InFieldRemapLeft.Contains(SampleIndex))
+						{
+							int32 IndexRight = InFieldRemapLeft.Find(SampleIndex);
 
-			TArray<FVector> NewVectorArray;
-			NewVectorArray.Init(FVector(0.f), NumVertices);
+							NewRemapArray[NewSampleCount] = SampleIndex;
 
-			if (IsConnected<FDataflowVertexSelection>(&VertexSelection))
-			{
-				const FDataflowVertexSelection& InVertexSelection = GetValue<FDataflowVertexSelection>(Context, &VertexSelection);
+							NewFloatArray[NewSampleCount] = InFieldFloat[IndexLeft];
+							NewVectorArrayLeft[NewSampleCount] = InFieldVectorLeft[IndexRight];
 
-				if (InVertexArray.Num() == InVertexSelection.Num())
-				{
-					for (int32 Idx = 0; Idx < NumVertices; ++Idx)
-					{
-						if (InVertexSelection.IsSelected(Idx))
-						{
-							NewVectorArray[Idx] = ResultsView[Idx];
+							NewSampleCount++;
 						}
 					}
+
+					NewRemapArray.SetNum(NewSampleCount, true);
+					NewFloatArray.SetNum(NewSampleCount, true);
+					NewVectorArrayLeft.SetNum(NewSampleCount, true);
+					NewFieldVectorResult.SetNum(NewSampleCount, true);
+
+					SumVectorEvaluate(&NewFloatArray,
+						&NewVectorArrayLeft,
+						nullptr,
+						NewFieldVectorResult,
+						NewSampleCount,
+						Operation,
+						Magnitude,
+						bSwapVectorInputs);
+
+					SetValue<TArray<FVector>>(Context, NewFieldVectorResult, &FieldVectorResult);
+					SetValue<TArray<int32>>(Context, NewRemapArray, &FieldRemap);
+
+					return;
 				}
 			}
+			// At least one of the Remap inputs are not connected
+			// There is no remap information, just process the two FloatArray inputs
+			// but they have to be the same size
 			else
 			{
-				for (int32 Idx = 0; Idx < NumVertices; ++Idx)
+				if (NumFieldFloat == NumFieldVectorLeft)
 				{
-					NewVectorArray[Idx] = ResultsView[Idx];
-				}
-			}
+					const int32 NumSamplePositions = NumFieldFloat;
 
-			SetValue<TArray<FVector>>(Context, NewVectorArray, &VectorArray);
+					NewFieldVectorResult.Init(FVector(0.f), NumSamplePositions);
 
-			return;
-		}
+					SumVectorEvaluate(&InFieldFloat,
+						&InFieldVectorLeft,
+						nullptr,
+						NewFieldVectorResult,
+						NumSamplePositions,
+						Operation,
+						Magnitude,
+						bSwapVectorInputs);
 
-		SetValue<TArray<FVector>>(Context, TArray<FVector>(), &VectorArray);
-	}
-}
+					SetValue<TArray<FVector>>(Context, NewFieldVectorResult, &FieldVectorResult);
+					SetValue<TArray<int32>>(Context, TArray<int32>(), &FieldRemap);
 
+					return;
+				}
+			}
+		}
+		else if (IsFieldFloatConnected && !IsFieldVectorLeftConnected && IsFieldVectorRightConnected)
+		{
+			// Both Remap inputs are connected but they can be different sizes
+			// Process the two array inputs and generate a new output remap
+			if (IsFieldFloatRemapConnected && IsFieldRemapRightConnected)
+			{
+				if (NumFieldFloat == NumFieldFloatRemap &&
+					NumFieldVectorRight == NumFieldRemapRight)
+				{
+					int32 NewArraySize = FMath::Max(NumFieldFloat, NumFieldVectorRight);
 
+					NewFieldVectorResult.Init(FVector(0.f), NewArraySize);
 
-void FNoiseFieldDataflowNode::Evaluate(Dataflow::FContext& Context, const FDataflowOutput* Out) const
-{
-	if (Out->IsA<TArray<float>>(&WeightArray))
-	{
-		const TArray<FVector3f>& InVertexArray = GetValue<TArray<FVector3f>>(Context, &VertexArray);
+					TArray<int32> NewRemapArray;
+					NewRemapArray.Init(false, NewArraySize);
 
-		const int32 NumVertices = InVertexArray.Num();
+					TArray<float> NewFloatArray;
+					TArray<FVector> NewVectorArrayRight;
+					NewFloatArray.Init(false, NewArraySize);
+					NewVectorArrayRight.Init(FVector(0.f), NewArraySize);
 
-		FFieldExecutionDatas ExecutionDatas;
+					int32 NewSampleCount = 0;
 
-		FFieldContextIndex::ContiguousIndices(ExecutionDatas.SampleIndices, NumVertices);
-		ExecutionDatas.SamplePositions.Init(FVector(0.f), NumVertices);
+					for (int32 IndexLeft = 0; IndexLeft < NumFieldFloat; ++IndexLeft)
+					{
+						int32 SampleIndex = InFieldFloatRemap[IndexLeft];
 
-		for (int32 Index = 0; Index < NumVertices; Index++)
-		{
-			ExecutionDatas.SamplePositions[Index] = (FVector)InVertexArray[Index];
-		}
+						if (InFieldRemapRight.Contains(SampleIndex))
+						{
+							int32 IndexRight = InFieldRemapRight.Find(SampleIndex);
 
-		if (TUniquePtr<FNoiseField> NoiseField = TUniquePtr<FNoiseField>(new FNoiseField()))
-		{
-			NoiseField->MinRange = MinRange;
-			NoiseField->MaxRange = MaxRange;
-			NoiseField->Transform = Transform;
+							NewRemapArray[NewSampleCount] = SampleIndex;
 
-			FFieldContext FieldContext{
-				ExecutionDatas,
-				FFieldContext::UniquePointerMap(),
-				0.0
-			};
+							NewFloatArray[NewSampleCount] = InFieldFloat[IndexLeft];
+							NewVectorArrayRight[NewSampleCount] = InFieldVectorRight[IndexRight];
 
-			TArray<float> ResultsArray;
-			ResultsArray.Init(false, NumVertices);
-			TFieldArrayView<float> ResultsView(ResultsArray, 0, ResultsArray.Num());
-			NoiseField->Evaluate(FieldContext, ResultsView);
+							NewSampleCount++;
+						}
+					}
 
-			TArray<float> NewWeightArray;
-			NewWeightArray.Init(0.f, NumVertices);
+					NewRemapArray.SetNum(NewSampleCount, true);
+					NewFloatArray.SetNum(NewSampleCount, true);
+					NewVectorArrayRight.SetNum(NewSampleCount, true);
+					NewFieldVectorResult.SetNum(NewSampleCount, true);
 
-			if (IsConnected<FDataflowVertexSelection>(&VertexSelection))
-			{
-				const FDataflowVertexSelection& InVertexSelection = GetValue<FDataflowVertexSelection>(Context, &VertexSelection);
+					SumVectorEvaluate(&NewFloatArray,
+						nullptr,
+						&NewVectorArrayRight,
+						NewFieldVectorResult,
+						NewSampleCount,
+						Operation,
+						Magnitude,
+						bSwapVectorInputs);
 
-				if (InVertexArray.Num() == InVertexSelection.Num())
-				{
-					for (int32 Idx = 0; Idx < NumVertices; ++Idx)
-					{
-						if (InVertexSelection.IsSelected(Idx))
-						{
-							NewWeightArray[Idx] = ResultsView[Idx];
-						}
-					}
+					SetValue<TArray<FVector>>(Context, NewFieldVectorResult, &FieldVectorResult);
+					SetValue<TArray<int32>>(Context, NewRemapArray, &FieldRemap);
+
+					return;
 				}
 			}
+			// At least one of the Remap inputs are not connected
+			// There is no remap information, just process the two FloatArray inputs
+			// but they have to be the same size
 			else
 			{
-				for (int32 Idx = 0; Idx < NumVertices; ++Idx)
+				if (NumFieldFloat == NumFieldVectorRight)
 				{
-					NewWeightArray[Idx] = ResultsView[Idx];
-				}
-			}
+					const int32 NumSamplePositions = NumFieldFloat;
 
-			SetValue<TArray<float>>(Context, NewWeightArray, &WeightArray);
+					NewFieldVectorResult.Init(FVector(0.f), NumSamplePositions);
 
-			return;
+					SumVectorEvaluate(&InFieldFloat,
+						nullptr,
+						&InFieldVectorRight,
+						NewFieldVectorResult,
+						NumSamplePositions,
+						Operation,
+						Magnitude,
+						bSwapVectorInputs);
+
+					SetValue<TArray<FVector>>(Context, NewFieldVectorResult, &FieldVectorResult);
+					SetValue<TArray<int32>>(Context, TArray<int32>(), &FieldRemap);
+
+					return;
+				}
+			}
 		}
+		else if (!IsFieldFloatConnected && IsFieldVectorLeftConnected && IsFieldVectorRightConnected)
+		{
+			// Both Remap inputs are connected but they can be different sizes
+			// Process the two array inputs and generate a new output remap
+			if (IsFieldRemapLeftConnected && IsFieldRemapRightConnected)
+			{
+				if (NumFieldVectorLeft == NumFieldRemapLeft &&
+					NumFieldVectorRight == NumFieldRemapRight)
+				{
+					int32 NewArraySize = FMath::Max(NumFieldVectorLeft, NumFieldVectorRight);
 
-		SetValue<TArray<float>>(Context, TArray<float>(), &WeightArray);
-	}
-}
+					NewFieldVectorResult.Init(FVector(0.f), NewArraySize);
 
+					TArray<int32> NewRemapArray;
+					NewRemapArray.Init(false, NewArraySize);
 
-void FUniformIntegerFieldDataflowNode::Evaluate(Dataflow::FContext& Context, const FDataflowOutput* Out) const
-{
-	if (Out->IsA<TArray<int32>>(&MaskArray))
-	{
-		const TArray<FVector3f>& InVertexArray = GetValue<TArray<FVector3f>>(Context, &VertexArray);
+					TArray<FVector> NewVectorArrayLeft, NewVectorArrayRight;
+					NewVectorArrayLeft.Init(FVector(0.f), NewArraySize);
+					NewVectorArrayLeft.Init(FVector(0.f), NewArraySize);
 
-		const int32 NumVertices = InVertexArray.Num();
+					int32 NewSampleCount = 0;
 
-		FFieldExecutionDatas ExecutionDatas;
+					for (int32 IndexLeft = 0; IndexLeft < NumFieldVectorLeft; ++IndexLeft)
+					{
+						int32 SampleIndex = InFieldRemapLeft[IndexLeft];
 
-		FFieldContextIndex::ContiguousIndices(ExecutionDatas.SampleIndices, NumVertices);
-		ExecutionDatas.SamplePositions.Init(FVector(0.f), NumVertices);
+						if (InFieldRemapRight.Contains(SampleIndex))
+						{
+							int32 IndexRight = InFieldRemapRight.Find(SampleIndex);
 
-		for (int32 Index = 0; Index < NumVertices; Index++)
-		{
-			ExecutionDatas.SamplePositions[Index] = (FVector)InVertexArray[Index];
-		}
+							NewRemapArray[NewSampleCount] = SampleIndex;
 
-		if (TUniquePtr<FUniformInteger> UniformIntegerField = TUniquePtr<FUniformInteger>(new FUniformInteger()))
-		{
-			UniformIntegerField->Magnitude = Magnitude;
+							NewVectorArrayLeft[NewSampleCount] = InFieldVectorLeft[IndexLeft];
+							NewVectorArrayRight[NewSampleCount] = InFieldVectorRight[IndexRight];
 
-			FFieldContext FieldContext{
-				ExecutionDatas,
-				FFieldContext::UniquePointerMap(),
-				0.0
-			};
+							NewSampleCount++;
+						}
+					}
 
-			TArray<int32> ResultsArray;
-			ResultsArray.Init(false, NumVertices);
-			TFieldArrayView<int32> ResultsView(ResultsArray, 0, ResultsArray.Num());
-			UniformIntegerField->Evaluate(FieldContext, ResultsView);
+					NewRemapArray.SetNum(NewSampleCount, true);
+					NewVectorArrayLeft.SetNum(NewSampleCount, true);
+					NewVectorArrayRight.SetNum(NewSampleCount, true);
+					NewFieldVectorResult.SetNum(NewSampleCount, true);
 
-			TArray<int32> NewMaskArray;
-			NewMaskArray.Init(0, NumVertices);
+					SumVectorEvaluate(nullptr,
+						&NewVectorArrayLeft,
+						&NewVectorArrayRight,
+						NewFieldVectorResult,
+						NewSampleCount,
+						Operation,
+						Magnitude,
+						bSwapVectorInputs);
 
-			if (IsConnected<FDataflowVertexSelection>(&VertexSelection))
-			{
-				const FDataflowVertexSelection& InVertexSelection = GetValue<FDataflowVertexSelection>(Context, &VertexSelection);
+					SetValue<TArray<FVector>>(Context, NewFieldVectorResult, &FieldVectorResult);
+					SetValue<TArray<int32>>(Context, NewRemapArray, &FieldRemap);
 
-				if (InVertexArray.Num() == InVertexSelection.Num())
-				{
-					for (int32 Idx = 0; Idx < NumVertices; ++Idx)
-					{
-						if (InVertexSelection.IsSelected(Idx))
-						{
-							NewMaskArray[Idx] = ResultsView[Idx];
-						}
-					}
+					return;
 				}
 			}
+			// At least one of the Remap inputs are not connected
+			// There is no remap information, just process the two FloatArray inputs
+			// but they have to be the same size
 			else
 			{
-				for (int32 Idx = 0; Idx < NumVertices; ++Idx)
+				if (NumFieldVectorLeft == NumFieldVectorRight)
 				{
-					NewMaskArray[Idx] = ResultsView[Idx];
-				}
-			}
-
-			SetValue<TArray<int32>>(Context, NewMaskArray, &MaskArray);
-
-			return;
-		}
-
-		SetValue<TArray<int32>>(Context, TArray<int32>(), &MaskArray);
-	}
-}
-
-
-void FWaveScalarFieldDataflowNode::Evaluate(Dataflow::FContext& Context, const FDataflowOutput* Out) const
-{
-	if (Out->IsA<TArray<float>>(&WeightArray))
-	{
-		const TArray<FVector3f>& InVertexArray = GetValue<TArray<FVector3f>>(Context, &VertexArray);
-		const FVector& InTranslation = GetValue<FVector>(Context, &Translation);
+					const int32 NumSamplePositions = NumFieldVectorLeft;
 
-		const int32 NumVertices = InVertexArray.Num();
+					NewFieldVectorResult.Init(FVector(0.f), NumSamplePositions);
 
-		FFieldExecutionDatas ExecutionDatas;
+					SumVectorEvaluate(nullptr,
+						&InFieldVectorLeft,
+						&InFieldVectorRight,
+						NewFieldVectorResult,
+						NumSamplePositions,
+						Operation,
+						Magnitude,
+						bSwapVectorInputs);
 
-		FFieldContextIndex::ContiguousIndices(ExecutionDatas.SampleIndices, NumVertices);
-		ExecutionDatas.SamplePositions.Init(FVector(0.f), NumVertices);
+					SetValue<TArray<FVector>>(Context, NewFieldVectorResult, &FieldVectorResult);
+					SetValue<TArray<int32>>(Context, TArray<int32>(), &FieldRemap);
 
-		for (int32 Index = 0; Index < NumVertices; Index++)
-		{
-			ExecutionDatas.SamplePositions[Index] = (FVector)InVertexArray[Index];
+					return;
+				}
+			}
 		}
-
-		if (TUniquePtr<FWaveScalar> WaveScalarField = TUniquePtr<FWaveScalar>(new FWaveScalar()))
+		else if (IsFieldFloatConnected && IsFieldVectorLeftConnected && IsFieldVectorRightConnected)
 		{
-			WaveScalarField->Magnitude = Magnitude;
-			WaveScalarField->Position = Position + InTranslation;
-			WaveScalarField->Wavelength = Wavelength;
-			WaveScalarField->Period = Period;
-			WaveScalarField->Function = (EWaveFunctionType)FunctionType;
-			WaveScalarField->Falloff = (EFieldFalloffType)FalloffType;
+			// All 3 Remap inputs are connected but they can be different sizes
+			// Process the 3 array inputs and generate a new output remap
+			if (IsFieldFloatRemapConnected && IsFieldRemapLeftConnected && IsFieldRemapRightConnected)
+			{
+				if (NumFieldFloat == NumFieldFloatRemap && 
+					NumFieldVectorLeft == NumFieldRemapLeft &&
+					NumFieldVectorRight == NumFieldRemapRight)
+				{
+					int32 NewArraySize = FMath::Max3(NumFieldFloat, NumFieldVectorLeft, NumFieldVectorRight);
 
-			FFieldContext FieldContext{
-				ExecutionDatas,
-				FFieldContext::UniquePointerMap(),
-				0.0
-			};
+					NewFieldVectorResult.Init(FVector(0.f), NewArraySize);
 
-			TArray<float> ResultsArray;
-			ResultsArray.Init(false, NumVertices);
-			TFieldArrayView<float> ResultsView(ResultsArray, 0, ResultsArray.Num());
-			WaveScalarField->Evaluate(FieldContext, ResultsView);
+					TArray<int32> NewRemapArray;
+					NewRemapArray.Init(false, NewArraySize);
 
-			TArray<float> NewWeightArray;
-			NewWeightArray.Init(0.f, NumVertices);
+					TArray<float> NewFloatArray;
+					TArray<FVector> NewVectorArrayLeft, NewVectorArrayRight;
+					NewFloatArray.Init(false, NewArraySize);
+					NewVectorArrayLeft.Init(FVector(0.f), NewArraySize);
+					NewVectorArrayRight.Init(FVector(0.f), NewArraySize);
 
-			if (IsConnected<FDataflowVertexSelection>(&VertexSelection))
-			{
-				const FDataflowVertexSelection& InVertexSelection = GetValue<FDataflowVertexSelection>(Context, &VertexSelection);
+					int32 NewSampleCount = 0;
 
-				if (InVertexArray.Num() == InVertexSelection.Num())
-				{
-					for (int32 Idx = 0; Idx < NumVertices; ++Idx)
+					for (int32 IndexLeft = 0; IndexLeft < NumFieldFloat; ++IndexLeft)
 					{
-						if (InVertexSelection.IsSelected(Idx))
+						int32 SampleIndex = InFieldFloatRemap[IndexLeft];
+
+						if (InFieldRemapLeft.Contains(SampleIndex) &&
+							InFieldRemapRight.Contains(SampleIndex))
 						{
-							NewWeightArray[Idx] = ResultsView[Idx];
+							int32 IndexRight1 = InFieldRemapLeft.Find(SampleIndex);
+							int32 IndexRight2 = InFieldRemapRight.Find(SampleIndex);
+
+							NewRemapArray[NewSampleCount] = SampleIndex;
+
+							NewFloatArray[NewSampleCount] = InFieldFloat[IndexLeft];
+							NewVectorArrayLeft[NewSampleCount] = InFieldVectorLeft[IndexRight1];
+							NewVectorArrayRight[NewSampleCount] = InFieldVectorRight[IndexRight2];
+
+							NewSampleCount++;
 						}
 					}
+
+					NewRemapArray.SetNum(NewSampleCount, true);
+					NewFloatArray.SetNum(NewSampleCount, true);
+					NewVectorArrayLeft.SetNum(NewSampleCount, true);
+					NewVectorArrayRight.SetNum(NewSampleCount, true);
+					NewFieldVectorResult.SetNum(NewSampleCount, true);
+
+					SumVectorEvaluate(&NewFloatArray,
+						&NewVectorArrayLeft,
+						&NewVectorArrayRight,
+						NewFieldVectorResult,
+						NewSampleCount,
+						Operation,
+						Magnitude,
+						bSwapVectorInputs);
+
+					SetValue<TArray<FVector>>(Context, NewFieldVectorResult, &FieldVectorResult);
+					SetValue<TArray<int32>>(Context, NewRemapArray, &FieldRemap);
+
+					return;
 				}
 			}
+			// At least one of the Remap inputs are not connected
+			// There is no remap information, just process the two FloatArray inputs
+			// but they have to be the same size
 			else
 			{
-				for (int32 Idx = 0; Idx < NumVertices; ++Idx)
+				if (NumFieldFloat == NumFieldVectorLeft &&
+					NumFieldFloat == NumFieldVectorRight)
 				{
-					NewWeightArray[Idx] = ResultsView[Idx];
-				}
-			}
+					const int32 NumSamplePositions = NumFieldFloat;
 
-			SetValue<TArray<float>>(Context, NewWeightArray, &WeightArray);
+					NewFieldVectorResult.Init(FVector(0.f), NumSamplePositions);
 
-			return;
+					SumVectorEvaluate(&InFieldFloat,
+						&InFieldVectorLeft,
+						&InFieldVectorRight,
+						NewFieldVectorResult,
+						NumSamplePositions,
+						Operation,
+						Magnitude,
+						bSwapVectorInputs);
+
+					SetValue<TArray<FVector>>(Context, NewFieldVectorResult, &FieldVectorResult);
+					SetValue<TArray<int32>>(Context, TArray<int32>(), &FieldRemap);
+
+					return;
+				}
+			}
 		}
 
-		SetValue<TArray<float>>(Context, TArray<float>(), &WeightArray);
+		SetValue<TArray<FVector>>(Context, TArray<FVector>(), &FieldVectorResult);
+		SetValue<TArray<int32>>(Context, TArray<int32>(), &FieldRemap);
 	}
 }
 
+// ----------------------------------------------------------------------------
 
-void FBoxFalloffFieldDataflowNode::Evaluate(Dataflow::FContext& Context, const FDataflowOutput* Out) const
+void FFieldMakeDenseFloatArrayDataflowNode::Evaluate(Dataflow::FContext& Context, const FDataflowOutput* Out) const
 {
-	if (Out->IsA<TArray<float>>(&WeightArray))
+	if (Out->IsA<TArray<float>>(&FieldFloatResult))
 	{
-		const TArray<FVector3f>& InVertexArray = GetValue<TArray<FVector3f>>(Context, &VertexArray);
-		const FBox& InBox = GetValue<FBox>(Context, &Box);
-		const FTransform& InTransform = GetValue<FTransform>(Context, &Transform);
-
-		const int32 NumVertices = InVertexArray.Num();
-
-		FFieldExecutionDatas ExecutionDatas;
+		const TArray<float>& InFieldFloatInput = GetValue<TArray<float>>(Context, &FieldFloatInput);
+		const TArray<int32>& InFieldRemap = GetValue<TArray<int32>>(Context, &FieldRemap);
+		const int32 InNumSamplePositions = GetValue<int32>(Context, &NumSamplePositions);
 
-		FFieldContextIndex::ContiguousIndices(ExecutionDatas.SampleIndices, NumVertices);
-		ExecutionDatas.SamplePositions.Init(FVector(0.f), NumVertices);
-
-		for (int32 Index = 0; Index < NumVertices; Index++)
-		{
-			ExecutionDatas.SamplePositions[Index] = (FVector)InVertexArray[Index];
-		}
-
-		if (TUniquePtr<FBoxFalloff> BoxFalloffField = TUniquePtr<FBoxFalloff>(new FBoxFalloff()))
+		if (InFieldFloatInput.Num() == InFieldRemap.Num() && InNumSamplePositions > 0)
 		{
-			BoxFalloffField->Magnitude = Magnitude;
-			BoxFalloffField->MinRange = MinRange;
-			BoxFalloffField->MaxRange = MaxRange;
-			BoxFalloffField->Default = Default;
-			BoxFalloffField->Transform = InTransform;
-			BoxFalloffField->Falloff = (EFieldFalloffType)FalloffType;
-
-			FFieldContext FieldContext{
-				ExecutionDatas,
-				FFieldContext::UniquePointerMap(),
-				0.0
-			};
-
-			TArray<float> ResultsArray;
-			ResultsArray.Init(false, NumVertices);
-			TFieldArrayView<float> ResultsView(ResultsArray, 0, ResultsArray.Num());
-			BoxFalloffField->Evaluate(FieldContext, ResultsView);
-
-			TArray<float> NewWeightArray;
-			NewWeightArray.Init(0.f, NumVertices);
-
-			if (IsConnected<FDataflowVertexSelection>(&VertexSelection))
-			{
-				const FDataflowVertexSelection& InVertexSelection = GetValue<FDataflowVertexSelection>(Context, &VertexSelection);
+			TArray<float> NewFieldFloatOutput;
+			NewFieldFloatOutput.Init(Default, InNumSamplePositions);
 
-				if (InVertexArray.Num() == InVertexSelection.Num())
-				{
-					for (int32 Idx = 0; Idx < NumVertices; ++Idx)
-					{
-						if (InVertexSelection.IsSelected(Idx))
-						{
-							NewWeightArray[Idx] = ResultsView[Idx];
-						}
-					}
-				}
-			}
-			else
+			for (int32 Idx = 0; Idx < InFieldRemap.Num(); ++Idx)
 			{
-				for (int32 Idx = 0; Idx < NumVertices; ++Idx)
+				int32 RemapIdx = InFieldRemap[Idx];
+				if (RemapIdx < InNumSamplePositions)
 				{
-					NewWeightArray[Idx] = ResultsView[Idx];
+					NewFieldFloatOutput[InFieldRemap[Idx]] = InFieldFloatInput[Idx];
 				}
 			}
 
-			SetValue<TArray<float>>(Context, NewWeightArray, &WeightArray);
+			SetValue<TArray<float>>(Context, NewFieldFloatOutput, &FieldFloatResult);
 
 			return;
 		}
 
-		SetValue<TArray<float>>(Context, TArray<float>(), &WeightArray);
+		SetValue<TArray<float>>(Context, TArray<float>(), &FieldFloatResult);
 	}
 }
 
 
 
-
-
-
-
-
-
-
diff --git a/Engine/Plugins/Experimental/GeometryCollectionPlugin/Source/GeometryCollectionNodes/Private/Dataflow/GeometryCollectionNodes.cpp b/Engine/Plugins/Experimental/GeometryCollectionPlugin/Source/GeometryCollectionNodes/Private/Dataflow/GeometryCollectionNodes.cpp
--- a/Engine/Plugins/Experimental/GeometryCollectionPlugin/Source/GeometryCollectionNodes/Private/Dataflow/GeometryCollectionNodes.cpp
+++ b/Engine/Plugins/Experimental/GeometryCollectionPlugin/Source/GeometryCollectionNodes/Private/Dataflow/GeometryCollectionNodes.cpp
@@ -97,8 +97,12 @@ namespace Dataflow
 		DATAFLOW_NODE_REGISTER_CREATION_FACTORY(FSetCollectionAttributeDataTypedDataflowNode);
 		DATAFLOW_NODE_REGISTER_CREATION_FACTORY(FBoolArrayToFaceSelectionDataflowNode);
 		DATAFLOW_NODE_REGISTER_CREATION_FACTORY(FFloatArrayToVertexSelectionDataflowNode);
-		DATAFLOW_NODE_REGISTER_CREATION_FACTORY(FSetVertexColorInCollectionDataflowNode);
+		DATAFLOW_NODE_REGISTER_CREATION_FACTORY(FSetVertexColorInCollectionFromVertexSelectionDataflowNode);
 		DATAFLOW_NODE_REGISTER_CREATION_FACTORY(FMakeTransformDataflowNode);
+		DATAFLOW_NODE_REGISTER_CREATION_FACTORY(FSetVertexColorInCollectionFromFloatArrayDataflowNode);
+		DATAFLOW_NODE_REGISTER_CREATION_FACTORY(FFloatArrayNormalizeDataflowNode);
+		DATAFLOW_NODE_REGISTER_CREATION_FACTORY(FVectorArrayNormalizeDataflowNode);
+		DATAFLOW_NODE_REGISTER_CREATION_FACTORY(FMakeQuaternionDataflowNode);
 		DATAFLOW_NODE_REGISTER_CREATION_FACTORY(FMultiplyTransformDataflowNode);
 		
 
@@ -115,7 +119,7 @@ namespace Dataflow
 		// Math
 		DATAFLOW_NODE_REGISTER_CREATION_FACTORY_NODE_COLORS_BY_CATEGORY("Math", FLinearColor(0.f, 0.4f, 0.8f), CDefaultNodeBodyTintColor);
 		// Generators
-		DATAFLOW_NODE_REGISTER_CREATION_FACTORY_NODE_COLORS_BY_CATEGORY("Generators", FLinearColor(.6f, 0.1f, 1.f), CDefaultNodeBodyTintColor);
+		DATAFLOW_NODE_REGISTER_CREATION_FACTORY_NODE_COLORS_BY_CATEGORY("Generators", FLinearColor(.4f, 0.8f, 0.f), CDefaultNodeBodyTintColor);
 	}
 }
 
@@ -1024,6 +1028,14 @@ static FName GetGroupName(const EStandardGroupNameEnum& InGroupName)
 	return GroupNameToUse;
 }
 
+
+template<typename T>
+static void AddAndFillAttribute(FManagedArrayCollection& InCollection, FName AttributeName, FName GroupName, const T& DefaultValue)
+{
+	TManagedArrayAccessor<T> CustomAttribute(InCollection, AttributeName, GroupName);
+	CustomAttribute.AddAndFill(DefaultValue);
+}
+
 void FAddCustomCollectionAttributeDataflowNode::Evaluate(Dataflow::FContext& Context, const FDataflowOutput* Out) const
 {
 	if (Out->IsA<FManagedArrayCollection>(&Collection))
@@ -1043,57 +1055,93 @@ void FAddCustomCollectionAttributeDataflowNode::Evaluate(Dataflow::FContext& Con
 
 		if (GroupNameToUse.GetStringLength() > 0 && AttrName.Len() > 0)
 		{
-			if (CustomAttributeType == ECustomAttributeTypeEnum::Dataflow_CustomAttributeType_Bool)
+			// If the group already exists don't change the number of elements
+			if (!InCollection.HasGroup(GroupNameToUse))
 			{
-				TManagedArrayAccessor<bool> CustomAttribute(InCollection, FName(*AttrName), GroupNameToUse);
-
-				// If the group already exists don't change the number of elements
-				if (!InCollection.HasGroup(GroupNameToUse))
-				{
-					CustomAttribute.AddElements(InNumElements);
-				}
-
-				CustomAttribute.AddAndFill(false);
-			}
-			else if (CustomAttributeType == ECustomAttributeTypeEnum::Dataflow_CustomAttributeType_Float)
-			{
-				TManagedArrayAccessor<float> CustomAttribute(InCollection, FName(*AttrName), GroupNameToUse);
-
-				// If the group already exists don't change the number of elements
-				if (!InCollection.HasGroup(GroupNameToUse))
-				{
-					CustomAttribute.AddElements(InNumElements);
-				}
-
-				CustomAttribute.AddAndFill(0.f);
+				InCollection.AddGroup(GroupNameToUse);
+				InCollection.AddElements(InNumElements, GroupNameToUse);
 			}
-			else if (CustomAttributeType == ECustomAttributeTypeEnum::Dataflow_CustomAttributeType_Int)
-			{
-				TManagedArrayAccessor<int32> CustomAttribute(InCollection, FName(*AttrName), GroupNameToUse);
 
-				// If the group already exists don't change the number of elements
-				if (!InCollection.HasGroup(GroupNameToUse))
-				{
-					CustomAttribute.AddElements(InNumElements);
-				}
+			FName AttributeNameToUse = FName(*AttrName);
 
-				CustomAttribute.AddAndFill(0);
-			}
-			else if (CustomAttributeType == ECustomAttributeTypeEnum::Dataflow_CustomAttributeType_Vector)
+			switch (CustomAttributeType)
 			{
-				TManagedArrayAccessor<FVector> CustomAttribute(InCollection, FName(*AttrName), GroupNameToUse);
-
-				// If the group already exists don't change the number of elements
-				if (!InCollection.HasGroup(GroupNameToUse))
-				{
-					CustomAttribute.AddElements(InNumElements);
-				}
-
-				CustomAttribute.AddAndFill(FVector(EForceInit::ForceInitToZero));
+			case ECustomAttributeTypeEnum::Dataflow_CustomAttributeType_UInt8:
+				AddAndFillAttribute<uint8>(InCollection, AttributeNameToUse, GroupNameToUse, uint8(0));
+				break;
+			case ECustomAttributeTypeEnum::Dataflow_CustomAttributeType_Int32:
+				AddAndFillAttribute<int32>(InCollection, AttributeNameToUse, GroupNameToUse, 0);
+				break;
+			case ECustomAttributeTypeEnum::Dataflow_CustomAttributeType_Float:
+				AddAndFillAttribute<float>(InCollection, AttributeNameToUse, GroupNameToUse, float(0.0));
+				break;
+			case ECustomAttributeTypeEnum::Dataflow_CustomAttributeType_Double:
+				AddAndFillAttribute<double>(InCollection, AttributeNameToUse, GroupNameToUse, double(0.0));
+				break;
+			case ECustomAttributeTypeEnum::Dataflow_CustomAttributeType_Bool:
+				AddAndFillAttribute<bool>(InCollection, AttributeNameToUse, GroupNameToUse, false);
+				break;
+			case ECustomAttributeTypeEnum::Dataflow_CustomAttributeType_String:
+				AddAndFillAttribute<FString>(InCollection, AttributeNameToUse, GroupNameToUse, FString());
+				break;
+			case ECustomAttributeTypeEnum::Dataflow_CustomAttributeType_Vector2f:
+				AddAndFillAttribute<FVector2f>(InCollection, AttributeNameToUse, GroupNameToUse, FVector2f(EForceInit::ForceInitToZero));
+				break;
+			case ECustomAttributeTypeEnum::Dataflow_CustomAttributeType_Vector3f:
+				AddAndFillAttribute<FVector3f>(InCollection, AttributeNameToUse, GroupNameToUse, FVector3f(EForceInit::ForceInitToZero));
+				break;
+			case ECustomAttributeTypeEnum::Dataflow_CustomAttributeType_Vector3d:
+				AddAndFillAttribute<FVector3d>(InCollection, AttributeNameToUse, GroupNameToUse, FVector3d(EForceInit::ForceInitToZero));
+				break;
+			case ECustomAttributeTypeEnum::Dataflow_CustomAttributeType_Vector4f:
+				AddAndFillAttribute<FVector4f>(InCollection, AttributeNameToUse, GroupNameToUse, FVector4f(EForceInit::ForceInitToZero));
+				break;
+			case ECustomAttributeTypeEnum::Dataflow_CustomAttributeType_LinearColor:
+				AddAndFillAttribute<FLinearColor>(InCollection, AttributeNameToUse, GroupNameToUse, FLinearColor(0.f, 0.f, 0.f, 1.f));
+				break;
+			case ECustomAttributeTypeEnum::Dataflow_CustomAttributeType_Transform:
+				AddAndFillAttribute<FTransform>(InCollection, AttributeNameToUse, GroupNameToUse, FTransform(FTransform::Identity));
+				break;
+			case ECustomAttributeTypeEnum::Dataflow_CustomAttributeType_Quat4f:
+				AddAndFillAttribute<FQuat4f>(InCollection, AttributeNameToUse, GroupNameToUse, FQuat4f(ForceInitToZero));
+				break;
+			case ECustomAttributeTypeEnum::Dataflow_CustomAttributeType_Box:
+				AddAndFillAttribute<FBox>(InCollection, AttributeNameToUse, GroupNameToUse, FBox(ForceInit));
+				break;
+			case ECustomAttributeTypeEnum::Dataflow_CustomAttributeType_Guid:
+				AddAndFillAttribute<FGuid>(InCollection, AttributeNameToUse, GroupNameToUse, FGuid());
+				break;
+			case ECustomAttributeTypeEnum::Dataflow_CustomAttributeType_Int32Set:
+				AddAndFillAttribute<TSet<int32>>(InCollection, AttributeNameToUse, GroupNameToUse, TSet<int32>());
+				break;
+			case ECustomAttributeTypeEnum::Dataflow_CustomAttributeType_Int32Array:
+				AddAndFillAttribute<TArray<int32>>(InCollection, AttributeNameToUse, GroupNameToUse, TArray<int32>());
+				break;
+			case ECustomAttributeTypeEnum::Dataflow_CustomAttributeType_IntVector:
+				AddAndFillAttribute<FIntVector>(InCollection, AttributeNameToUse, GroupNameToUse, FIntVector());
+				break;
+			case ECustomAttributeTypeEnum::Dataflow_CustomAttributeType_IntVector2:
+				AddAndFillAttribute<FIntVector2>(InCollection, AttributeNameToUse, GroupNameToUse, FIntVector2());
+				break;
+			case ECustomAttributeTypeEnum::Dataflow_CustomAttributeType_IntVector4:
+				AddAndFillAttribute<FIntVector4>(InCollection, AttributeNameToUse, GroupNameToUse, FIntVector4());
+				break;
+			case ECustomAttributeTypeEnum::Dataflow_CustomAttributeType_IntVector2Array:
+				AddAndFillAttribute<TArray<FIntVector2>>(InCollection, AttributeNameToUse, GroupNameToUse, TArray<FIntVector2>());
+				break;
+			case ECustomAttributeTypeEnum::Dataflow_CustomAttributeType_FloatArray:
+				AddAndFillAttribute<TArray<float>>(InCollection, AttributeNameToUse, GroupNameToUse, TArray<float>());
+				break;
+			case ECustomAttributeTypeEnum::Dataflow_CustomAttributeType_Vector2fArray:
+				AddAndFillAttribute<TArray<FVector2f>>(InCollection, AttributeNameToUse, GroupNameToUse, TArray<FVector2f>());
+				break;
+			case ECustomAttributeTypeEnum::Dataflow_CustomAttributeType_FVector3fArray:
+				AddAndFillAttribute<TArray<FVector3f>>(InCollection, AttributeNameToUse, GroupNameToUse, TArray<FVector3f>());
+				break;
 			}
 		}
 
-		SetValue<FManagedArrayCollection>(Context, InCollection, &Collection);
+		SetValue<FManagedArrayCollection>(Context, MoveTemp(InCollection), &Collection);
 	}
 }
 
@@ -1133,8 +1181,11 @@ void FGetCollectionAttributeDataTypedDataflowNode::Evaluate(Dataflow::FContext&
 {
 	if (Out->IsA<TArray<bool>>(&BoolAttributeData) ||
 		Out->IsA<TArray<float>>(&FloatAttributeData) ||
-		Out->IsA<TArray<int32>>(&IntAttributeData) ||
-		Out->IsA<TArray<FVector3f>>(&VectorAttributeData))
+		Out->IsA<TArray<double>>(&DoubleAttributeData) ||
+		Out->IsA<TArray<int32>>(&Int32AttributeData) ||
+		Out->IsA<TArray<FString>>(&StringAttributeData) ||
+		Out->IsA<TArray<FVector3f>>(&Vector3fAttributeData) ||
+		Out->IsA<TArray<FVector3d>>(&Vector3dAttributeData))
 	{
 		const FManagedArrayCollection& InCollection = GetValue<FManagedArrayCollection>(Context, &Collection);
 
@@ -1148,6 +1199,14 @@ void FGetCollectionAttributeDataTypedDataflowNode::Evaluate(Dataflow::FContext&
 			GroupNameToUse = FName(*CustomGroupName);
 		}
 
+		SetValue<TArray<bool>>(Context, TArray<bool>(), &BoolAttributeData);
+		SetValue<TArray<float>>(Context, TArray<float>(), &FloatAttributeData);
+		SetValue<TArray<double>>(Context, TArray<double>(), &DoubleAttributeData);
+		SetValue<TArray<int32>>(Context, TArray<int32>(), &Int32AttributeData);
+		SetValue<TArray<FString>>(Context, TArray<FString>(), &StringAttributeData);
+		SetValue<TArray<FVector3f>>(Context, TArray<FVector3f>(), &Vector3fAttributeData);
+		SetValue<TArray<FVector3d>>(Context, TArray<FVector3d>(), &Vector3dAttributeData);
+
 		if (GroupNameToUse.GetStringLength() > 0 && AttrName.Len() > 0)
 		{
 			if (InCollection.HasGroup(GroupNameToUse))
@@ -1156,57 +1215,40 @@ void FGetCollectionAttributeDataTypedDataflowNode::Evaluate(Dataflow::FContext&
 				{
 					FString TypeStr = GetArrayTypeString(InCollection.GetAttributeType(FName(*AttrName), GroupNameToUse)).ToString();
 
-					if (Out->IsA<TArray<bool>>(&BoolAttributeData))
+					if (TypeStr == FString("Bool"))
 					{
-						if (TypeStr == FString("bool"))
-						{
-							const TManagedArray<bool>& BoolAttributeArr = InCollection.GetAttribute<bool>(FName(*AttrName), GroupNameToUse);
-
-							SetValue<TArray<bool>>(Context, BoolAttributeArr.GetConstArray(), &BoolAttributeData);
-						}
-						else
-						{
-							SetValue<TArray<bool>>(Context, TArray<bool>(), &BoolAttributeData);
-						}
+						const TManagedArray<bool>& AttributeArr = InCollection.GetAttribute<bool>(FName(*AttrName), GroupNameToUse);
+						SetValue<TArray<bool>>(Context, AttributeArr.GetConstArray(), &BoolAttributeData);
 					}
-					else if (Out->IsA<TArray<float>>(&FloatAttributeData))
+					else if (TypeStr == FString("Float"))
 					{
-						if (TypeStr == FString("float"))
-						{
-							const TManagedArray<float>& FloatAttributeArr = InCollection.GetAttribute<float>(FName(*AttrName), GroupNameToUse);
-
-							SetValue<TArray<float>>(Context, FloatAttributeArr.GetConstArray(), &FloatAttributeData);
-						}
-						else
-						{
-							SetValue<TArray<float>>(Context, TArray<float>(), &FloatAttributeData);
-						}
+						const TManagedArray<float>& AttributeArr = InCollection.GetAttribute<float>(FName(*AttrName), GroupNameToUse);
+						SetValue<TArray<float>>(Context, AttributeArr.GetConstArray(), &FloatAttributeData);
 					}
-					else if (Out->IsA<TArray<int32>>(&IntAttributeData))
+					else if (TypeStr == FString("Double"))
 					{
-						if (TypeStr == FString("int32"))
-						{
-							const TManagedArray<int32>& IntAttributeArr = InCollection.GetAttribute<int32>(FName(*AttrName), GroupNameToUse);
-
-							SetValue<TArray<int32>>(Context, IntAttributeArr.GetConstArray(), &IntAttributeData);
-						}
-						else
-						{
-							SetValue<TArray<int32>>(Context, TArray<int32>(), &IntAttributeData);
-						}
+						const TManagedArray<double>& AttributeArr = InCollection.GetAttribute<double>(FName(*AttrName), GroupNameToUse);
+						SetValue<TArray<double>>(Context, AttributeArr.GetConstArray(), &DoubleAttributeData);
 					}
-					else if (Out->IsA<TArray<FVector3f>>(&VectorAttributeData))
+					else if (TypeStr == FString("Int32"))
 					{
-						if (TypeStr == FString("Vector"))
-						{
-							const TManagedArray<FVector3f>& VectorAttributeArr = InCollection.GetAttribute<FVector3f>(FName(*AttrName), GroupNameToUse);
-
-							SetValue<TArray<FVector3f>>(Context, VectorAttributeArr.GetConstArray(), &VectorAttributeData);
-						}
-						else
-						{
-							SetValue<TArray<FVector3f>>(Context, TArray<FVector3f>(), &VectorAttributeData);
-						}
+						const TManagedArray<int32>& AttributeArr = InCollection.GetAttribute<int32>(FName(*AttrName), GroupNameToUse);
+						SetValue<TArray<int32>>(Context, AttributeArr.GetConstArray(), &Int32AttributeData);
+					}
+					else if (TypeStr == FString("String"))
+					{
+						const TManagedArray<FString>& AttributeArr = InCollection.GetAttribute<FString>(FName(*AttrName), GroupNameToUse);
+						SetValue<TArray<FString>>(Context, AttributeArr.GetConstArray(), &StringAttributeData);
+					}
+					else if (TypeStr == FString("Vector"))
+					{
+						const TManagedArray<FVector3f>& AttributeArr = InCollection.GetAttribute<FVector3f>(FName(*AttrName), GroupNameToUse);
+						SetValue<TArray<FVector3f>>(Context, AttributeArr.GetConstArray(), &Vector3fAttributeData);
+					}
+					else if (TypeStr == FString("Vector3d"))
+					{
+						const TManagedArray<FVector3d>& AttributeArr = InCollection.GetAttribute<FVector3d>(FName(*AttrName), GroupNameToUse);
+						SetValue<TArray<FVector3d>>(Context, AttributeArr.GetConstArray(), &Vector3dAttributeData);
 					}
 				}
 			}
@@ -1214,6 +1256,23 @@ void FGetCollectionAttributeDataTypedDataflowNode::Evaluate(Dataflow::FContext&
 	}
 }
 
+template<typename T>
+static void SetAttributeData(const FDataflowNode* DataflowNode, Dataflow::FContext& Context, FManagedArrayCollection& InCollection, const TArray<T>& Property, FName AttributeName, FName GroupName)
+{
+	if (DataflowNode && DataflowNode->IsConnected<TArray<T>>(&Property))
+	{
+		TArray<T> AttributeData = DataflowNode->GetValue<TArray<T>>(Context, &Property);
+		TManagedArray<T>& AttributeArray = InCollection.ModifyAttribute<T>(AttributeName, GroupName);
+
+		if (AttributeData.Num() == AttributeArray.Num())
+		{
+			for (int32 Idx = 0; Idx < AttributeArray.Num(); ++Idx)
+			{
+				AttributeArray[Idx] = AttributeData[Idx];
+			}
+		}
+	}
+}
 
 void FSetCollectionAttributeDataTypedDataflowNode::Evaluate(Dataflow::FContext& Context, const FDataflowOutput* Out) const
 {
@@ -1237,89 +1296,42 @@ void FSetCollectionAttributeDataTypedDataflowNode::Evaluate(Dataflow::FContext&
 			{
 				if (InCollection.HasAttribute(FName(*AttrName), GroupNameToUse))
 				{
-					FString TypeStr = GetArrayTypeString(InCollection.GetAttributeType(FName(*AttrName), GroupNameToUse)).ToString();
-
-					if (TypeStr == FString("bool"))
+					FName AttributeName = FName(*AttrName);
+					FString TypeStr = GetArrayTypeString(InCollection.GetAttributeType(AttributeName, GroupNameToUse)).ToString();
+					
+					if (TypeStr == FString("Bool"))
 					{
-						if (IsConnected<TArray<bool>>(&BoolAttributeData))
-						{
-							TArray<bool> InBoolAttributeData = GetValue<TArray<bool>>(Context, &BoolAttributeData);
-
-							TManagedArray<bool>& BoolAttributeArr = InCollection.ModifyAttribute<bool>(FName(*AttrName), GroupNameToUse);
-
-							if (InBoolAttributeData.Num() == BoolAttributeArr.Num())
-							{
-								for (int32 Idx = 0; Idx < BoolAttributeArr.Num(); ++Idx)
-								{
-									BoolAttributeArr[Idx] = InBoolAttributeData[Idx];
-								}
-							}
-
-							SetValue<FManagedArrayCollection>(Context, InCollection, &Collection);
-						}
+						SetAttributeData<bool>(this, Context, InCollection, BoolAttributeData, AttributeName, GroupNameToUse);
 					}
-					else if (TypeStr == FString("float"))
+					else if (TypeStr == FString("Float"))
 					{
-						if (IsConnected<TArray<float>>(&FloatAttributeData))
-						{
-							TArray<float> InFloatAttributeData = GetValue<TArray<float>>(Context, &FloatAttributeData);
-
-							TManagedArray<float>& FloatAttributeArr = InCollection.ModifyAttribute<float>(FName(*AttrName), GroupNameToUse);
-
-							if (InFloatAttributeData.Num() == FloatAttributeArr.Num())
-							{
-								for (int32 Idx = 0; Idx < FloatAttributeArr.Num(); ++Idx)
-								{
-									FloatAttributeArr[Idx] = InFloatAttributeData[Idx];
-								}
-							}
-
-							SetValue<FManagedArrayCollection>(Context, InCollection, &Collection);
-						}
+						SetAttributeData<float>(this, Context, InCollection, FloatAttributeData, AttributeName, GroupNameToUse);
 					}
-					else if (TypeStr == FString("int32"))
+					else if (TypeStr == FString("Double"))
 					{
-						if (IsConnected<TArray<int32>>(&IntAttributeData))
-						{
-							TArray<int32> InIntAttributeData = GetValue<TArray<int32>>(Context, &IntAttributeData);
-
-							TManagedArray<int32>& IntAttributeArr = InCollection.ModifyAttribute<int32>(FName(*AttrName), GroupNameToUse);
-
-							if (InIntAttributeData.Num() == IntAttributeArr.Num())
-							{
-								for (int32 Idx = 0; Idx < IntAttributeArr.Num(); ++Idx)
-								{
-									IntAttributeArr[Idx] = InIntAttributeData[Idx];
-								}
-							}
-
-							SetValue<FManagedArrayCollection>(Context, InCollection, &Collection);
-						}
+						SetAttributeData<double>(this, Context, InCollection, DoubleAttributeData, AttributeName, GroupNameToUse);
+					}
+					else if (TypeStr == FString("Int32"))
+					{
+						SetAttributeData<int32>(this, Context, InCollection, Int32AttributeData, AttributeName, GroupNameToUse);
+					}
+					else if (TypeStr == FString("String"))
+					{
+						SetAttributeData<FString>(this, Context, InCollection, StringAttributeData, AttributeName, GroupNameToUse);
+					}
+					else if (TypeStr == FString("Vector"))
+					{
+						SetAttributeData<FVector3f>(this, Context, InCollection, Vector3fAttributeData, AttributeName, GroupNameToUse);
 					}
 					else if (TypeStr == FString("Vector3d"))
 					{
-						if (IsConnected<TArray<FVector>>(&VectorAttributeData))
-						{
-							TArray<FVector> InVectorAttributeData = GetValue<TArray<FVector>>(Context, &VectorAttributeData);
-
-							TManagedArray<FVector>& VectorAttributeArr = InCollection.ModifyAttribute<FVector>(FName(*AttrName), GroupNameToUse);
-
-							if (InVectorAttributeData.Num() == VectorAttributeArr.Num())
-							{
-								for (int32 Idx = 0; Idx < VectorAttributeArr.Num(); ++Idx)
-								{
-									VectorAttributeArr[Idx] = InVectorAttributeData[Idx];
-								}
-							}
-
-							SetValue<FManagedArrayCollection>(Context, InCollection, &Collection);
-						}
+						SetAttributeData<FVector3d>(this, Context, InCollection, Vector3dAttributeData, AttributeName, GroupNameToUse);
 					}
 				}
 			}
 		}
 
-		SetValue<FManagedArrayCollection>(Context, InCollection, &Collection);
+		SetValue<FManagedArrayCollection>(Context, MoveTemp(InCollection), &Collection);
 	}
 }
 
@@ -1392,7 +1404,7 @@ void FFloatArrayToVertexSelectionDataflowNode::Evaluate(Dataflow::FContext& Cont
 }
 
 
-void FSetVertexColorInCollectionDataflowNode::Evaluate(Dataflow::FContext& Context, const FDataflowOutput* Out) const
+void FSetVertexColorInCollectionFromVertexSelectionDataflowNode::Evaluate(Dataflow::FContext& Context, const FDataflowOutput* Out) const
 {
 	if (Out->IsA<FManagedArrayCollection>(&Collection))
 	{
@@ -1418,9 +1430,9 @@ void FSetVertexColorInCollectionDataflowNode::Evaluate(Dataflow::FContext& Conte
 					}
 				}
 			}
-
-			SetValue<FManagedArrayCollection>(Context, MoveTemp(InCollection), &Collection);
 		}
+
+		SetValue<FManagedArrayCollection>(Context, MoveTemp(InCollection), &Collection);
 	}
 }
 
@@ -1429,11 +1441,188 @@ void FMakeTransformDataflowNode::Evaluate(Dataflow::FContext& Context, const FDa
 {
 	if (Out->IsA<FTransform>(&OutTransform))
 	{
-		SetValue<FTransform>(Context, 
-							 FTransform(FQuat::MakeFromEuler(GetValue<FVector>(Context, &InRotation))
-							 , GetValue<FVector>(Context, &InTranslation)
-							 , GetValue<FVector>(Context, &InScale))
-							 ,&OutTransform);
+		SetValue<FTransform>(Context,
+			FTransform(FQuat::MakeFromEuler(GetValue<FVector>(Context, &InRotation))
+			, GetValue<FVector>(Context, &InTranslation)
+			, GetValue<FVector>(Context, &InScale))
+			, &OutTransform);
+	}
+}
+
+
+void FSetVertexColorInCollectionFromFloatArrayDataflowNode::Evaluate(Dataflow::FContext& Context, const FDataflowOutput* Out) const
+{
+	if (Out->IsA<FManagedArrayCollection>(&Collection))
+	{
+		FManagedArrayCollection InCollection = GetValue<FManagedArrayCollection>(Context, &Collection);
+		const TArray<float>& InFloatArray = GetValue<TArray<float>>(Context, &FloatArray);
+
+		const int32 NumVertices = InCollection.NumElements(FGeometryCollection::VerticesGroup);
+
+		if (InFloatArray.Num() == NumVertices)
+		{
+			if (TManagedArray<FLinearColor>* VertexColors = InCollection.FindAttribute<FLinearColor>("Color", FGeometryCollection::VerticesGroup))
+			{
+				for (int32 Idx = 0; Idx < NumVertices; ++Idx)
+				{
+					(*VertexColors)[Idx] = FLinearColor(Scale * InFloatArray[Idx], Scale * InFloatArray[Idx], Scale * InFloatArray[Idx]);
+				}
+			}
+		}
+
+		SetValue<FManagedArrayCollection>(Context, MoveTemp(InCollection), &Collection);
+	}
+}
+
+
+void FFloatArrayNormalizeDataflowNode::Evaluate(Dataflow::FContext& Context, const FDataflowOutput* Out) const
+{
+	if (Out->IsA<TArray<float>>(&OutFloatArray))
+	{
+		const TArray<float>& InInFloatArray = GetValue<TArray<float>>(Context, &InFloatArray);
+		const FDataflowVertexSelection& InSelection = GetValue<FDataflowVertexSelection>(Context, &Selection);
+		const float InMinRange = GetValue<float>(Context, &MinRange);
+		const float InMaxRange = GetValue<float>(Context, &MaxRange);
+
+		const int32 NumElems = InFloatArray.Num();
+
+		TArray<float> NewFloatArray;
+		NewFloatArray.Init(false, NumElems);
+
+		if (IsConnected<FDataflowVertexSelection>(&Selection))
+		{
+			if (InInFloatArray.Num() == InSelection.Num())
+			{
+				// Compute Min/Max
+				float MinValue = 1e9, MaxValue = 1e-9;
+
+				for (int32 Idx = 0; Idx < NumElems; ++Idx)
+				{
+					if (InSelection.IsSelected(Idx))
+					{
+						if (InInFloatArray[Idx] < MinValue)
+						{
+							MinValue = InInFloatArray[Idx];
+						}
+
+						if (InInFloatArray[Idx] > MaxValue)
+						{
+							MaxValue = InInFloatArray[Idx];
+						}
+					}
+				}
+
+				for (int32 Idx = 0; Idx < NumElems; ++Idx)
+				{
+					if (InSelection.IsSelected(Idx))
+					{
+						// Normalize it
+						NewFloatArray[Idx] = (NewFloatArray[Idx] - MinValue) / (MaxValue - MinValue);
+
+						// Transform it into (MinRange, MaxRange)
+						NewFloatArray[Idx] = InMinRange + NewFloatArray[Idx] * (InMaxRange - InMinRange);
+					}
+				}
+
+				SetValue<TArray<float>>(Context, NewFloatArray, &OutFloatArray);
+				return;
+			}
+		}
+		else
+		{
+			// Compute Min/Max
+			float MinValue = 1e9, MaxValue = 1e-9;
+
+			for (int32 Idx = 0; Idx < NumElems; ++Idx)
+			{
+				if (InInFloatArray[Idx] < MinValue)
+				{
+					MinValue = InInFloatArray[Idx];
+				}
+
+				if (InInFloatArray[Idx] > MaxValue)
+				{
+					MaxValue = InInFloatArray[Idx];
+				}
+			}
+
+			for (int32 Idx = 0; Idx < NumElems; ++Idx)
+			{
+				// Normalize it
+				NewFloatArray[Idx] = (NewFloatArray[Idx] - MinValue) / (MaxValue - MinValue);
+
+				// Transform it into (MinRange, MaxRange)
+				NewFloatArray[Idx] = InMinRange + NewFloatArray[Idx] * (InMaxRange - InMinRange);
+			}
+
+			SetValue<TArray<float>>(Context, NewFloatArray, &OutFloatArray);
+			return;
+		}
+
+		SetValue<TArray<float>>(Context, TArray<float>(), &OutFloatArray);
+	}
+}
+
+
+void FVectorArrayNormalizeDataflowNode::Evaluate(Dataflow::FContext& Context, const FDataflowOutput* Out) const
+{
+	if (Out->IsA<TArray<FVector>>(&OutVectorArray))
+	{
+		const TArray<FVector>& InInVectorArray = GetValue<TArray<FVector>>(Context, &InVectorArray);
+		const FDataflowVertexSelection& InSelection = GetValue<FDataflowVertexSelection>(Context, &Selection);
+		const float InMagnitude = GetValue<float>(Context, &Magnitude);
+
+		const int32 NumElems = InInVectorArray.Num();
+
+		TArray<FVector> NewVectorArray;
+		NewVectorArray.Init(FVector(0.f), NumElems);
+
+		if (IsConnected<FDataflowVertexSelection>(&Selection))
+		{
+			if (InInVectorArray.Num() == InSelection.Num())
+			{
+				for (int32 Idx = 0; Idx < NumElems; ++Idx)
+				{
+					FVector Vector = InInVectorArray[Idx];
+
+					if (InSelection.IsSelected(Idx))
+					{
+						Vector.Normalize();
+						Vector *= InMagnitude;
+						NewVectorArray[Idx] = Vector;
+					}
+				}
+
+				SetValue<TArray<FVector>>(Context, NewVectorArray, &OutVectorArray);
+				return;
+			}
+		}
+		else
+		{
+			for (int32 Idx = 0; Idx < NumElems; ++Idx)
+			{
+				FVector Vector = InInVectorArray[Idx];
+
+				Vector.Normalize();
+				Vector *= InMagnitude;
+				NewVectorArray[Idx] = Vector;
+			}
+
+			SetValue<TArray<FVector>>(Context, NewVectorArray, &OutVectorArray);
+			return;
+		}
+
+		SetValue<TArray<FVector>>(Context, TArray<FVector>(), &OutVectorArray);
+	}
+}
+
+
+void FMakeQuaternionDataflowNode::Evaluate(Dataflow::FContext& Context, const FDataflowOutput* Out) const
+{
+	if (Out->IsA<FQuat>(&Quaternion))
+	{
+		FQuat Value(GetValue<float>(Context, &X, X), GetValue<float>(Context, &Y, Y), GetValue<float>(Context, &Z, Z), GetValue<float>(Context, &W, W));
+		SetValue<FQuat>(Context, Value, &Quaternion);
 	}
 }
 
@@ -1451,4 +1640,3 @@ void FMultiplyTransformDataflowNode::Evaluate(Dataflow::FContext& Context, const
 
 
 
-
diff --git a/Engine/Plugins/Experimental/GeometryCollectionPlugin/Source/GeometryCollectionNodes/Private/Dataflow/GeometryCollectionSelectionNodes.cpp b/Engine/Plugins/Experimental/GeometryCollectionPlugin/Source/GeometryCollectionNodes/Private/Dataflow/GeometryCollectionSelectionNodes.cpp
--- a/Engine/Plugins/Experimental/GeometryCollectionPlugin/Source/GeometryCollectionNodes/Private/Dataflow/GeometryCollectionSelectionNodes.cpp
+++ b/Engine/Plugins/Experimental/GeometryCollectionPlugin/Source/GeometryCollectionNodes/Private/Dataflow/GeometryCollectionSelectionNodes.cpp
@@ -64,6 +64,8 @@ namespace Dataflow
 		DATAFLOW_NODE_REGISTER_CREATION_FACTORY(FCollectionFaceSelectionCustomDataflowNode);
 		DATAFLOW_NODE_REGISTER_CREATION_FACTORY(FCollectionSelectionConvertDataflowNode);
 		DATAFLOW_NODE_REGISTER_CREATION_FACTORY(FCollectionFaceSelectionInvertDataflowNode);
+		DATAFLOW_NODE_REGISTER_CREATION_FACTORY(FCollectionVertexSelectionByPercentageDataflowNode);
+		DATAFLOW_NODE_REGISTER_CREATION_FACTORY(FCollectionVertexSelectionSetOperationDataflowNode);
 
 		// GeometryCollection|Selection
 		DATAFLOW_NODE_REGISTER_CREATION_FACTORY_NODE_COLORS_BY_CATEGORY("GeometryCollection|Selection", FLinearColor(1.f, 1.f, 0.05f), CDefaultNodeBodyTintColor);
@@ -999,3 +1001,58 @@ void FCollectionFaceSelectionInvertDataflowNode::Evaluate(Dataflow::FContext& Co
 	}
 }
 
+
+void FCollectionVertexSelectionByPercentageDataflowNode::Evaluate(Dataflow::FContext& Context, const FDataflowOutput* Out) const
+{
+	if (Out->IsA<FDataflowVertexSelection>(&VertexSelection))
+	{
+		FDataflowVertexSelection InVertexSelection = GetValue<FDataflowVertexSelection>(Context, &VertexSelection);
+
+		int32 InPercentage = GetValue<int32>(Context, &Percentage);
+		float InRandomSeed = GetValue<float>(Context, &RandomSeed);
+
+		TArray<int32> SelectionArr = InVertexSelection.AsArray();
+
+		GeometryCollection::Facades::FCollectionTransformSelectionFacade::SelectByPercentage(SelectionArr, InPercentage, bDeterministic, InRandomSeed);
+
+		InVertexSelection.SetFromArray(SelectionArr);
+		SetValue<FDataflowVertexSelection>(Context, InVertexSelection, &VertexSelection);
+	}
+}
+
+
+void FCollectionVertexSelectionSetOperationDataflowNode::Evaluate(Dataflow::FContext& Context, const FDataflowOutput* Out) const
+{
+	if (Out->IsA<FDataflowVertexSelection>(&VertexSelection))
+	{
+		const FDataflowVertexSelection& InVertexSelectionA = GetValue<FDataflowVertexSelection>(Context, &VertexSelectionA);
+		const FDataflowVertexSelection& InVertexSelectionB = GetValue<FDataflowVertexSelection>(Context, &VertexSelectionB);
+
+		FDataflowVertexSelection NewVertexSelection;
+
+		if (InVertexSelectionA.Num() == InVertexSelectionB.Num())
+		{
+			if (Operation == ESetOperationEnum::Dataflow_SetOperation_AND)
+			{
+				InVertexSelectionA.AND(InVertexSelectionB, NewVertexSelection);
+			}
+			else if (Operation == ESetOperationEnum::Dataflow_SetOperation_OR)
+			{
+				InVertexSelectionA.OR(InVertexSelectionB, NewVertexSelection);
+			}
+			else if (Operation == ESetOperationEnum::Dataflow_SetOperation_XOR)
+			{
+				InVertexSelectionA.XOR(InVertexSelectionB, NewVertexSelection);
+			}
+		}
+		else
+		{
+			// ERROR: INPUT TRANSFORMSELECTIONS HAVE DIFFERENT NUMBER OF ELEMENTS
+			FString ErrorStr = "Input VertexSelections have different number of elements.";
+			UE_LOG(LogTemp, Error, TEXT("[Dataflow ERROR] %s"), *ErrorStr);
+		}
+
+		SetValue<FDataflowVertexSelection>(Context, NewVertexSelection, &VertexSelection);
+	}
+}
+
diff --git a/Engine/Plugins/Experimental/GeometryCollectionPlugin/Source/GeometryCollectionNodes/Public/Dataflow/GeometryCollectionFieldNodes.h b/Engine/Plugins/Experimental/GeometryCollectionPlugin/Source/GeometryCollectionNodes/Public/Dataflow/GeometryCollectionFieldNodes.h
--- a/Engine/Plugins/Experimental/GeometryCollectionPlugin/Source/GeometryCollectionNodes/Public/Dataflow/GeometryCollectionFieldNodes.h
+++ b/Engine/Plugins/Experimental/GeometryCollectionPlugin/Source/GeometryCollectionNodes/Public/Dataflow/GeometryCollectionFieldNodes.h
@@ -29,6 +29,7 @@ enum class EDataflowFieldFalloffType : uint8
 	Dataflow_Max                UMETA(Hidden)
 };
 
+
 /**
  *
  * RadialFalloff Field Dataflow node
@@ -42,11 +43,11 @@ struct FRadialFalloffFieldDataflowNode : public FDataflowNode
 
 public:
 	UPROPERTY(meta = (DataflowInput, DataflowIntrinsic))
-	TArray<FVector3f> VertexArray;
+	TArray<FVector3f> SamplePositions;
 
 	/**  */
-	UPROPERTY(meta = (DataflowInput, DisplayName = "VertexSelection"))
-	FDataflowVertexSelection VertexSelection;
+	UPROPERTY(meta = (DataflowInput))
+	FDataflowVertexSelection SampleIndices;
 
 	/**  */
 	UPROPERTY(meta = (DataflowInput, DataflowIntrinsic))
@@ -57,37 +58,56 @@ public:
 	FVector Translation = FVector(0.f);
 
 	/**  */
-	UPROPERTY(EditAnywhere, Category = "Field")
+	UPROPERTY(EditAnywhere, Category = "Field", meta = (DataflowInput))
 	float Magnitude = 1.f;
 
 	/**  */
-	UPROPERTY(EditAnywhere, Category = "Field")
+	UPROPERTY(EditAnywhere, Category = "Field", meta = (DataflowInput))
 	float MinRange = 0.f;
 
 	/**  */
-	UPROPERTY(EditAnywhere, Category = "Field")
+	UPROPERTY(EditAnywhere, Category = "Field", meta = (DataflowInput))
 	float MaxRange = 1.f;
 
 	/**  */
-	UPROPERTY(EditAnywhere, Category = "Field")
+	UPROPERTY(EditAnywhere, Category = "Field", meta = (DataflowInput))
 	float Default = 0.f;
 
 	/**  */
 	UPROPERTY(EditAnywhere, Category = "Field")
-		EDataflowFieldFalloffType FalloffType = EDataflowFieldFalloffType::Dataflow_FieldFalloffType_Linear;
+	EDataflowFieldFalloffType FalloffType = EDataflowFieldFalloffType::Dataflow_FieldFalloffType_Linear;
+
+	/**  */
+	UPROPERTY(meta = (DataflowOutput))
+	TArray<float> FieldFloatResult;
+
+	/**  */
+	UPROPERTY(meta = (DataflowOutput))
+	TArray<int32> FieldRemap;
+
+	/**  */
+	UPROPERTY(meta = (DataflowOutput))
+	FDataflowVertexSelection FieldSelectionMask;
 
 	/**  */
 	UPROPERTY(meta = (DataflowOutput))
-	TArray<float> WeightArray;
+	int32 NumSamplePositions;
 
 	FRadialFalloffFieldDataflowNode(const Dataflow::FNodeParameters& InParam, FGuid InGuid = FGuid::NewGuid())
 		: FDataflowNode(InParam, InGuid)
 	{
-		RegisterInputConnection(&VertexArray);
-		RegisterInputConnection(&VertexSelection);
+		RegisterInputConnection(&SamplePositions);
+		RegisterInputConnection(&SampleIndices);
 		RegisterInputConnection(&Sphere);
 		RegisterInputConnection(&Translation);
-		RegisterOutputConnection(&WeightArray);
+		RegisterInputConnection(&Magnitude);
+		RegisterInputConnection(&MinRange);
+		RegisterInputConnection(&MaxRange);
+		RegisterInputConnection(&Default);
+		RegisterOutputConnection(&FieldFloatResult);
+		RegisterOutputConnection(&FieldRemap);
+		RegisterOutputConnection(&FieldSelectionMask);
+		RegisterOutputConnection(&NumSamplePositions);
 	}
 
 	virtual void Evaluate(Dataflow::FContext& Context, const FDataflowOutput* Out) const override;
@@ -97,27 +117,30 @@ public:
 
 /**
  *
- * PlaneFalloff Field Dataflow node
+ * BoxFalloff Field Dataflow node
  *
  */
 USTRUCT()
-struct FPlaneFalloffFieldDataflowNode : public FDataflowNode
+struct FBoxFalloffFieldDataflowNode : public FDataflowNode
 {
 	GENERATED_USTRUCT_BODY()
-	DATAFLOW_NODE_DEFINE_INTERNAL(FPlaneFalloffFieldDataflowNode, "PlaneFalloffField", "Fields", "")
+	DATAFLOW_NODE_DEFINE_INTERNAL(FBoxFalloffFieldDataflowNode, "BoxFalloffField", "Fields", "")
 
 public:
-	/**  */
 	UPROPERTY(meta = (DataflowInput, DataflowIntrinsic))
-	TArray<FVector3f> VertexArray;
+	TArray<FVector3f> SamplePositions;
 
 	/**  */
-	UPROPERTY(meta = (DataflowInput, DisplayName = "VertexSelection"))
-	FDataflowVertexSelection VertexSelection;
+	UPROPERTY(meta = (DataflowInput))
+	FDataflowVertexSelection SampleIndices;
+
+	/**  */
+	UPROPERTY(meta = (DataflowInput, DataflowIntrinsic))
+	FBox Box = FBox(ForceInit);
 
 	/**  */
 	UPROPERTY(EditAnywhere, Category = "Field", meta = (DataflowInput))
-	FVector Translation = FVector(0.f);
+	FTransform Transform = FTransform::Identity;
 
 	/**  */
 	UPROPERTY(EditAnywhere, Category = "Field")
@@ -133,18 +156,100 @@ public:
 
 	/**  */
 	UPROPERTY(EditAnywhere, Category = "Field")
-	float Distance = 0.f;
+	float Default = 0.f;
 
 	/**  */
 	UPROPERTY(EditAnywhere, Category = "Field")
+	EDataflowFieldFalloffType FalloffType = EDataflowFieldFalloffType::Dataflow_FieldFalloffType_Linear;
+
+	/**  */
+	UPROPERTY(meta = (DataflowOutput))
+	TArray<float> FieldFloatResult;
+
+	/**  */
+	UPROPERTY(meta = (DataflowOutput))
+	TArray<int32> FieldRemap;
+
+	/**  */
+	UPROPERTY(meta = (DataflowOutput))
+	FDataflowVertexSelection FieldSelectionMask;
+
+	/**  */
+	UPROPERTY(meta = (DataflowOutput))
+	int32 NumSamplePositions;
+
+	FBoxFalloffFieldDataflowNode(const Dataflow::FNodeParameters& InParam, FGuid InGuid = FGuid::NewGuid())
+		: FDataflowNode(InParam, InGuid)
+	{
+		RegisterInputConnection(&SamplePositions);
+		RegisterInputConnection(&SampleIndices);
+		RegisterInputConnection(&Box);
+		RegisterInputConnection(&Transform);
+		RegisterInputConnection(&Magnitude);
+		RegisterInputConnection(&MinRange);
+		RegisterInputConnection(&MaxRange);
+		RegisterInputConnection(&Default);
+		RegisterOutputConnection(&FieldFloatResult);
+		RegisterOutputConnection(&FieldRemap);
+		RegisterOutputConnection(&FieldSelectionMask);
+		RegisterOutputConnection(&NumSamplePositions);
+	}
+
+	virtual void Evaluate(Dataflow::FContext& Context, const FDataflowOutput* Out) const override;
+
+};
+
+
+/**
+ *
+ * PlaneFalloff Field Dataflow node
+ *
+ */
+USTRUCT()
+struct FPlaneFalloffFieldDataflowNode : public FDataflowNode
+{
+	GENERATED_USTRUCT_BODY()
+	DATAFLOW_NODE_DEFINE_INTERNAL(FPlaneFalloffFieldDataflowNode, "PlaneFalloffField", "Fields", "")
+
+public:
+	/**  */
+	UPROPERTY(meta = (DataflowInput, DataflowIntrinsic))
+	TArray<FVector3f> SamplePositions;
+
+	/**  */
+	UPROPERTY(meta = (DataflowInput))
+	FDataflowVertexSelection SampleIndices;
+
+	/**  */
+	UPROPERTY(EditAnywhere, Category = "Field", meta = (DataflowInput))
 	FVector Position = FVector(0.f);
 
 	/**  */
-	UPROPERTY(EditAnywhere, Category = "Field")
+	UPROPERTY(EditAnywhere, Category = "Field", meta = (DataflowInput))
 	FVector Normal = FVector(0.f, 0.f, 1.f);
 
 	/**  */
-	UPROPERTY(EditAnywhere, Category = "Field")
+	UPROPERTY(EditAnywhere, Category = "Field", meta = (DataflowInput))
+	float Distance = 0.f;
+
+	/**  */
+	UPROPERTY(EditAnywhere, Category = "Field", meta = (DataflowInput))
+	FVector Translation = FVector(0.f);
+
+	/**  */
+	UPROPERTY(EditAnywhere, Category = "Field", meta = (DataflowInput))
+	float Magnitude = 1.f;
+
+	/**  */
+	UPROPERTY(EditAnywhere, Category = "Field", meta = (DataflowInput))
+	float MinRange = 0.f;
+
+	/**  */
+	UPROPERTY(EditAnywhere, Category = "Field", meta = (DataflowInput))
+	float MaxRange = 1.f;
+
+	/**  */
+	UPROPERTY(EditAnywhere, Category = "Field", meta = (DataflowInput))
 	float Default = 0.f;
 
 	/**  */
@@ -153,17 +258,37 @@ public:
 
 	/**  */
 	UPROPERTY(meta = (DataflowOutput))
-	TArray<float> WeightArray;
+	TArray<float> FieldFloatResult;
+
+	/**  */
+	UPROPERTY(meta = (DataflowOutput))
+	TArray<int32> FieldRemap;
+
+	/**  */
+	UPROPERTY(meta = (DataflowOutput))
+	FDataflowVertexSelection FieldSelectionMask;
+
+	/**  */
+	UPROPERTY(meta = (DataflowOutput))
+	int32 NumSamplePositions;
 
 	FPlaneFalloffFieldDataflowNode(const Dataflow::FNodeParameters& InParam, FGuid InGuid = FGuid::NewGuid())
 		: FDataflowNode(InParam, InGuid)
 	{
-		RegisterInputConnection(&VertexArray);
-		RegisterInputConnection(&VertexSelection);
+		RegisterInputConnection(&SamplePositions);
+		RegisterInputConnection(&SampleIndices);
 		RegisterInputConnection(&Position);
 		RegisterInputConnection(&Normal);
+		RegisterInputConnection(&Distance);
 		RegisterInputConnection(&Translation);
-		RegisterOutputConnection(&WeightArray);
+		RegisterInputConnection(&Magnitude);
+		RegisterInputConnection(&MinRange);
+		RegisterInputConnection(&MaxRange);
+		RegisterInputConnection(&Default);
+		RegisterOutputConnection(&FieldFloatResult);
+		RegisterOutputConnection(&FieldRemap);
+		RegisterOutputConnection(&FieldSelectionMask);
+		RegisterOutputConnection(&NumSamplePositions);
 	}
 
 	virtual void Evaluate(Dataflow::FContext& Context, const FDataflowOutput* Out) const override;
@@ -201,11 +326,11 @@ struct FRadialIntMaskFieldDataflowNode : public FDataflowNode
 public:
 	/**  */
 	UPROPERTY(meta = (DataflowInput, DataflowIntrinsic))
-	TArray<FVector3f> VertexArray;
+	TArray<FVector3f> SamplePositions;
 
 	/**  */
-	UPROPERTY(meta = (DataflowInput, DisplayName = "VertexSelection"))
-	FDataflowVertexSelection VertexSelection;
+	UPROPERTY(meta = (DataflowInput))
+	FDataflowVertexSelection SampleIndices;
 
 	/**  */
 	UPROPERTY(meta = (DataflowInput, DataflowIntrinsic))
@@ -216,29 +341,41 @@ public:
 	FVector Translation = FVector(0.f);
 
 	/**  */
-	UPROPERTY(EditAnywhere, Category = "Field")
+	UPROPERTY(EditAnywhere, Category = "Field", meta = (DataflowInput))
 	int32 InteriorValue = 1;
 
 	/**  */
-	UPROPERTY(EditAnywhere, Category = "Field")
+	UPROPERTY(EditAnywhere, Category = "Field", meta = (DataflowInput))
 	int32 ExteriorValue = 0;
 
 	/**  */
 	UPROPERTY(EditAnywhere, Category = "Field")
-	EDataflowSetMaskConditionType SetMaskCondition = EDataflowSetMaskConditionType::Dataflow_SetMaskConditionType_Always;
+	EDataflowSetMaskConditionType SetMaskConditionType = EDataflowSetMaskConditionType::Dataflow_SetMaskConditionType_Always;
+
+	/**  */
+	UPROPERTY(meta = (DataflowOutput))
+	TArray<int32> FieldIntResult;
 
 	/**  */
 	UPROPERTY(meta = (DataflowOutput))
-	TArray<int32> MaskArray;
+	TArray<int32> FieldRemap;
+
+	/**  */
+	UPROPERTY(meta = (DataflowOutput))
+	int32 NumSamplePositions;
 
 	FRadialIntMaskFieldDataflowNode(const Dataflow::FNodeParameters& InParam, FGuid InGuid = FGuid::NewGuid())
 		: FDataflowNode(InParam, InGuid)
 	{
-		RegisterInputConnection(&VertexArray);
-		RegisterInputConnection(&VertexSelection);
+		RegisterInputConnection(&SamplePositions);
+		RegisterInputConnection(&SampleIndices);
 		RegisterInputConnection(&Sphere);
 		RegisterInputConnection(&Translation);
-		RegisterOutputConnection(&MaskArray);
+		RegisterInputConnection(&InteriorValue);
+		RegisterInputConnection(&ExteriorValue);
+		RegisterOutputConnection(&FieldIntResult);
+		RegisterOutputConnection(&FieldRemap);
+		RegisterOutputConnection(&NumSamplePositions);
 	}
 
 	virtual void Evaluate(Dataflow::FContext& Context, const FDataflowOutput* Out) const override;
@@ -260,26 +397,37 @@ struct FUniformScalarFieldDataflowNode : public FDataflowNode
 public:
 	/**  */
 	UPROPERTY(meta = (DataflowInput, DataflowIntrinsic))
-	TArray<FVector3f> VertexArray;
+	TArray<FVector3f> SamplePositions;
 
 	/**  */
-	UPROPERTY(meta = (DataflowInput, DisplayName = "VertexSelection"))
-	FDataflowVertexSelection VertexSelection;
+	UPROPERTY(meta = (DataflowInput))
+	FDataflowVertexSelection SampleIndices;
 
 	/**  */
-	UPROPERTY(EditAnywhere, Category = "Field")
+	UPROPERTY(EditAnywhere, Category = "Field", meta = (DataflowInput))
 	float Magnitude = 1.f;
 
 	/**  */
 	UPROPERTY(meta = (DataflowOutput))
-	TArray<float> WeightArray;
+	TArray<float> FieldFloatResult;
+
+	/**  */
+	UPROPERTY(meta = (DataflowOutput))
+	TArray<int32> FieldRemap;
+
+	/**  */
+	UPROPERTY(meta = (DataflowOutput))
+	int32 NumSamplePositions;
 
 	FUniformScalarFieldDataflowNode(const Dataflow::FNodeParameters& InParam, FGuid InGuid = FGuid::NewGuid())
 		: FDataflowNode(InParam, InGuid)
 	{
-		RegisterInputConnection(&VertexArray);
-		RegisterInputConnection(&VertexSelection);
-		RegisterOutputConnection(&WeightArray);
+		RegisterInputConnection(&SamplePositions);
+		RegisterInputConnection(&SampleIndices);
+		RegisterInputConnection(&Magnitude);
+		RegisterOutputConnection(&FieldFloatResult);
+		RegisterOutputConnection(&FieldRemap);
+		RegisterOutputConnection(&NumSamplePositions);
 	}
 
 	virtual void Evaluate(Dataflow::FContext& Context, const FDataflowOutput* Out) const override;
@@ -301,30 +449,42 @@ struct FUniformVectorFieldDataflowNode : public FDataflowNode
 public:
 	/**  */
 	UPROPERTY(meta = (DataflowInput, DataflowIntrinsic))
-	TArray<FVector3f> VertexArray;
+	TArray<FVector3f> SamplePositions;
 
 	/**  */
-	UPROPERTY(meta = (DataflowInput, DisplayName = "VertexSelection"))
-	FDataflowVertexSelection VertexSelection;
+	UPROPERTY(meta = (DataflowInput))
+	FDataflowVertexSelection SampleIndices;
 
 	/**  */
-	UPROPERTY(EditAnywhere, Category = "Field")
+	UPROPERTY(EditAnywhere, Category = "Field", meta = (DataflowInput))
 	float Magnitude = 1.f;
 
 	/**  */
-	UPROPERTY(EditAnywhere, Category = "Field")
-	FVector Direction = FVector(0.f);
+	UPROPERTY(EditAnywhere, Category = "Field", meta = (DataflowInput))
+	FVector Direction = FVector(1.f, 0.f, 0.f);
 
 	/**  */
 	UPROPERTY(meta = (DataflowOutput))
-	TArray<FVector> VectorArray;
+	TArray<FVector> FieldVectorResult;
+
+	/**  */
+	UPROPERTY(meta = (DataflowOutput))
+	TArray<int32> FieldRemap;
+
+	/**  */
+	UPROPERTY(meta = (DataflowOutput))
+	int32 NumSamplePositions;
 
 	FUniformVectorFieldDataflowNode(const Dataflow::FNodeParameters& InParam, FGuid InGuid = FGuid::NewGuid())
 		: FDataflowNode(InParam, InGuid)
 	{
-		RegisterInputConnection(&VertexArray);
-		RegisterInputConnection(&VertexSelection);
-		RegisterOutputConnection(&VectorArray);
+		RegisterInputConnection(&SamplePositions);
+		RegisterInputConnection(&SampleIndices);
+		RegisterInputConnection(&Magnitude);
+		RegisterInputConnection(&Direction);
+		RegisterOutputConnection(&FieldVectorResult);
+		RegisterOutputConnection(&FieldRemap);
+		RegisterOutputConnection(&NumSamplePositions);
 	}
 
 	virtual void Evaluate(Dataflow::FContext& Context, const FDataflowOutput* Out) const override;
@@ -346,30 +506,42 @@ struct FRadialVectorFieldDataflowNode : public FDataflowNode
 public:
 	/**  */
 	UPROPERTY(meta = (DataflowInput, DataflowIntrinsic))
-	TArray<FVector3f> VertexArray;
+	TArray<FVector3f> SamplePositions;
 
 	/**  */
-	UPROPERTY(meta = (DataflowInput, DisplayName = "VertexSelection"))
-	FDataflowVertexSelection VertexSelection;
+	UPROPERTY(meta = (DataflowInput))
+	FDataflowVertexSelection SampleIndices;
 
 	/**  */
-	UPROPERTY(EditAnywhere, Category = "Field")
+	UPROPERTY(EditAnywhere, Category = "Field", meta = (DataflowInput))
 	float Magnitude = 1.f;
 
 	/**  */
-	UPROPERTY(EditAnywhere, Category = "Field")
+	UPROPERTY(EditAnywhere, Category = "Field", meta = (DataflowInput))
 	FVector Position = FVector(0.f);
 
 	/**  */
 	UPROPERTY(meta = (DataflowOutput))
-	TArray<FVector> VectorArray;
+	TArray<FVector> FieldVectorResult;
+
+	/**  */
+	UPROPERTY(meta = (DataflowOutput))
+	TArray<int32> FieldRemap;
+
+	/**  */
+	UPROPERTY(meta = (DataflowOutput))
+	int32 NumSamplePositions;
 
 	FRadialVectorFieldDataflowNode(const Dataflow::FNodeParameters& InParam, FGuid InGuid = FGuid::NewGuid())
 		: FDataflowNode(InParam, InGuid)
 	{
-		RegisterInputConnection(&VertexArray);
-		RegisterInputConnection(&VertexSelection);
-		RegisterOutputConnection(&VectorArray);
+		RegisterInputConnection(&SamplePositions);
+		RegisterInputConnection(&SampleIndices);
+		RegisterInputConnection(&Magnitude);
+		RegisterInputConnection(&Position);
+		RegisterOutputConnection(&FieldVectorResult);
+		RegisterOutputConnection(&FieldRemap);
+		RegisterOutputConnection(&NumSamplePositions);
 	}
 
 	virtual void Evaluate(Dataflow::FContext& Context, const FDataflowOutput* Out) const override;
@@ -391,26 +563,37 @@ struct FRandomVectorFieldDataflowNode : public FDataflowNode
 public:
 	/**  */
 	UPROPERTY(meta = (DataflowInput, DataflowIntrinsic))
-	TArray<FVector3f> VertexArray;
+	TArray<FVector3f> SamplePositions;
 
 	/**  */
-	UPROPERTY(meta = (DataflowInput, DisplayName = "VertexSelection"))
-	FDataflowVertexSelection VertexSelection;
+	UPROPERTY(meta = (DataflowInput))
+	FDataflowVertexSelection SampleIndices;
 
 	/**  */
-	UPROPERTY(EditAnywhere, Category = "Field")
+	UPROPERTY(EditAnywhere, Category = "Field", meta = (DataflowInput))
 	float Magnitude = 1.f;
 
 	/**  */
 	UPROPERTY(meta = (DataflowOutput))
-	TArray<FVector> VectorArray;
+	TArray<FVector> FieldVectorResult;
+
+	/**  */
+	UPROPERTY(meta = (DataflowOutput))
+	TArray<int32> FieldRemap;
+
+	/**  */
+	UPROPERTY(meta = (DataflowOutput))
+	int32 NumSamplePositions;
 
 	FRandomVectorFieldDataflowNode(const Dataflow::FNodeParameters& InParam, FGuid InGuid = FGuid::NewGuid())
 		: FDataflowNode(InParam, InGuid)
 	{
-		RegisterInputConnection(&VertexArray);
-		RegisterInputConnection(&VertexSelection);
-		RegisterOutputConnection(&VectorArray);
+		RegisterInputConnection(&SamplePositions);
+		RegisterInputConnection(&SampleIndices);
+		RegisterInputConnection(&Magnitude);
+		RegisterOutputConnection(&FieldVectorResult);
+		RegisterOutputConnection(&FieldRemap);
+		RegisterOutputConnection(&NumSamplePositions);
 	}
 
 	virtual void Evaluate(Dataflow::FContext& Context, const FDataflowOutput* Out) const override;
@@ -432,18 +615,18 @@ struct FNoiseFieldDataflowNode : public FDataflowNode
 public:
 	/**  */
 	UPROPERTY(meta = (DataflowInput, DataflowIntrinsic))
-	TArray<FVector3f> VertexArray;
+	TArray<FVector3f> SamplePositions;
 
 	/**  */
-	UPROPERTY(meta = (DataflowInput, DisplayName = "VertexSelection"))
-	FDataflowVertexSelection VertexSelection;
+	UPROPERTY(meta = (DataflowInput))
+	FDataflowVertexSelection SampleIndices;
 
 	/**  */
-	UPROPERTY(EditAnywhere, Category = "Field")
+	UPROPERTY(EditAnywhere, Category = "Field", meta = (DataflowInput))
 	float MinRange = 0.f;
 
 	/**  */
-	UPROPERTY(EditAnywhere, Category = "Field")
+	UPROPERTY(EditAnywhere, Category = "Field", meta = (DataflowInput))
 	float MaxRange = 1.f;
 
 	/**  */
@@ -452,15 +635,27 @@ public:
 
 	/**  */
 	UPROPERTY(meta = (DataflowOutput))
-	TArray<float> WeightArray;
+	TArray<float> FieldFloatResult;
+
+	/**  */
+	UPROPERTY(meta = (DataflowOutput))
+	TArray<int32> FieldRemap;
+
+	/**  */
+	UPROPERTY(meta = (DataflowOutput))
+	int32 NumSamplePositions;
 
 	FNoiseFieldDataflowNode(const Dataflow::FNodeParameters& InParam, FGuid InGuid = FGuid::NewGuid())
 		: FDataflowNode(InParam, InGuid)
 	{
-		RegisterInputConnection(&VertexArray);
-		RegisterInputConnection(&VertexSelection);
+		RegisterInputConnection(&SamplePositions);
+		RegisterInputConnection(&SampleIndices);
+		RegisterInputConnection(&MinRange);
+		RegisterInputConnection(&MaxRange);
 		RegisterInputConnection(&Transform);
-		RegisterOutputConnection(&WeightArray);
+		RegisterOutputConnection(&FieldFloatResult);
+		RegisterOutputConnection(&FieldRemap);
+		RegisterOutputConnection(&NumSamplePositions);
 	}
 
 	virtual void Evaluate(Dataflow::FContext& Context, const FDataflowOutput* Out) const override;
@@ -482,26 +677,37 @@ struct FUniformIntegerFieldDataflowNode : public FDataflowNode
 public:
 	/**  */
 	UPROPERTY(meta = (DataflowInput, DataflowIntrinsic))
-	TArray<FVector3f> VertexArray;
+	TArray<FVector3f> SamplePositions;
 
 	/**  */
-	UPROPERTY(meta = (DataflowInput, DisplayName = "VertexSelection"))
-	FDataflowVertexSelection VertexSelection;
+	UPROPERTY(meta = (DataflowInput))
+	FDataflowVertexSelection SampleIndices;
 
 	/**  */
-	UPROPERTY(EditAnywhere, Category = "Field")
+	UPROPERTY(EditAnywhere, Category = "Field", meta = (DataflowInput))
 	int32 Magnitude = 0;
 
 	/**  */
 	UPROPERTY(meta = (DataflowOutput))
-	TArray<int32> MaskArray;
+	TArray<int32> FieldIntResult;
+
+	/**  */
+	UPROPERTY(meta = (DataflowOutput))
+	TArray<int32> FieldRemap;
+
+	/**  */
+	UPROPERTY(meta = (DataflowOutput))
+	int32 NumSamplePositions;
 
 	FUniformIntegerFieldDataflowNode(const Dataflow::FNodeParameters& InParam, FGuid InGuid = FGuid::NewGuid())
 		: FDataflowNode(InParam, InGuid)
 	{
-		RegisterInputConnection(&VertexArray);
-		RegisterInputConnection(&VertexSelection);
-		RegisterOutputConnection(&MaskArray);
+		RegisterInputConnection(&SamplePositions);
+		RegisterInputConnection(&SampleIndices);
+		RegisterInputConnection(&Magnitude);
+		RegisterOutputConnection(&FieldIntResult);
+		RegisterOutputConnection(&FieldRemap);
+		RegisterOutputConnection(&NumSamplePositions);
 	}
 
 	virtual void Evaluate(Dataflow::FContext& Context, const FDataflowOutput* Out) const override;
@@ -540,18 +746,18 @@ struct FWaveScalarFieldDataflowNode : public FDataflowNode
 public:
 	/**  */
 	UPROPERTY(meta = (DataflowInput, DataflowIntrinsic))
-	TArray<FVector3f> VertexArray;
+	TArray<FVector3f> SamplePositions;
 
 	/**  */
-	UPROPERTY(meta = (DataflowInput, DisplayName = "VertexSelection"))
-	FDataflowVertexSelection VertexSelection;
+	UPROPERTY(meta = (DataflowInput))
+	FDataflowVertexSelection SampleIndices;
 
 	/**  */
-	UPROPERTY(EditAnywhere, Category = "Field")
+	UPROPERTY(EditAnywhere, Category = "Field", meta = (DataflowInput))
 	float Magnitude = 1.f;
 
 	/**  */
-	UPROPERTY(meta = (DataflowInput, DataflowIntrinsic))
+	UPROPERTY(meta = (DataflowInput))
 	FVector Position = FVector(0.f);
 
 	/**  */
@@ -559,11 +765,11 @@ public:
 	FVector Translation = FVector(0.f);
 
 	/**  */
-	UPROPERTY(EditAnywhere, Category = "Field")
+	UPROPERTY(EditAnywhere, Category = "Field", meta = (DataflowInput))
 	float Wavelength = 1000.f;
 
 	/**  */
-	UPROPERTY(EditAnywhere, Category = "Field")
+	UPROPERTY(EditAnywhere, Category = "Field", meta = (DataflowInput))
 	float Period = 1.f;
 
 	/**  */
@@ -576,15 +782,29 @@ public:
 
 	/**  */
 	UPROPERTY(meta = (DataflowOutput))
-	TArray<float> WeightArray;
+	TArray<float> FieldFloatResult;
+
+	/**  */
+	UPROPERTY(meta = (DataflowOutput))
+	TArray<int32> FieldRemap;
+
+	/**  */
+	UPROPERTY(meta = (DataflowOutput))
+	int32 NumSamplePositions;
 
 	FWaveScalarFieldDataflowNode(const Dataflow::FNodeParameters& InParam, FGuid InGuid = FGuid::NewGuid())
 		: FDataflowNode(InParam, InGuid)
 	{
-		RegisterInputConnection(&VertexArray);
-		RegisterInputConnection(&VertexSelection);
+		RegisterInputConnection(&SamplePositions);
+		RegisterInputConnection(&SampleIndices);
+		RegisterInputConnection(&Magnitude);
+		RegisterInputConnection(&Position);
 		RegisterInputConnection(&Translation);
-		RegisterOutputConnection(&WeightArray);
+		RegisterInputConnection(&Wavelength);
+		RegisterInputConnection(&Period);
+		RegisterOutputConnection(&FieldFloatResult);
+		RegisterOutputConnection(&FieldRemap);
+		RegisterOutputConnection(&NumSamplePositions);
 	}
 
 	virtual void Evaluate(Dataflow::FContext& Context, const FDataflowOutput* Out) const override;
@@ -594,30 +814,49 @@ public:
 
 /**
  *
- * BoxFalloff Field Dataflow node
+ *
+ *
+ */
+UENUM(BlueprintType)
+enum class EDataflowFloatFieldOperationType : uint8
+{
+	Dataflow_FloatFieldOperationType_Multiply	UMETA(DisplayName = "Multiply", ToolTip = "Multiply the fields output values : Output = Left * Right"),
+	Dataflow_FloatFieldFalloffType_Divide		UMETA(DisplayName = "Divide", ToolTip = "Divide the fields output values : Output = Left / Right"),
+	Dataflow_FloatFieldFalloffType_Add			UMETA(DisplayName = "Add", ToolTip = "Add the fields output values : Output = Left + Right"),
+	Dataflow_FloatFieldFalloffType_Substract	UMETA(DisplayName = "Subtract", ToolTip = "Subtract the fields output : Output = Left - Right"),
+	Dataflow_FloatFieldFalloffType_Min			UMETA(DisplayName = "Min", ToolTip = "Min of the fields output values: Output = Min(Left, Right)"),
+	Dataflow_FloatFieldFalloffType_Max			UMETA(DisplayName = "Max", ToolTip = "Max of the fields output values: Output = Max(Left, Right)"),
+	//~~~
+	//256th entry
+	Dataflow_Max                UMETA(Hidden)
+};
+
+
+/**
+ *
+ *
  *
  */
 USTRUCT()
-struct FBoxFalloffFieldDataflowNode : public FDataflowNode
+struct FSumScalarFieldDataflowNode : public FDataflowNode
 {
 	GENERATED_USTRUCT_BODY()
-	DATAFLOW_NODE_DEFINE_INTERNAL(FBoxFalloffFieldDataflowNode, "BoxFalloffField", "Fields", "")
+	DATAFLOW_NODE_DEFINE_INTERNAL(FSumScalarFieldDataflowNode, "SumScalarField", "Fields", "")
 
 public:
 	UPROPERTY(meta = (DataflowInput, DataflowIntrinsic))
-	TArray<FVector3f> VertexArray;
+	TArray<float> FieldFloatLeft;
 
 	/**  */
-	UPROPERTY(meta = (DataflowInput, DisplayName = "VertexSelection"))
-	FDataflowVertexSelection VertexSelection;
+	UPROPERTY(meta = (DataflowInput))
+	TArray<int32> FieldRemapLeft;
 
-	/**  */
 	UPROPERTY(meta = (DataflowInput, DataflowIntrinsic))
-	FBox Box = FBox(ForceInit);
+	TArray<float> FieldFloatRight;
 
 	/**  */
-	UPROPERTY(EditAnywhere, Category = "Field", meta = (DataflowInput))
-	FTransform Transform = FTransform::Identity;
+	UPROPERTY(meta = (DataflowInput))
+	TArray<int32> FieldRemapRight;
 
 	/**  */
 	UPROPERTY(EditAnywhere, Category = "Field")
@@ -625,32 +864,118 @@ public:
 
 	/**  */
 	UPROPERTY(EditAnywhere, Category = "Field")
-	float MinRange = 0.f;
+	EDataflowFloatFieldOperationType Operation = EDataflowFloatFieldOperationType::Dataflow_FloatFieldFalloffType_Add;
 
 	/**  */
 	UPROPERTY(EditAnywhere, Category = "Field")
-	float MaxRange = 1.f;
+	bool bSwapInputs = false;
 
 	/**  */
-	UPROPERTY(EditAnywhere, Category = "Field")
-	float Default = 0.f;
+	UPROPERTY(meta = (DataflowOutput))
+	TArray<float> FieldFloatResult;
+
+	/**  */
+	UPROPERTY(meta = (DataflowOutput))
+	TArray<int32> FieldRemap;
+
+	FSumScalarFieldDataflowNode(const Dataflow::FNodeParameters& InParam, FGuid InGuid = FGuid::NewGuid())
+		: FDataflowNode(InParam, InGuid)
+	{
+		RegisterInputConnection(&FieldFloatLeft);
+		RegisterInputConnection(&FieldRemapLeft);
+		RegisterInputConnection(&FieldFloatRight);
+		RegisterInputConnection(&FieldRemapRight);
+		RegisterOutputConnection(&FieldFloatResult);
+		RegisterOutputConnection(&FieldRemap);
+	}
+
+	virtual void Evaluate(Dataflow::FContext& Context, const FDataflowOutput* Out) const override;
+
+};
+
+
+/**
+ *
+ *
+ *
+ */
+UENUM(BlueprintType)
+enum class EDataflowVectorFieldOperationType : uint8
+{
+	Dataflow_VectorFieldOperationType_Multiply		UMETA(DisplayName = "Multiply", ToolTip = "Multiply the fields output values : Output = Left * Right"),
+	Dataflow_VectorFieldFalloffType_Divide			UMETA(DisplayName = "Divide", ToolTip = "Divide the fields output values : Output = Left / Right"),
+	Dataflow_VectorFieldFalloffType_Add				UMETA(DisplayName = "Add", ToolTip = "Add the fields output values : Output = Left + Right"),
+	Dataflow_VectorFieldFalloffType_Substract		UMETA(DisplayName = "Subtract", ToolTip = "Subtract the fields output : Output = Left - Right"),
+	Dataflow_VectorFieldFalloffType_CrossProduct	UMETA(DisplayName = "Cross product", ToolTip = "Cross product of the fields output values: Output = Left x Right"),
+	//~~~
+	//256th entry
+	Dataflow_Max                UMETA(Hidden)
+};
+
+/**
+ *
+ *
+ *
+ */
+USTRUCT()
+struct FSumVectorFieldDataflowNode : public FDataflowNode
+{
+	GENERATED_USTRUCT_BODY()
+	DATAFLOW_NODE_DEFINE_INTERNAL(FSumVectorFieldDataflowNode, "SumVectorField", "Fields", "")
+
+public:
+	UPROPERTY(meta = (DataflowInput))
+	TArray<float> FieldFloat;
+
+	/**  */
+	UPROPERTY(meta = (DataflowInput))
+	TArray<int32> FieldFloatRemap;
+
+	UPROPERTY(meta = (DataflowInput))
+	TArray<FVector> FieldVectorLeft;
+
+	/**  */
+	UPROPERTY(meta = (DataflowInput))
+	TArray<int32> FieldRemapLeft;
+
+	UPROPERTY(meta = (DataflowInput))
+	TArray<FVector> FieldVectorRight;
+
+	/**  */
+	UPROPERTY(meta = (DataflowInput))
+	TArray<int32> FieldRemapRight;
+
+	/**  */
+	UPROPERTY(EditAnywhere, Category = "Field", meta = (DataflowInput))
+	float Magnitude = 1.f;
 
 	/**  */
 	UPROPERTY(EditAnywhere, Category = "Field")
-	EDataflowFieldFalloffType FalloffType = EDataflowFieldFalloffType::Dataflow_FieldFalloffType_Linear;
+	EDataflowVectorFieldOperationType Operation = EDataflowVectorFieldOperationType::Dataflow_VectorFieldFalloffType_Add;
+
+	UPROPERTY(EditAnywhere, Category = "Field")
+	bool bSwapVectorInputs = false;
 
 	/**  */
 	UPROPERTY(meta = (DataflowOutput))
-	TArray<float> WeightArray;
+	TArray<FVector> FieldVectorResult;
 
-	FBoxFalloffFieldDataflowNode(const Dataflow::FNodeParameters& InParam, FGuid InGuid = FGuid::NewGuid())
+	/**  */
+	UPROPERTY(meta = (DataflowOutput))
+	TArray<int32> FieldRemap;
+
+	FSumVectorFieldDataflowNode(const Dataflow::FNodeParameters& InParam, FGuid InGuid = FGuid::NewGuid())
 		: FDataflowNode(InParam, InGuid)
 	{
-		RegisterInputConnection(&VertexArray);
-		RegisterInputConnection(&VertexSelection);
-		RegisterInputConnection(&Box);
-		RegisterInputConnection(&Transform);
-		RegisterOutputConnection(&WeightArray);
+		RegisterInputConnection(&FieldFloat);
+		RegisterInputConnection(&FieldFloatRemap);
+		RegisterInputConnection(&FieldVectorLeft);
+		RegisterInputConnection(&FieldRemapLeft);
+		RegisterInputConnection(&FieldVectorRight);
+		RegisterInputConnection(&FieldRemapRight);
+		RegisterInputConnection(&Magnitude);
+		RegisterOutputConnection(&FieldVectorResult);
+		RegisterOutputConnection(&FieldRemap);
 	}
 
 	virtual void Evaluate(Dataflow::FContext& Context, const FDataflowOutput* Out) const override;
@@ -658,7 +983,52 @@ public:
 };
 
 
+/**
+ *
+ * Converts a sparse FloatArray (a selected subset of the whole incoming array) into a dense FloatArray
+ * (same number of elements as the incoming array using NumSamplePositions) using the Remap input
+ * NumSamplePositions controls the size of the output array, only indices smaller than l to than NumSamplePositions
+ * will be processed
+ *
+ */
+USTRUCT()
+struct FFieldMakeDenseFloatArrayDataflowNode : public FDataflowNode
+{
+	GENERATED_USTRUCT_BODY()
+	DATAFLOW_NODE_DEFINE_INTERNAL(FFieldMakeDenseFloatArrayDataflowNode, "FieldMakeDenseFloatArray", "Fields", "")
 
+public:
+	UPROPERTY(meta = (DataflowInput, DataflowIntrinsic))
+	TArray<float> FieldFloatInput;
+
+	/**  */
+	UPROPERTY(meta = (DataflowInput, DataflowIntrinsic))
+	TArray<int32> FieldRemap;
+
+	/**  */
+	UPROPERTY(EditAnywhere, Category = "Field", meta = (DataflowInput, DisplayName = "Number of Sample Positions"))
+	int32 NumSamplePositions = 0;
+
+	/**  */
+	UPROPERTY(EditAnywhere, Category = "Field")
+	float Default = 0.f;
+
+	/**  */
+	UPROPERTY(meta = (DataflowOutput))
+	TArray<float> FieldFloatResult;
+
+	FFieldMakeDenseFloatArrayDataflowNode(const Dataflow::FNodeParameters& InParam, FGuid InGuid = FGuid::NewGuid())
+		: FDataflowNode(InParam, InGuid)
+	{
+		RegisterInputConnection(&FieldFloatInput);
+		RegisterInputConnection(&FieldRemap);
+		RegisterInputConnection(&NumSamplePositions);
+		RegisterOutputConnection(&FieldFloatResult);
+	}
+
+	virtual void Evaluate(Dataflow::FContext& Context, const FDataflowOutput* Out) const override;
+
+};
 
 namespace Dataflow
 {
diff --git a/Engine/Plugins/Experimental/GeometryCollectionPlugin/Source/GeometryCollectionNodes/Public/Dataflow/GeometryCollectionNodes.h b/Engine/Plugins/Experimental/GeometryCollectionPlugin/Source/GeometryCollectionNodes/Public/Dataflow/GeometryCollectionNodes.h
--- a/Engine/Plugins/Experimental/GeometryCollectionPlugin/Source/GeometryCollectionNodes/Public/Dataflow/GeometryCollectionNodes.h
+++ b/Engine/Plugins/Experimental/GeometryCollectionPlugin/Source/GeometryCollectionNodes/Public/Dataflow/GeometryCollectionNodes.h
@@ -524,7 +524,7 @@ public:
 
 /**
  *
- * Converts an integer to a float
+ * 
  *
  */
 USTRUCT()
@@ -1741,20 +1741,13 @@ public:
 UENUM(BlueprintType)
 enum class EStandardGroupNameEnum : uint8
 {
-	/** Transform group */
-	Dataflow_EStandardGroupNameEnum_Transform UMETA(DisplayName = "Transform"),
-	/** Geometry group */
-	Dataflow_EStandardGroupNameEnum_Geometry UMETA(DisplayName = "Geometry"),
-	/** Faces group */
-	Dataflow_EStandardGroupNameEnum_Faces UMETA(DisplayName = "Faces"),
-	/** Vertices group */
-	Dataflow_EStandardGroupNameEnum_Vertices UMETA(DisplayName = "Vertices"),
-	/** Material group */
-	Dataflow_EStandardGroupNameEnum_Material UMETA(DisplayName = "Material"),
-	/** Breaking group */
-	Dataflow_EStandardGroupNameEnum_Breaking UMETA(DisplayName = "Breaking"),
-	/** User specified group */
-	Dataflow_EStandardGroupNameEnum_Custom UMETA(DisplayName = "Custom"),
+	Dataflow_EStandardGroupNameEnum_Transform	UMETA(DisplayName = "Transform", ToolTip = ""),
+	Dataflow_EStandardGroupNameEnum_Geometry	UMETA(DisplayName = "Geometry", ToolTip = ""),
+	Dataflow_EStandardGroupNameEnum_Faces		UMETA(DisplayName = "Faces", ToolTip = ""),
+	Dataflow_EStandardGroupNameEnum_Vertices	UMETA(DisplayName = "Vertices", ToolTip = ""),
+	Dataflow_EStandardGroupNameEnum_Material	UMETA(DisplayName = "Material", ToolTip = ""),
+	Dataflow_EStandardGroupNameEnum_Breaking	UMETA(DisplayName = "Breaking", ToolTip = ""),
+	Dataflow_EStandardGroupNameEnum_Custom		UMETA(DisplayName = "Custom", ToolTip = ""),
 	//~~~
 	//256th entry
 	Dataflow_Max                UMETA(Hidden)
@@ -1763,14 +1756,31 @@ enum class EStandardGroupNameEnum : uint8
 UENUM(BlueprintType)
 enum class ECustomAttributeTypeEnum : uint8
 {
-	/** Int custom attribute type */
-	Dataflow_CustomAttributeType_Int UMETA(DisplayName = "Int"),
-	/** Float custom attribute type */
-	Dataflow_CustomAttributeType_Float UMETA(DisplayName = "Float"),
-	/** Bool custom attribute type */
-	Dataflow_CustomAttributeType_Bool UMETA(DisplayName = "Bool"),
-	/** Vector custom attribute type */
-	Dataflow_CustomAttributeType_Vector UMETA(DisplayName = "Vector"),
+	Dataflow_CustomAttributeType_UInt8				UMETA(DisplayName = "UInt8", ToolTip = "Data type: UInt8"),
+	Dataflow_CustomAttributeType_Int32				UMETA(DisplayName = "Int32", ToolTip = "Data type: Int32"),
+	Dataflow_CustomAttributeType_Float				UMETA(DisplayName = "Float", ToolTip = "Data type: Float"),
+	Dataflow_CustomAttributeType_Double				UMETA(DisplayName = "Double", ToolTip = "Data type: Double"),
+	Dataflow_CustomAttributeType_Bool				UMETA(DisplayName = "Bool", ToolTip = "Data type: Bool"),
+	Dataflow_CustomAttributeType_String				UMETA(DisplayName = "String", ToolTip = "Data type: FString"),
+	Dataflow_CustomAttributeType_Vector2f			UMETA(DisplayName = "Vector2D", ToolTip = "Data type: FVector2f"),
+	Dataflow_CustomAttributeType_Vector3f			UMETA(DisplayName = "Vector", ToolTip = "Data type: FVector3f"),
+	Dataflow_CustomAttributeType_Vector3d			UMETA(DisplayName = "Vector3d", ToolTip = "Data type: FVector3d"),
+	Dataflow_CustomAttributeType_Vector4f			UMETA(DisplayName = "Vector4f", ToolTip = "Data type: FVector4f"),
+	Dataflow_CustomAttributeType_LinearColor		UMETA(DisplayName = "LinearColor", ToolTip = "Data type: FLinearColor"),
+	Dataflow_CustomAttributeType_Transform			UMETA(DisplayName = "Transform", ToolTip = "Data type: FTransform"),
+	Dataflow_CustomAttributeType_Quat4f				UMETA(DisplayName = "Quat", ToolTip = "Data type: FQuat4f"),
+	Dataflow_CustomAttributeType_Box				UMETA(DisplayName = "Box", ToolTip = "Data type: FBox"),
+	Dataflow_CustomAttributeType_Guid				UMETA(DisplayName = "Guid", ToolTip = "Data type: FGuid"),
+	Dataflow_CustomAttributeType_Int32Set			UMETA(DisplayName = "IntArray", ToolTip = "Data type: TSet<int32>"),
+	Dataflow_CustomAttributeType_Int32Array			UMETA(DisplayName = "Int32Array", ToolTip = "Data type: TArray<int32>"),
+	Dataflow_CustomAttributeType_IntVector			UMETA(DisplayName = "IntVector", ToolTip = "Data type: FIntVector"),
+	Dataflow_CustomAttributeType_IntVector2			UMETA(DisplayName = "IntVector2", ToolTip = "Data type: FIntVector2"),
+	Dataflow_CustomAttributeType_IntVector4			UMETA(DisplayName = "IntVector4", ToolTip = "Data type: FIntVector4"),
+	Dataflow_CustomAttributeType_IntVector2Array	UMETA(DisplayName = "IntVector2Array", ToolTip = "Data type: TArray<FIntVector2>"),
+	Dataflow_CustomAttributeType_FloatArray			UMETA(DisplayName = "FloatArray", ToolTip = "Data type: TArray<float>"),
+	Dataflow_CustomAttributeType_Vector2fArray		UMETA(DisplayName = "Vector2DArray", ToolTip = "Data type: TArray<FVector2f>"),
+	Dataflow_CustomAttributeType_FVector3fArray		UMETA(DisplayName = "FVectorArray", ToolTip = "Data type: TArray<FVector3f>"),
+
 	//~~~
 	//256th entry
 	Dataflow_Max                UMETA(Hidden)
@@ -1901,13 +1911,25 @@ public:
 	UPROPERTY(meta = (DataflowOutput));
 	TArray<float> FloatAttributeData;
 
+	/** Float type attribute data */
+	UPROPERTY(meta = (DataflowOutput));
+	TArray<double> DoubleAttributeData;
+
 	/** Int type attribute data */
 	UPROPERTY(meta = (DataflowOutput));
-	TArray<int32> IntAttributeData;
+	TArray<int32> Int32AttributeData;
 
-	/** Vector type attribute data */
+	/** Int type attribute data */
+	UPROPERTY(meta = (DataflowOutput));
+	TArray<FString> StringAttributeData;
+
+	/** Vector3f type attribute data */
 	UPROPERTY(meta = (DataflowOutput));
-	TArray<FVector3f> VectorAttributeData;
+	TArray<FVector3f> Vector3fAttributeData;
+
+	/** Vector3d type attribute data */
+	UPROPERTY(meta = (DataflowOutput));
+	TArray<FVector3d> Vector3dAttributeData;
 
 	FGetCollectionAttributeDataTypedDataflowNode(const Dataflow::FNodeParameters& InParam, FGuid InGuid = FGuid::NewGuid())
 		: FDataflowNode(InParam, InGuid)
@@ -1915,8 +1937,11 @@ public:
 		RegisterInputConnection(&Collection);
 		RegisterOutputConnection(&BoolAttributeData);
 		RegisterOutputConnection(&FloatAttributeData);
-		RegisterOutputConnection(&IntAttributeData);
-		RegisterOutputConnection(&VectorAttributeData);
+		RegisterOutputConnection(&DoubleAttributeData);
+		RegisterOutputConnection(&Int32AttributeData);
+		RegisterOutputConnection(&StringAttributeData);
+		RegisterOutputConnection(&Vector3fAttributeData);
+		RegisterOutputConnection(&Vector3dAttributeData);
 	}
 
 	virtual void Evaluate(Dataflow::FContext& Context, const FDataflowOutput* Out) const override;
@@ -1960,13 +1985,25 @@ public:
 	UPROPERTY(meta = (DataflowInput));
 	TArray<float> FloatAttributeData;
 
+	/** Float type attribute data */
+	UPROPERTY(meta = (DataflowInput));
+	TArray<double> DoubleAttributeData;
+
 	/** Int type attribute data */
 	UPROPERTY(meta = (DataflowInput));
-	TArray<int32> IntAttributeData;
+	TArray<int32> Int32AttributeData;
+
+	/** Int type attribute data */
+	UPROPERTY(meta = (DataflowInput));
+	TArray<FString> StringAttributeData;
+
+	/** Vector3f type attribute data */
+	UPROPERTY(meta = (DataflowInput));
+	TArray<FVector3f> Vector3fAttributeData;
 
-	/** Vector type attribute data */
+	/** Vector3d type attribute data */
 	UPROPERTY(meta = (DataflowInput));
-	TArray<FVector> VectorAttributeData;
+	TArray<FVector3d> Vector3dAttributeData;
 
 	FSetCollectionAttributeDataTypedDataflowNode(const Dataflow::FNodeParameters& InParam, FGuid InGuid = FGuid::NewGuid())
 		: FDataflowNode(InParam, InGuid)
@@ -1974,8 +2011,11 @@ public:
 		RegisterInputConnection(&Collection);
 		RegisterInputConnection(&BoolAttributeData);
 		RegisterInputConnection(&FloatAttributeData);
-		RegisterInputConnection(&IntAttributeData);
-		RegisterInputConnection(&VectorAttributeData);
+		RegisterInputConnection(&DoubleAttributeData);
+		RegisterInputConnection(&Int32AttributeData);
+		RegisterInputConnection(&StringAttributeData);
+		RegisterInputConnection(&Vector3fAttributeData);
+		RegisterInputConnection(&Vector3dAttributeData);
 		RegisterOutputConnection(&Collection, &Collection);
 	}
 
@@ -2060,10 +2100,10 @@ public:
  *
  */
 USTRUCT()
-struct FSetVertexColorInCollectionDataflowNode : public FDataflowNode
+struct FSetVertexColorInCollectionFromVertexSelectionDataflowNode : public FDataflowNode
 {
 	GENERATED_USTRUCT_BODY()
-	DATAFLOW_NODE_DEFINE_INTERNAL(FSetVertexColorInCollectionDataflowNode, "SetVertexColorInCollection", "Collection|Utilities", "")
+	DATAFLOW_NODE_DEFINE_INTERNAL(FSetVertexColorInCollectionFromVertexSelectionDataflowNode, "SetVertexColorInCollectionFromVertexSelection", "Collection|Utilities", "")
 
 public:
 	/** Collection */
@@ -2082,7 +2122,7 @@ public:
 	UPROPERTY(EditAnywhere, Category = "Color", meta = (DisplayName = "NonSelected Color"))
 	FLinearColor NonSelectedColor = FLinearColor(FColor::Blue);
 
-	FSetVertexColorInCollectionDataflowNode(const Dataflow::FNodeParameters& InParam, FGuid InGuid = FGuid::NewGuid())
+	FSetVertexColorInCollectionFromVertexSelectionDataflowNode(const Dataflow::FNodeParameters& InParam, FGuid InGuid = FGuid::NewGuid())
 		: FDataflowNode(InParam, InGuid)
 	{
 		RegisterInputConnection(&Collection);
@@ -2109,13 +2149,13 @@ struct FMakeTransformDataflowNode : public FDataflowNode
 public:
 
 	UPROPERTY(EditAnywhere, Category = "Transform", meta = (DataflowInput, DisplayName = "Translation"));
-	FVector InTranslation = FVector(0,0,0);
+	FVector InTranslation = FVector(0, 0, 0);
 
 	UPROPERTY(EditAnywhere, Category = "Transform", meta = (DataflowInput, DisplayName = "Rotation"));
 	FVector InRotation = FVector(0, 0, 0);
 
 	UPROPERTY(EditAnywhere, Category = "Transform", meta = (DataflowInput, DisplayName = "Scale"));
-	FVector InScale = FVector(1,1,1);
+	FVector InScale = FVector(1, 1, 1);
 
 	UPROPERTY(meta = (DataflowOutput, DisplayName = "Transform"));
 	FTransform OutTransform = FTransform::Identity;
@@ -2133,6 +2173,132 @@ public:
 
 };
 
+
+/**
+ *
+ *
+ *
+ */
+USTRUCT()
+struct FSetVertexColorInCollectionFromFloatArrayDataflowNode : public FDataflowNode
+{
+	GENERATED_USTRUCT_BODY()
+	DATAFLOW_NODE_DEFINE_INTERNAL(FSetVertexColorInCollectionFromFloatArrayDataflowNode, "SetVertexColorInCollectionFromFloatArray", "Collection|Utilities", "")
+
+public:
+	/** Collection */
+	UPROPERTY(meta = (DataflowInput, DataflowOutput, DataflowPassthrough = "Collection", DataflowIntrinsic))
+	FManagedArrayCollection Collection;
+
+	/**  */
+	UPROPERTY(meta = (DataflowInput, DataflowIntrinsic))
+	TArray<float> FloatArray;
+
+	/**  */
+	UPROPERTY(EditAnywhere, Category = "Color")
+	float Scale = 1.f;
+
+	FSetVertexColorInCollectionFromFloatArrayDataflowNode(const Dataflow::FNodeParameters& InParam, FGuid InGuid = FGuid::NewGuid())
+		: FDataflowNode(InParam, InGuid)
+	{
+		RegisterInputConnection(&Collection);
+		RegisterInputConnection(&FloatArray);
+		RegisterOutputConnection(&Collection, &Collection);
+	}
+
+	virtual void Evaluate(Dataflow::FContext& Context, const FDataflowOutput* Out) const override;
+
+};
+
+
+/**
+ *
+ * Normalize the selected float data in a FloatArray
+ *
+ */
+USTRUCT()
+struct FFloatArrayNormalizeDataflowNode : public FDataflowNode
+{
+	GENERATED_USTRUCT_BODY()
+	DATAFLOW_NODE_DEFINE_INTERNAL(FFloatArrayNormalizeDataflowNode, "FloatArrayNormalize", "Math|Float", "")
+
+public:
+	/** Input VectorArray */
+	UPROPERTY(meta = (DataflowInput, DataflowIntrinsic))
+	TArray<float> InFloatArray;
+
+	/** Selection for the operation */
+	UPROPERTY(meta = (DataflowInput))
+	FDataflowVertexSelection Selection;
+
+	/**  */
+	UPROPERTY(EditAnywhere, Category = "Field", meta = (DataflowInput))
+	float MinRange = 0.f;
+
+	/**  */
+	UPROPERTY(EditAnywhere, Category = "Field", meta = (DataflowInput))
+	float MaxRange = 1.f;
+
+	/**  */
+	UPROPERTY(meta = (DataflowOutput))
+	TArray<float> OutFloatArray;
+
+	FFloatArrayNormalizeDataflowNode(const Dataflow::FNodeParameters& InParam, FGuid InGuid = FGuid::NewGuid())
+		: FDataflowNode(InParam, InGuid)
+	{
+		RegisterInputConnection(&InFloatArray);
+		RegisterInputConnection(&Selection);
+		RegisterInputConnection(&MinRange);
+		RegisterInputConnection(&MaxRange);
+		RegisterOutputConnection(&OutFloatArray);
+	}
+
+	virtual void Evaluate(Dataflow::FContext& Context, const FDataflowOutput* Out) const override;
+
+};
+
+
+/**
+ *
+ * Normalize all the selected vectors in a VectorArray
+ *
+ */
+USTRUCT()
+struct FVectorArrayNormalizeDataflowNode : public FDataflowNode
+{
+	GENERATED_USTRUCT_BODY()
+	DATAFLOW_NODE_DEFINE_INTERNAL(FVectorArrayNormalizeDataflowNode, "VectorArrayNormalize", "Math|Vector", "")
+
+public:
+	/** Input VectorArray */
+	UPROPERTY(meta = (DataflowInput, DataflowIntrinsic))
+	TArray<FVector> InVectorArray;
+
+	/** Selection for the operation */
+	UPROPERTY(meta = (DataflowInput))
+	FDataflowVertexSelection Selection;
+
+	/**  */
+	UPROPERTY(EditAnywhere, Category = "Field", meta = (DataflowInput))
+	float Magnitude = 1.f;
+
+	/**  */
+	UPROPERTY(meta = (DataflowOutput))
+	TArray<FVector> OutVectorArray;
+
+	FVectorArrayNormalizeDataflowNode(const Dataflow::FNodeParameters& InParam, FGuid InGuid = FGuid::NewGuid())
+		: FDataflowNode(InParam, InGuid)
+	{
+		RegisterInputConnection(&InVectorArray);
+		RegisterInputConnection(&Selection);
+		RegisterInputConnection(&Magnitude);
+		RegisterOutputConnection(&OutVectorArray);
+	}
+
+	virtual void Evaluate(Dataflow::FContext& Context, const FDataflowOutput* Out) const override;
+
+};
+
 /**
  *
  * Description for this node
@@ -2166,6 +2332,47 @@ public:
 
 };
 
+/**
+ *
+ *
+ *
+ */
+USTRUCT()
+struct FMakeQuaternionDataflowNode : public FDataflowNode
+{
+	GENERATED_USTRUCT_BODY()
+	DATAFLOW_NODE_DEFINE_INTERNAL(FMakeQuaternionDataflowNode, "MakeQuaternion", "Math|Vector", "")
+
+public:
+	UPROPERTY(EditAnywhere, Category = "Quaternion ", meta = (DataflowInput));
+	float X = float(0.0);
+
+	UPROPERTY(EditAnywhere, Category = "Quaternion", meta = (DataflowInput));
+	float Y = float(0.0);
+
+	UPROPERTY(EditAnywhere, Category = "Quaternion", meta = (DataflowInput));
+	float Z = float(0.0);
+
+	UPROPERTY(EditAnywhere, Category = "Quaternion", meta = (DataflowInput));
+	float W = float(0.0);
+
+	UPROPERTY(meta = (DataflowOutput, DisplayName = "Quaternion"))
+	FQuat Quaternion = FQuat(ForceInitToZero);
+
+	FMakeQuaternionDataflowNode(const Dataflow::FNodeParameters& InParam, FGuid InGuid = FGuid::NewGuid())
+		: FDataflowNode(InParam, InGuid)
+	{
+		RegisterInputConnection(&X);
+		RegisterInputConnection(&Y);
+		RegisterInputConnection(&Z);
+		RegisterInputConnection(&W);
+		RegisterOutputConnection(&Quaternion);
+	}
+
+	virtual void Evaluate(Dataflow::FContext& Context, const FDataflowOutput* Out) const override;
+
+};
+
 
 namespace Dataflow
 {
diff --git a/Engine/Plugins/Experimental/GeometryCollectionPlugin/Source/GeometryCollectionNodes/Public/Dataflow/GeometryCollectionSelectionNodes.h b/Engine/Plugins/Experimental/GeometryCollectionPlugin/Source/GeometryCollectionNodes/Public/Dataflow/GeometryCollectionSelectionNodes.h
--- a/Engine/Plugins/Experimental/GeometryCollectionPlugin/Source/GeometryCollectionNodes/Public/Dataflow/GeometryCollectionSelectionNodes.h
+++ b/Engine/Plugins/Experimental/GeometryCollectionPlugin/Source/GeometryCollectionNodes/Public/Dataflow/GeometryCollectionSelectionNodes.h
@@ -1103,6 +1103,89 @@ public:
 };
 
 
+/**
+ *
+ * Outputs the specified percentage of the selected vertices
+ *
+ */
+USTRUCT(meta = (DataflowGeometryCollection))
+struct FCollectionVertexSelectionByPercentageDataflowNode : public FDataflowNode
+{
+	GENERATED_USTRUCT_BODY()
+	DATAFLOW_NODE_DEFINE_INTERNAL(FCollectionVertexSelectionByPercentageDataflowNode, "CollectionVertexSelectByPercentage", "GeometryCollection|Selection|Vertex", "")
+
+public:
+	/** Array of the selected bone indicies */
+	UPROPERTY(meta = (DataflowInput, DataflowOutput, DisplayName = "VertexSelection", DataflowPassthrough = "VertexSelection", DataflowIntrinsic))
+	FDataflowVertexSelection VertexSelection;
+
+	/** Percentage to keep from the original selection */
+	UPROPERTY(EditAnywhere, Category = "Selection", meta = (UIMin = 0, UIMax = 100))
+	int32 Percentage = 100;
+
+	/** Sets the random generation to deterministic */
+	UPROPERTY(EditAnywhere, Category = "Random")
+	bool bDeterministic = false;
+
+	/** Seed value for the random generation */
+	UPROPERTY(EditAnywhere, Category = "Random", meta = (DataflowInput, EditCondition = "bDeterministic"))
+	float RandomSeed = 0.f;
+
+	FCollectionVertexSelectionByPercentageDataflowNode(const Dataflow::FNodeParameters& InParam, FGuid InGuid = FGuid::NewGuid())
+		: FDataflowNode(InParam, InGuid)
+	{
+		RandomSeed = FMath::RandRange(-100000, 100000);
+		RegisterInputConnection(&VertexSelection);
+		RegisterInputConnection(&Percentage);
+		RegisterInputConnection(&RandomSeed);
+		RegisterOutputConnection(&VertexSelection, &VertexSelection);
+	}
+
+	virtual void Evaluate(Dataflow::FContext& Context, const FDataflowOutput* Out) const override;
+
+};
+
+
+/**
+ *
+ * Runs boolean operation on VertexSelections
+ *
+ */
+USTRUCT(meta = (DataflowGeometryCollection))
+struct FCollectionVertexSelectionSetOperationDataflowNode : public FDataflowNode
+{
+	GENERATED_USTRUCT_BODY()
+	DATAFLOW_NODE_DEFINE_INTERNAL(FCollectionVertexSelectionSetOperationDataflowNode, "CollectionVertexSelectionSetOperation", "GeometryCollection|Selection|Vertex", "")
+
+public:
+	/** Boolean operation */
+	UPROPERTY(EditAnywhere, Category = "Compare");
+	ESetOperationEnum Operation = ESetOperationEnum::Dataflow_SetOperation_AND;
+
+	/** Array of the selected vertex indicies */
+	UPROPERTY(meta = (DataflowInput, DisplayName = "VertexSelectionA", DataflowIntrinsic))
+	FDataflowVertexSelection VertexSelectionA;
+
+	/** Array of the selected vertex indicies */
+	UPROPERTY(meta = (DataflowInput, DisplayName = "VertexSelectionB", DataflowIntrinsic))
+	FDataflowVertexSelection VertexSelectionB;
+
+	/** Array of the selected vertex indicies after operation */
+	UPROPERTY(meta = (DataflowOutput, DisplayName = "VertexSelection", DataflowPassthrough = "VertexSelectionA"))
+	FDataflowVertexSelection VertexSelection;
+
+	FCollectionVertexSelectionSetOperationDataflowNode(const Dataflow::FNodeParameters& InParam, FGuid InGuid = FGuid::NewGuid())
+		: FDataflowNode(InParam, InGuid)
+	{
+		RegisterInputConnection(&VertexSelectionA);
+		RegisterInputConnection(&VertexSelectionB);
+		RegisterOutputConnection(&VertexSelection, &VertexSelectionA);
+	}
+
+	virtual void Evaluate(Dataflow::FContext& Context, const FDataflowOutput* Out) const override;
+
+};
+
 namespace Dataflow
 {
 	void GeometryCollectionSelectionNodes();
diff --git a/Engine/Source/Runtime/Experimental/Dataflow/Core/Private/Dataflow/DataflowSettings.cpp b/Engine/Source/Runtime/Experimental/Dataflow/Core/Private/Dataflow/DataflowSettings.cpp
--- a/Engine/Source/Runtime/Experimental/Dataflow/Core/Private/Dataflow/DataflowSettings.cpp
+++ b/Engine/Source/Runtime/Experimental/Dataflow/Core/Private/Dataflow/DataflowSettings.cpp
@@ -10,6 +10,7 @@ UDataflowSettings::UDataflowSettings(const FObjectInitializer& ObjectInitlaizer)
 	ManagedArrayCollectionPinTypeColor = CManagedArrayCollectionPinTypeColor;
 	ArrayPinTypeColor = CArrayPinTypeColor;
 	BoxPinTypeColor = CBoxPinTypeColor;
+	SpherePinTypeColor = CSpherePinTypeColor;
 }
 
 FName UDataflowSettings::GetCategoryName() const
diff --git a/Engine/Source/Runtime/Experimental/Dataflow/Core/Public/Dataflow/DataflowSettings.h b/Engine/Source/Runtime/Experimental/Dataflow/Core/Public/Dataflow/DataflowSettings.h
--- a/Engine/Source/Runtime/Experimental/Dataflow/Core/Public/Dataflow/DataflowSettings.h
+++ b/Engine/Source/Runtime/Experimental/Dataflow/Core/Public/Dataflow/DataflowSettings.h
@@ -10,6 +10,7 @@
 static const FLinearColor CManagedArrayCollectionPinTypeColor = FLinearColor(0.353393f, 0.454175f, 1.0f, 1.0f);
 static const FLinearColor CArrayPinTypeColor = FLinearColor(1.0f, 0.172585f, 0.0f, 1.0f);
 static const FLinearColor CBoxPinTypeColor = FLinearColor(0.013575f, 0.770000f, 0.429609f, 1.0f);
+static const FLinearColor CSpherePinTypeColor = FLinearColor(0.2f, 0.6f, 1.f, 1.0f);
 
 typedef TMap<FName, FNodeColors> FNodeColorsMap;
 
@@ -42,6 +43,10 @@ class DATAFLOWCORE_API UDataflowSettings : public UDeveloperSettings
 	UPROPERTY(config, EditAnywhere, Category = PinColors)
 	FLinearColor BoxPinTypeColor;
 
+	/** FSphere pin type color. The other pin colors are defined in the general editor settings. */
+	UPROPERTY(config, EditAnywhere, Category = PinColors)
+	FLinearColor SpherePinTypeColor;
+
 	UPROPERTY(config, EditAnywhere, Category = NodeColors)
 	TMap<FName, FNodeColors> NodeColorsMap;
 
