Dataflow:
- Fixed missing initialization

diff --git a/Engine/Plugins/Experimental/GeometryCollectionPlugin/Source/GeometryCollectionNodes/Public/Dataflow/GeometryCollectionFieldNodes.h b/Engine/Plugins/Experimental/GeometryCollectionPlugin/Source/GeometryCollectionNodes/Public/Dataflow/GeometryCollectionFieldNodes.h
--- a/Engine/Plugins/Experimental/GeometryCollectionPlugin/Source/GeometryCollectionNodes/Public/Dataflow/GeometryCollectionFieldNodes.h
+++ b/Engine/Plugins/Experimental/GeometryCollectionPlugin/Source/GeometryCollectionNodes/Public/Dataflow/GeometryCollectionFieldNodes.h
@@ -91,7 +91,7 @@ public:
 
 	/**  */
 	UPROPERTY(meta = (DataflowOutput))
-	int32 NumSamplePositions;
+	int32 NumSamplePositions = 0;
 
 	FRadialFalloffFieldDataflowNode(const Dataflow::FNodeParameters& InParam, FGuid InGuid = FGuid::NewGuid())
 		: FDataflowNode(InParam, InGuid)
@@ -176,7 +176,7 @@ public:
 
 	/**  */
 	UPROPERTY(meta = (DataflowOutput))
-	int32 NumSamplePositions;
+	int32 NumSamplePositions = 0;
 
 	FBoxFalloffFieldDataflowNode(const Dataflow::FNodeParameters& InParam, FGuid InGuid = FGuid::NewGuid())
 		: FDataflowNode(InParam, InGuid)
@@ -270,7 +270,7 @@ public:
 
 	/**  */
 	UPROPERTY(meta = (DataflowOutput))
-	int32 NumSamplePositions;
+	int32 NumSamplePositions = 0;
 
 	FPlaneFalloffFieldDataflowNode(const Dataflow::FNodeParameters& InParam, FGuid InGuid = FGuid::NewGuid())
 		: FDataflowNode(InParam, InGuid)
@@ -362,7 +362,7 @@ public:
 
 	/**  */
 	UPROPERTY(meta = (DataflowOutput))
-	int32 NumSamplePositions;
+	int32 NumSamplePositions = 0;
 
 	FRadialIntMaskFieldDataflowNode(const Dataflow::FNodeParameters& InParam, FGuid InGuid = FGuid::NewGuid())
 		: FDataflowNode(InParam, InGuid)
@@ -417,7 +417,7 @@ public:
 
 	/**  */
 	UPROPERTY(meta = (DataflowOutput))
-	int32 NumSamplePositions;
+	int32 NumSamplePositions = 0;
 
 	FUniformScalarFieldDataflowNode(const Dataflow::FNodeParameters& InParam, FGuid InGuid = FGuid::NewGuid())
 		: FDataflowNode(InParam, InGuid)
@@ -473,7 +473,7 @@ public:
 
 	/**  */
 	UPROPERTY(meta = (DataflowOutput))
-	int32 NumSamplePositions;
+	int32 NumSamplePositions = 0;
 
 	FUniformVectorFieldDataflowNode(const Dataflow::FNodeParameters& InParam, FGuid InGuid = FGuid::NewGuid())
 		: FDataflowNode(InParam, InGuid)
@@ -530,7 +530,7 @@ public:
 
 	/**  */
 	UPROPERTY(meta = (DataflowOutput))
-	int32 NumSamplePositions;
+	int32 NumSamplePositions = 0;
 
 	FRadialVectorFieldDataflowNode(const Dataflow::FNodeParameters& InParam, FGuid InGuid = FGuid::NewGuid())
 		: FDataflowNode(InParam, InGuid)
@@ -583,7 +583,7 @@ public:
 
 	/**  */
 	UPROPERTY(meta = (DataflowOutput))
-	int32 NumSamplePositions;
+	int32 NumSamplePositions = 0;
 
 	FRandomVectorFieldDataflowNode(const Dataflow::FNodeParameters& InParam, FGuid InGuid = FGuid::NewGuid())
 		: FDataflowNode(InParam, InGuid)
@@ -643,7 +643,7 @@ public:
 
 	/**  */
 	UPROPERTY(meta = (DataflowOutput))
-	int32 NumSamplePositions;
+	int32 NumSamplePositions = 0;
 
 	FNoiseFieldDataflowNode(const Dataflow::FNodeParameters& InParam, FGuid InGuid = FGuid::NewGuid())
 		: FDataflowNode(InParam, InGuid)
@@ -697,7 +697,7 @@ public:
 
 	/**  */
 	UPROPERTY(meta = (DataflowOutput))
-	int32 NumSamplePositions;
+	int32 NumSamplePositions = 0;
 
 	FUniformIntegerFieldDataflowNode(const Dataflow::FNodeParameters& InParam, FGuid InGuid = FGuid::NewGuid())
 		: FDataflowNode(InParam, InGuid)
@@ -790,7 +790,7 @@ public:
 
 	/**  */
 	UPROPERTY(meta = (DataflowOutput))
-	int32 NumSamplePositions;
+	int32 NumSamplePositions = 0;
 
 	FWaveScalarFieldDataflowNode(const Dataflow::FNodeParameters& InParam, FGuid InGuid = FGuid::NewGuid())
 		: FDataflowNode(InParam, InGuid)
