Invert Xf dataflow node.

diff --git a/Engine/Plugins/Experimental/GeometryCollectionPlugin/Source/GeometryCollectionNodes/Private/Dataflow/GeometryCollectionNodes.cpp b/Engine/Plugins/Experimental/GeometryCollectionPlugin/Source/GeometryCollectionNodes/Private/Dataflow/GeometryCollectionNodes.cpp
--- a/Engine/Plugins/Experimental/GeometryCollectionPlugin/Source/GeometryCollectionNodes/Private/Dataflow/GeometryCollectionNodes.cpp
+++ b/Engine/Plugins/Experimental/GeometryCollectionPlugin/Source/GeometryCollectionNodes/Private/Dataflow/GeometryCollectionNodes.cpp
@@ -104,7 +104,7 @@ namespace Dataflow
 		DATAFLOW_NODE_REGISTER_CREATION_FACTORY(FVectorArrayNormalizeDataflowNode);
 		DATAFLOW_NODE_REGISTER_CREATION_FACTORY(FMakeQuaternionDataflowNode);
 		DATAFLOW_NODE_REGISTER_CREATION_FACTORY(FMultiplyTransformDataflowNode);
-		
+		DATAFLOW_NODE_REGISTER_CREATION_FACTORY(FInvertTransformDataflowNode);
 
 		// GeometryCollection
 		DATAFLOW_NODE_REGISTER_CREATION_FACTORY_NODE_COLORS_BY_CATEGORY("GeometryCollection", FLinearColor(0.55f, 0.45f, 1.0f), CDefaultNodeBodyTintColor);
@@ -1637,6 +1637,12 @@ void FMultiplyTransformDataflowNode::Evaluate(Dataflow::FContext& Context, const
 	}
 }
 
-
-
-
+void FInvertTransformDataflowNode::Evaluate(Dataflow::FContext& Context, const FDataflowOutput* Out) const
+{
+	if (Out->IsA<FTransform>(&OutTransform))
+	{
+		FTransform InXf = GetValue<FTransform>(Context, &InTransform, FTransform::Identity);
+		FTransform OutXf = InXf.Inverse();
+		SetValue<FTransform>(Context, OutXf, &OutTransform);
+	}
+}
diff --git a/Engine/Plugins/Experimental/GeometryCollectionPlugin/Source/GeometryCollectionNodes/Public/Dataflow/GeometryCollectionNodes.h b/Engine/Plugins/Experimental/GeometryCollectionPlugin/Source/GeometryCollectionNodes/Public/Dataflow/GeometryCollectionNodes.h
--- a/Engine/Plugins/Experimental/GeometryCollectionPlugin/Source/GeometryCollectionNodes/Public/Dataflow/GeometryCollectionNodes.h
+++ b/Engine/Plugins/Experimental/GeometryCollectionPlugin/Source/GeometryCollectionNodes/Public/Dataflow/GeometryCollectionNodes.h
@@ -2332,6 +2332,35 @@ public:
 
 };
 
+/**
+ *
+ * Invert a transform.
+ *
+ */
+USTRUCT()
+struct FInvertTransformDataflowNode : public FDataflowNode
+{
+	GENERATED_USTRUCT_BODY()
+	DATAFLOW_NODE_DEFINE_INTERNAL(FInvertTransformDataflowNode, "InvertTransform", "Math|Transform", "")
+
+public:
+	UPROPERTY(meta = (DataflowInput, DisplayName = "InTransform"));
+	FTransform InTransform = FTransform::Identity;
+
+	UPROPERTY(meta = (DataflowOutput, DisplayName = "OutTransform"));
+	FTransform OutTransform = FTransform::Identity;
+
+	FInvertTransformDataflowNode(const Dataflow::FNodeParameters& InParam, FGuid InGuid = FGuid::NewGuid())
+		: FDataflowNode(InParam, InGuid)
+	{
+		RegisterInputConnection(&InTransform);
+		RegisterOutputConnection(&OutTransform);
+	}
+
+	virtual void Evaluate(Dataflow::FContext& Context, const FDataflowOutput* Out) const override;
+};
+
+
 /**
  *
  *
