Dataflow:
- Added math library nodes
- Moved make nodes to separate file
- Moved conversion nodes to separate file
- bug fixes
(submitted on behalf of gustav.melich)

diff --git a/Engine/Plugins/Experimental/GeometryCollectionPlugin/Source/GeometryCollectionNodes/Private/Dataflow/GeometryCollectionConversionNodes.cpp b/Engine/Plugins/Experimental/GeometryCollectionPlugin/Source/GeometryCollectionNodes/Private/Dataflow/GeometryCollectionConversionNodes.cpp
new file mode 100644
--- /dev/null
+++ b/Engine/Plugins/Experimental/GeometryCollectionPlugin/Source/GeometryCollectionNodes/Private/Dataflow/GeometryCollectionConversionNodes.cpp
@@ -0,0 +1,116 @@
+// Copyright Epic Games, Inc. All Rights Reserved.
+
+#include "Dataflow/GeometryCollectionConversionNodes.h"
+#include "Dataflow/DataflowCore.h"
+
+#include UE_INLINE_GENERATED_CPP_BY_NAME(GeometryCollectionConversionNodes)
+
+namespace Dataflow
+{
+	void GeometryCollectionConversionNodes()
+	{
+		static const FLinearColor CDefaultNodeBodyTintColor = FLinearColor(0.f, 0.f, 0.f, 0.5f);
+
+		DATAFLOW_NODE_REGISTER_CREATION_FACTORY(FVectorToStringDataflowNode);
+		DATAFLOW_NODE_REGISTER_CREATION_FACTORY(FFloatToStringDataflowNode);
+		DATAFLOW_NODE_REGISTER_CREATION_FACTORY(FIntToStringDataflowNode);
+		DATAFLOW_NODE_REGISTER_CREATION_FACTORY(FBoolToStringDataflowNode);
+		DATAFLOW_NODE_REGISTER_CREATION_FACTORY(FIntToFloatDataflowNode);
+		DATAFLOW_NODE_REGISTER_CREATION_FACTORY(FFloatToIntDataflowNode);
+		DATAFLOW_NODE_REGISTER_CREATION_FACTORY(FIntToBoolDataflowNode);
+		DATAFLOW_NODE_REGISTER_CREATION_FACTORY(FBoolToIntDataflowNode);
+
+
+		// Generators
+		DATAFLOW_NODE_REGISTER_CREATION_FACTORY_NODE_COLORS_BY_CATEGORY("Generators", FLinearColor(.4f, 0.8f, 0.f), CDefaultNodeBodyTintColor);
+	}
+}
+
+void FVectorToStringDataflowNode::Evaluate(Dataflow::FContext& Context, const FDataflowOutput* Out) const
+{
+	if (Out->IsA<FString>(&String))
+	{
+		FString Value = GetValue<FVector>(Context, &Vector).ToString();
+		SetValue<FString>(Context, Value, &String);
+	}
+}
+
+void FFloatToStringDataflowNode::Evaluate(Dataflow::FContext& Context, const FDataflowOutput* Out) const
+{
+	if (Out->IsA<FString>(&String))
+	{
+		FString Value = FString::Printf(TEXT("%f"), GetValue<float>(Context, &Float));
+		SetValue<FString>(Context, Value, &String);
+	}
+}
+
+void FIntToStringDataflowNode::Evaluate(Dataflow::FContext& Context, const FDataflowOutput* Out) const
+{
+	if (Out->IsA<FString>(&String))
+	{
+		FString Value = FString::Printf(TEXT("%d"), GetValue<int32>(Context, &Int));
+		SetValue<FString>(Context, Value, &String);
+	}
+}
+
+void FBoolToStringDataflowNode::Evaluate(Dataflow::FContext& Context, const FDataflowOutput* Out) const
+{
+	if (Out->IsA<FString>(&String))
+	{
+		FString Value = FString::Printf(TEXT("%s"), GetValue<bool>(Context, &Bool) ? TEXT("true") : TEXT("false"));
+		SetValue<FString>(Context, Value, &String);
+	}
+}
+
+void FIntToFloatDataflowNode::Evaluate(Dataflow::FContext& Context, const FDataflowOutput* Out) const
+{
+	if (Out->IsA<float>(&Float))
+	{
+		float Value = float(GetValue<int32>(Context, &Int));
+		SetValue<float>(Context, Value, &Float);
+	}
+}
+
+void FFloatToIntDataflowNode::Evaluate(Dataflow::FContext& Context, const FDataflowOutput* Out) const
+{
+	if (Out->IsA<int32>(&Int))
+	{
+		float FloatVal = GetValue<float>(Context, &Float);
+		if (Function == EFloatToIntFunctionEnum::Dataflow_FloatToInt_Function_Floor)
+		{
+			SetValue<int32>(Context, FMath::FloorToInt32(FloatVal), &Int);
+		}
+		else if (Function == EFloatToIntFunctionEnum::Dataflow_FloatToInt_Function_Ceil)
+		{
+			SetValue<int32>(Context, FMath::CeilToInt32(FloatVal), &Int);
+		}
+		else if (Function == EFloatToIntFunctionEnum::Dataflow_FloatToInt_Function_Round)
+		{
+			SetValue<int32>(Context, FMath::RoundToInt32(FloatVal), &Int);
+		}
+		else if (Function == EFloatToIntFunctionEnum::Dataflow_FloatToInt_Function_Truncate)
+		{
+			SetValue<int32>(Context, int32(FMath::TruncToFloat(FloatVal)), &Int);
+		}
+	}
+}
+
+void FIntToBoolDataflowNode::Evaluate(Dataflow::FContext& Context, const FDataflowOutput* Out) const
+{
+	if (Out->IsA<bool>(&Bool))
+	{
+		bool Value = GetValue<int32>(Context, &Int, Int) == 0 ? false : true;
+		SetValue<bool>(Context, Value, &Bool);
+	}
+}
+
+void FBoolToIntDataflowNode::Evaluate(Dataflow::FContext& Context, const FDataflowOutput* Out) const
+{
+	if (Out->IsA<int32>(&Int))
+	{
+		int32 Value = (int32)GetValue<bool>(Context, &Bool, Bool);
+		SetValue<int32>(Context, Value, &Int);
+	}
+}
+
+
diff --git a/Engine/Plugins/Experimental/GeometryCollectionPlugin/Source/GeometryCollectionNodes/Private/Dataflow/GeometryCollectionMakeNodes.cpp b/Engine/Plugins/Experimental/GeometryCollectionPlugin/Source/GeometryCollectionNodes/Private/Dataflow/GeometryCollectionMakeNodes.cpp
new file mode 100644
--- /dev/null
+++ b/Engine/Plugins/Experimental/GeometryCollectionPlugin/Source/GeometryCollectionNodes/Private/Dataflow/GeometryCollectionMakeNodes.cpp
@@ -0,0 +1,169 @@
+// Copyright Epic Games, Inc. All Rights Reserved.
+
+#include "Dataflow/GeometryCollectionMakeNodes.h"
+#include "Dataflow/DataflowCore.h"
+
+#include "Engine/Engine.h"
+#include "Engine/StaticMesh.h"
+#include "GeometryCollection/Facades/CollectionMeshFacade.h"
+#include "GeometryCollection/GeometryCollectionObject.h"
+#include "GeometryCollection/ManagedArrayCollection.h"
+#include "GeometryCollection/GeometryCollection.h"
+#include "GeometryCollection/GeometryCollectionEngineUtility.h"
+#include "GeometryCollection/GeometryCollectionEngineRemoval.h"
+#include "GeometryCollection/GeometryCollectionEngineConversion.h"
+#include "Logging/LogMacros.h"
+#include "Templates/SharedPointer.h"
+#include "UObject/UnrealTypePrivate.h"
+#include "DynamicMeshToMeshDescription.h"
+#include "MeshDescriptionToDynamicMesh.h"
+#include "StaticMeshAttributes.h"
+#include "DynamicMeshEditor.h"
+#include "Operations/MeshBoolean.h"
+#include "Materials/Material.h"
+
+#include "EngineGlobals.h"
+#include "GeometryCollection/GeometryCollectionAlgo.h"
+#include "GeometryCollection/GeometryCollectionClusteringUtility.h"
+#include "GeometryCollection/GeometryCollectionConvexUtility.h"
+#include "Voronoi/Voronoi.h"
+#include "PlanarCut.h"
+#include "GeometryCollection/GeometryCollectionProximityUtility.h"
+#include "FractureEngineClustering.h"
+#include "FractureEngineSelection.h"
+#include "GeometryCollection/Facades/CollectionBoundsFacade.h"
+#include "GeometryCollection/Facades/CollectionAnchoringFacade.h"
+#include "GeometryCollection/Facades/CollectionRemoveOnBreakFacade.h"
+#include "GeometryCollection/Facades/CollectionTransformFacade.h"
+#include "DynamicMesh/MeshTransforms.h"
+#include "DynamicMesh/DynamicMesh3.h"
+
+#include UE_INLINE_GENERATED_CPP_BY_NAME(GeometryCollectionMakeNodes)
+
+namespace Dataflow
+{
+	void GeometryCollectionMakeNodes()
+	{
+		static const FLinearColor CDefaultNodeBodyTintColor = FLinearColor(0.f, 0.f, 0.f, 0.5f);
+
+		DATAFLOW_NODE_REGISTER_CREATION_FACTORY(FMakeLiteralStringDataflowNode);
+		DATAFLOW_NODE_REGISTER_CREATION_FACTORY(FMakePointsDataflowNode);
+		DATAFLOW_NODE_REGISTER_CREATION_FACTORY(FMakeBoxDataflowNode);
+		DATAFLOW_NODE_REGISTER_CREATION_FACTORY(FMakeSphereDataflowNode);
+		DATAFLOW_NODE_REGISTER_CREATION_FACTORY(FMakeLiteralFloatDataflowNode);
+		DATAFLOW_NODE_REGISTER_CREATION_FACTORY(FMakeLiteralIntDataflowNode);
+		DATAFLOW_NODE_REGISTER_CREATION_FACTORY(FMakeLiteralBoolDataflowNode);
+		DATAFLOW_NODE_REGISTER_CREATION_FACTORY(FMakeLiteralVectorDataflowNode);
+		DATAFLOW_NODE_REGISTER_CREATION_FACTORY(FMakeTransformDataflowNode);
+		DATAFLOW_NODE_REGISTER_CREATION_FACTORY(FMakeQuaternionDataflowNode);
+
+		// Generators
+		DATAFLOW_NODE_REGISTER_CREATION_FACTORY_NODE_COLORS_BY_CATEGORY("Generators", FLinearColor(.4f, 0.8f, 0.f), CDefaultNodeBodyTintColor);
+	}
+}
+
+void FMakeLiteralStringDataflowNode::Evaluate(Dataflow::FContext& Context, const FDataflowOutput* Out) const
+{
+	if (Out->IsA<FString>(&String))
+	{
+		SetValue<FString>(Context, Value, &String);
+	}
+}
+
+void FMakePointsDataflowNode::Evaluate(Dataflow::FContext& Context, const FDataflowOutput* Out) const
+{
+	if (Out->IsA<TArray<FVector>>(&Points))
+	{
+		SetValue<TArray<FVector>>(Context, Point, &Points);
+	}
+}
+
+void FMakeBoxDataflowNode::Evaluate(Dataflow::FContext& Context, const FDataflowOutput* Out) const
+{
+	if (Out->IsA<FBox>(&Box))
+	{
+		if (DataType == EMakeBoxDataTypeEnum::Dataflow_MakeBox_DataType_MinMax)
+		{
+			FVector MinVal = GetValue<FVector>(Context, &Min);
+			FVector MaxVal = GetValue<FVector>(Context, &Max);
+
+			SetValue<FBox>(Context, FBox(MinVal, MaxVal), &Box);
+		}
+		else if (DataType == EMakeBoxDataTypeEnum::Dataflow_MakeBox_DataType_CenterSize)
+		{
+			FVector CenterVal = GetValue<FVector>(Context, &Center);
+			FVector SizeVal = GetValue<FVector>(Context, &Size);
+
+			SetValue<FBox>(Context, FBox(CenterVal - 0.5 * SizeVal, CenterVal + 0.5 * SizeVal), &Box);
+		}
+	}
+}
+
+
+void FMakeSphereDataflowNode::Evaluate(Dataflow::FContext& Context, const FDataflowOutput* Out) const
+{
+	if (Out->IsA<FSphere>(&Sphere))
+	{
+		FVector CenterVal = GetValue<FVector>(Context, &Center);
+		float RadiusVal = GetValue<float>(Context, &Radius);
+
+		SetValue<FSphere>(Context, FSphere(CenterVal, RadiusVal), &Sphere);
+	}
+}
+
+
+void FMakeLiteralFloatDataflowNode::Evaluate(Dataflow::FContext& Context, const FDataflowOutput* Out) const
+{
+	if (Out->IsA<float>(&Float))
+	{
+		SetValue<float>(Context, Value, &Float);
+	}
+}
+
+void FMakeLiteralIntDataflowNode::Evaluate(Dataflow::FContext& Context, const FDataflowOutput* Out) const
+{
+	if (Out->IsA<int32>(&Int))
+	{
+		SetValue<int32>(Context, Value, &Int);
+	}
+}
+
+void FMakeLiteralBoolDataflowNode::Evaluate(Dataflow::FContext& Context, const FDataflowOutput* Out) const
+{
+	if (Out->IsA<bool>(&Bool))
+	{
+		SetValue<bool>(Context, Value, &Bool);
+	}
+}
+
+void FMakeLiteralVectorDataflowNode::Evaluate(Dataflow::FContext& Context, const FDataflowOutput* Out) const
+{
+	if (Out->IsA<FVector>(&Vector))
+	{
+		FVector Value(GetValue<float>(Context, &X, X), GetValue<float>(Context, &Y, Y), GetValue<float>(Context, &Z, Z));
+		SetValue<FVector>(Context, Value, &Vector);
+	}
+}
+
+void FMakeTransformDataflowNode::Evaluate(Dataflow::FContext& Context, const FDataflowOutput* Out) const
+{
+	if (Out->IsA<FTransform>(&OutTransform))
+	{
+		SetValue<FTransform>(Context,
+			FTransform(FQuat::MakeFromEuler(GetValue<FVector>(Context, &InRotation))
+				, GetValue<FVector>(Context, &InTranslation)
+				, GetValue<FVector>(Context, &InScale))
+			, &OutTransform);
+	}
+}
+
+void FMakeQuaternionDataflowNode::Evaluate(Dataflow::FContext& Context, const FDataflowOutput* Out) const
+{
+	if (Out->IsA<FQuat>(&Quaternion))
+	{
+		FQuat Value(GetValue<float>(Context, &X, X), GetValue<float>(Context, &Y, Y), GetValue<float>(Context, &Z, Z), GetValue<float>(Context, &W, W));
+		SetValue<FQuat>(Context, Value, &Quaternion);
+	}
+}
+
+
diff --git a/Engine/Plugins/Experimental/GeometryCollectionPlugin/Source/GeometryCollectionNodes/Private/Dataflow/GeometryCollectionMathNodes.cpp b/Engine/Plugins/Experimental/GeometryCollectionPlugin/Source/GeometryCollectionNodes/Private/Dataflow/GeometryCollectionMathNodes.cpp
--- a/Engine/Plugins/Experimental/GeometryCollectionPlugin/Source/GeometryCollectionNodes/Private/Dataflow/GeometryCollectionMathNodes.cpp
+++ b/Engine/Plugins/Experimental/GeometryCollectionPlugin/Source/GeometryCollectionNodes/Private/Dataflow/GeometryCollectionMathNodes.cpp
@@ -4,7 +4,7 @@
 #include "Dataflow/DataflowCore.h"
 
 
-//#include UE_INLINE_GENERATED_CPP_BY_NAME(GeometryCollectionMathNodes)
+#include UE_INLINE_GENERATED_CPP_BY_NAME(GeometryCollectionMathNodes)
 
 namespace Dataflow
 {
@@ -13,7 +13,824 @@ namespace Dataflow
 	{
 		static const FLinearColor CDefaultNodeBodyTintColor = FLinearColor(0.f, 0.f, 0.f, 0.5f);
 
+		DATAFLOW_NODE_REGISTER_CREATION_FACTORY(FAddDataflowNode);
+		DATAFLOW_NODE_REGISTER_CREATION_FACTORY(FSubtractDataflowNode);
+		DATAFLOW_NODE_REGISTER_CREATION_FACTORY(FMultiplyDataflowNode);
+		DATAFLOW_NODE_REGISTER_CREATION_FACTORY(FSafeDivideDataflowNode);
+		DATAFLOW_NODE_REGISTER_CREATION_FACTORY(FDivisionDataflowNode);
+		DATAFLOW_NODE_REGISTER_CREATION_FACTORY(FSafeReciprocalDataflowNode);
+		DATAFLOW_NODE_REGISTER_CREATION_FACTORY(FSquareDataflowNode);
+		DATAFLOW_NODE_REGISTER_CREATION_FACTORY(FSquareRootDataflowNode);
+		DATAFLOW_NODE_REGISTER_CREATION_FACTORY(FInverseSqrtDataflowNode);
+		DATAFLOW_NODE_REGISTER_CREATION_FACTORY(FCubeDataflowNode);
+		DATAFLOW_NODE_REGISTER_CREATION_FACTORY(FNegateDataflowNode);
+		DATAFLOW_NODE_REGISTER_CREATION_FACTORY(FAbsDataflowNode);
+		DATAFLOW_NODE_REGISTER_CREATION_FACTORY(FFloorDataflowNode);
+		DATAFLOW_NODE_REGISTER_CREATION_FACTORY(FCeilDataflowNode);
+		DATAFLOW_NODE_REGISTER_CREATION_FACTORY(FRoundDataflowNode);
+		DATAFLOW_NODE_REGISTER_CREATION_FACTORY(FTruncDataflowNode);
+		DATAFLOW_NODE_REGISTER_CREATION_FACTORY(FFracDataflowNode);
+		DATAFLOW_NODE_REGISTER_CREATION_FACTORY(FMinDataflowNode);
+		DATAFLOW_NODE_REGISTER_CREATION_FACTORY(FMaxDataflowNode);
+		DATAFLOW_NODE_REGISTER_CREATION_FACTORY(FMin3DataflowNode);
+		DATAFLOW_NODE_REGISTER_CREATION_FACTORY(FMax3DataflowNode);
+		DATAFLOW_NODE_REGISTER_CREATION_FACTORY(FSignDataflowNode);
+		DATAFLOW_NODE_REGISTER_CREATION_FACTORY(FClampDataflowNode);
+		DATAFLOW_NODE_REGISTER_CREATION_FACTORY(FFitDataflowNode);
+		DATAFLOW_NODE_REGISTER_CREATION_FACTORY(FEFitDataflowNode);
+		DATAFLOW_NODE_REGISTER_CREATION_FACTORY(FPowDataflowNode);
+		DATAFLOW_NODE_REGISTER_CREATION_FACTORY(FLogDataflowNode);
+		DATAFLOW_NODE_REGISTER_CREATION_FACTORY(FLogeDataflowNode);
+		DATAFLOW_NODE_REGISTER_CREATION_FACTORY(FLerpDataflowNode);
+		DATAFLOW_NODE_REGISTER_CREATION_FACTORY(FWrapDataflowNode);
+		DATAFLOW_NODE_REGISTER_CREATION_FACTORY(FExpDataflowNode);
+		DATAFLOW_NODE_REGISTER_CREATION_FACTORY(FSinDataflowNode);
+		DATAFLOW_NODE_REGISTER_CREATION_FACTORY(FArcSinDataflowNode);
+		DATAFLOW_NODE_REGISTER_CREATION_FACTORY(FCosDataflowNode);
+		DATAFLOW_NODE_REGISTER_CREATION_FACTORY(FArcCosDataflowNode);
+		DATAFLOW_NODE_REGISTER_CREATION_FACTORY(FTanDataflowNode);
+		DATAFLOW_NODE_REGISTER_CREATION_FACTORY(FArcTanDataflowNode);
+		DATAFLOW_NODE_REGISTER_CREATION_FACTORY(FArcTan2DataflowNode);
+		DATAFLOW_NODE_REGISTER_CREATION_FACTORY(FNormalizeToRangeDataflowNode);
+		DATAFLOW_NODE_REGISTER_CREATION_FACTORY(FDotProductDataflowNode);
+		DATAFLOW_NODE_REGISTER_CREATION_FACTORY(FCrossProductDataflowNode);
+		DATAFLOW_NODE_REGISTER_CREATION_FACTORY(FNormalizeDataflowNode);
+		DATAFLOW_NODE_REGISTER_CREATION_FACTORY(FLengthDataflowNode);
+		DATAFLOW_NODE_REGISTER_CREATION_FACTORY(FDistanceDataflowNode);
+		DATAFLOW_NODE_REGISTER_CREATION_FACTORY(FIsNearlyZeroDataflowNode);
+		DATAFLOW_NODE_REGISTER_CREATION_FACTORY(FRandomFloatDataflowNode);
+		DATAFLOW_NODE_REGISTER_CREATION_FACTORY(FRandomFloatInRangeDataflowNode);
+		DATAFLOW_NODE_REGISTER_CREATION_FACTORY(FRandomUnitVectorDataflowNode);
+		DATAFLOW_NODE_REGISTER_CREATION_FACTORY(FRandomUnitVectorInConeDataflowNode);
+		DATAFLOW_NODE_REGISTER_CREATION_FACTORY(FRadiansToDegreesDataflowNode);
+		DATAFLOW_NODE_REGISTER_CREATION_FACTORY(FDegreesToRadiansDataflowNode);
+		DATAFLOW_NODE_REGISTER_CREATION_FACTORY(FMathConstantsDataflowNode);
+
+		// Math
+		DATAFLOW_NODE_REGISTER_CREATION_FACTORY_NODE_COLORS_BY_CATEGORY("Math", FLinearColor(0.f, 0.4f, 0.8f), CDefaultNodeBodyTintColor);
+	}
+}
+
+void FAddDataflowNode::Evaluate(Dataflow::FContext& Context, const FDataflowOutput* Out) const
+{
+	if (Out->IsA<float>(&ReturnValue))
+	{
+		const float InFloatA = GetValue<float>(Context, &FloatA, FloatA);
+		const float InFloatB = GetValue<float>(Context, &FloatB, FloatB);
+
+		float Result = InFloatA + InFloatB;
+		SetValue<float>(Context, Result, &ReturnValue);
+	}
+}
+
+void FSubtractDataflowNode::Evaluate(Dataflow::FContext& Context, const FDataflowOutput* Out) const
+{
+	if (Out->IsA<float>(&ReturnValue))
+	{
+		const float InFloatA = GetValue<float>(Context, &FloatA, FloatA);
+		const float InFloatB = GetValue<float>(Context, &FloatB, FloatB);
+
+		float Result = InFloatA - InFloatB;
+		SetValue<float>(Context, Result, &ReturnValue);
+	}
+}
+
+void FMultiplyDataflowNode::Evaluate(Dataflow::FContext& Context, const FDataflowOutput* Out) const
+{
+	if (Out->IsA<float>(&ReturnValue))
+	{
+		const float InFloatA = GetValue<float>(Context, &FloatA, FloatA);
+		const float InFloatB = GetValue<float>(Context, &FloatB, FloatB);
+
+		float Result = InFloatA * InFloatB;
+		SetValue<float>(Context, Result, &ReturnValue);
+	}
+}
+
+void FSafeDivideDataflowNode::Evaluate(Dataflow::FContext& Context, const FDataflowOutput* Out) const
+{
+	if (Out->IsA<float>(&ReturnValue))
+	{
+		const float InFloatA = GetValue<float>(Context, &FloatA, FloatA);
+		const float InFloatB = GetValue<float>(Context, &FloatB, FloatB);
+
+		float Result = 0.f;
+		if (InFloatB != 0.f)
+		{
+			Result = InFloatA / InFloatB;
+		}
+		SetValue<float>(Context, Result, &ReturnValue);
+	}
+}
+
+void FDivisionDataflowNode::Evaluate(Dataflow::FContext& Context, const FDataflowOutput* Out) const
+{
+	if (Out->IsA<float>(&Remainder) ||
+		Out->IsA<int32>(&ReturnValue))
+	{
+		const float InDividend = GetValue<float>(Context, &Dividend, Dividend);
+		const float InDivisor = GetValue<float>(Context, &Divisor, Divisor);
+
+		float ResultRemainder = 0.f;
+		int32 Result = 0;
+
+		if (InDivisor != 0.f)
+		{
+			Result = (int32)(InDividend / InDivisor);
+			ResultRemainder = InDividend - (float)Result * InDivisor;
+		}
+		SetValue<float>(Context, ResultRemainder, &Remainder);
+		SetValue<int32>(Context, Result, &ReturnValue);
+	}
+}
+
+void FSafeReciprocalDataflowNode::Evaluate(Dataflow::FContext& Context, const FDataflowOutput* Out) const
+{
+	if (Out->IsA<float>(&ReturnValue))
+	{
+		const float InFloat = GetValue<float>(Context, &Float, Float);
+
+		float Result = 0.f;
+		if (InFloat != 0.f)
+		{
+			Result = 1.f / InFloat;
+		}
+		SetValue<float>(Context, Result, &ReturnValue);
+	}
+}
+
+void FSquareDataflowNode::Evaluate(Dataflow::FContext& Context, const FDataflowOutput* Out) const
+{
+	if (Out->IsA<float>(&ReturnValue))
+	{
+		const float InFloat = GetValue<float>(Context, &Float, Float);
+
+		float Result = InFloat * InFloat;
+		SetValue<float>(Context, Result, &ReturnValue);
+	}
+}
+
+void FSquareRootDataflowNode::Evaluate(Dataflow::FContext& Context, const FDataflowOutput* Out) const
+{
+	if (Out->IsA<float>(&ReturnValue))
+	{
+		const float InFloat = GetValue<float>(Context, &Float, Float);
+
+		float Result = 0.f;
+		if (InFloat > 0.f)
+		{
+			Result = FMath::Sqrt(InFloat);
+		}
+		SetValue<float>(Context, Result, &ReturnValue);
+	}
+}
+
+void FInverseSqrtDataflowNode::Evaluate(Dataflow::FContext& Context, const FDataflowOutput* Out) const
+{
+	if (Out->IsA<float>(&ReturnValue))
+	{
+		const float InFloat = GetValue<float>(Context, &Float, Float);
+
+		float Result = 0.f;
+		if (InFloat > 0.f)
+		{
+			Result = FMath::InvSqrt(InFloat);
+		}
+		SetValue<float>(Context, Result, &ReturnValue);
+	}
+}
+
+void FCubeDataflowNode::Evaluate(Dataflow::FContext& Context, const FDataflowOutput* Out) const
+{
+	if (Out->IsA<float>(&ReturnValue))
+	{
+		const float InFloat = GetValue<float>(Context, &Float, Float);
+
+		float Result = InFloat * InFloat * InFloat;
+		SetValue<float>(Context, Result, &ReturnValue);
+	}
+}
+
+void FNegateDataflowNode::Evaluate(Dataflow::FContext& Context, const FDataflowOutput* Out) const
+{
+	if (Out->IsA<float>(&ReturnValue))
+	{
+		const float InFloat = GetValue<float>(Context, &Float, Float);
+
+		float Result = -1.f * InFloat;
+		SetValue<float>(Context, Result, &ReturnValue);
+	}
+}
+
+void FAbsDataflowNode::Evaluate(Dataflow::FContext& Context, const FDataflowOutput* Out) const
+{
+	if (Out->IsA<float>(&ReturnValue))
+	{
+		const float InFloat = GetValue<float>(Context, &Float, Float);
+
+		float Result = FMath::Abs(InFloat);
+		SetValue<float>(Context, Result, &ReturnValue);
+	}
+}
+
+void FFloorDataflowNode::Evaluate(Dataflow::FContext& Context, const FDataflowOutput* Out) const
+{
+	if (Out->IsA<float>(&ReturnValue))
+	{
+		const float InFloat = GetValue<float>(Context, &Float, Float);
+
+		float Result = FMath::Floor(InFloat);
+		SetValue<float>(Context, Result, &ReturnValue);
 	}
 }
 
+void FCeilDataflowNode::Evaluate(Dataflow::FContext& Context, const FDataflowOutput* Out) const
+{
+	if (Out->IsA<float>(&ReturnValue))
+	{
+		const float InFloat = GetValue<float>(Context, &Float, Float);
+
+		float Result = FMath::CeilToFloat(InFloat);
+		SetValue<float>(Context, Result, &ReturnValue);
+	}
+}
+
+void FRoundDataflowNode::Evaluate(Dataflow::FContext& Context, const FDataflowOutput* Out) const
+{
+	if (Out->IsA<float>(&ReturnValue))
+	{
+		const float InFloat = GetValue<float>(Context, &Float, Float);
+
+		float Result = FMath::RoundToFloat(InFloat);
+		SetValue<float>(Context, Result, &ReturnValue);
+	}
+}
+
+void FTruncDataflowNode::Evaluate(Dataflow::FContext& Context, const FDataflowOutput* Out) const
+{
+	if (Out->IsA<float>(&ReturnValue))
+	{
+		const float InFloat = GetValue<float>(Context, &Float, Float);
+
+		float Result = FMath::TruncToFloat(InFloat);
+		SetValue<float>(Context, Result, &ReturnValue);
+	}
+}
+
+void FFracDataflowNode::Evaluate(Dataflow::FContext& Context, const FDataflowOutput* Out) const
+{
+	if (Out->IsA<float>(&ReturnValue))
+	{
+		const float InFloat = GetValue<float>(Context, &Float, Float);
+
+		float Result = FMath::Frac(InFloat);
+		SetValue<float>(Context, Result, &ReturnValue);
+	}
+}
+
+void FMinDataflowNode::Evaluate(Dataflow::FContext& Context, const FDataflowOutput* Out) const
+{
+	if (Out->IsA<float>(&ReturnValue))
+	{
+		const float InFloatA = GetValue<float>(Context, &FloatA, FloatA);
+		const float InFloatB = GetValue<float>(Context, &FloatB, FloatB);
+
+		float Result = FMath::Min(InFloatA, InFloatB);
+		SetValue<float>(Context, Result, &ReturnValue);
+	}
+}
+
+void FMaxDataflowNode::Evaluate(Dataflow::FContext& Context, const FDataflowOutput* Out) const
+{
+	if (Out->IsA<float>(&ReturnValue))
+	{
+		const float InFloatA = GetValue<float>(Context, &FloatA, FloatA);
+		const float InFloatB = GetValue<float>(Context, &FloatB, FloatB);
+
+		float Result = FMath::Max(InFloatA, InFloatB);
+		SetValue<float>(Context, Result, &ReturnValue);
+	}
+}
+
+void FMin3DataflowNode::Evaluate(Dataflow::FContext& Context, const FDataflowOutput* Out) const
+{
+	if (Out->IsA<float>(&ReturnValue))
+	{
+		const float InFloatA = GetValue<float>(Context, &FloatA, FloatA);
+		const float InFloatB = GetValue<float>(Context, &FloatB, FloatB);
+		const float InFloatC = GetValue<float>(Context, &FloatC, FloatB);
+
+		float Result = FMath::Min3(InFloatA, InFloatB, InFloatC);
+		SetValue<float>(Context, Result, &ReturnValue);
+	}
+}
+
+void FMax3DataflowNode::Evaluate(Dataflow::FContext& Context, const FDataflowOutput* Out) const
+{
+	if (Out->IsA<float>(&ReturnValue))
+	{
+		const float InFloatA = GetValue<float>(Context, &FloatA, FloatA);
+		const float InFloatB = GetValue<float>(Context, &FloatB, FloatB);
+		const float InFloatC = GetValue<float>(Context, &FloatC, FloatB);
+
+		float Result = FMath::Max3(InFloatA, InFloatB, InFloatC);
+		SetValue<float>(Context, Result, &ReturnValue);
+	}
+}
+
+void FSignDataflowNode::Evaluate(Dataflow::FContext& Context, const FDataflowOutput* Out) const
+{
+	if (Out->IsA<float>(&ReturnValue))
+	{
+		const float InFloat = GetValue<float>(Context, &Float, Float);
+
+		float Result = FMath::Sign(InFloat);
+		SetValue<float>(Context, Result, &ReturnValue);
+	}
+}
+
+void FClampDataflowNode::Evaluate(Dataflow::FContext& Context, const FDataflowOutput* Out) const
+{
+	if (Out->IsA<float>(&ReturnValue))
+	{
+		const float InFloat = GetValue<float>(Context, &Float, Float);
+		const float InMin = GetValue<float>(Context, &Min, Min);
+		const float InMax = GetValue<float>(Context, &Max, Max);
+
+		float Result = FMath::Clamp(InFloat, InMin, InMax);
+		SetValue<float>(Context, Result, &ReturnValue);
+	}
+}
+
+void FFitDataflowNode::Evaluate(Dataflow::FContext& Context, const FDataflowOutput* Out) const
+{
+	if (Out->IsA<float>(&ReturnValue))
+	{
+		float InFloat = GetValue<float>(Context, &Float, Float);
+		const float InOldMin = GetValue<float>(Context, &OldMin, OldMin);
+		const float InOldMax = GetValue<float>(Context, &OldMax, OldMax);
+		const float InNewMin = GetValue<float>(Context, &NewMin, NewMin);
+		const float InNewMax = GetValue<float>(Context, &NewMax, NewMax);
+
+		float Result = InFloat;
+		if (InOldMax > InOldMin && InNewMax > InNewMin)
+		{
+			InFloat = FMath::Clamp(InFloat, InOldMin, InOldMax);
+			float Q = (InFloat - InOldMin) / (InOldMax - InOldMin);
+			Result = InNewMin + Q * (InNewMax - InNewMin);
+		}
+		SetValue<float>(Context, Result, &ReturnValue);
+	}
+}
+
+void FEFitDataflowNode::Evaluate(Dataflow::FContext& Context, const FDataflowOutput* Out) const
+{
+	if (Out->IsA<float>(&ReturnValue))
+	{
+		const float InFloat = GetValue<float>(Context, &Float, Float);
+		const float InOldMin = GetValue<float>(Context, &OldMin, OldMin);
+		const float InOldMax = GetValue<float>(Context, &OldMax, OldMax);
+		const float InNewMin = GetValue<float>(Context, &NewMin, NewMin);
+		const float InNewMax = GetValue<float>(Context, &NewMax, NewMax);
+
+		float Result = InFloat;
+		if (InOldMax > InOldMin && InNewMax > InNewMin)
+		{
+			float Q = (InFloat - InOldMin) / (InOldMax - InOldMin);
+			Result = InNewMin + Q * (InNewMax - InNewMin);
+		}
+		SetValue<float>(Context, Result, &ReturnValue);
+	}
+}
+
+void FPowDataflowNode::Evaluate(Dataflow::FContext& Context, const FDataflowOutput* Out) const
+{
+	if (Out->IsA<float>(&ReturnValue))
+	{
+		const float InBase = GetValue<float>(Context, &Base, Base);
+		const float InExp = GetValue<float>(Context, &Exp, Exp);
+
+		float Result = FMath::Pow(InBase, InExp);
+		SetValue<float>(Context, Result, &ReturnValue);
+	}
+}
+
+void FLogDataflowNode::Evaluate(Dataflow::FContext& Context, const FDataflowOutput* Out) const
+{
+	if (Out->IsA<float>(&ReturnValue))
+	{
+		const float InBase = GetValue<float>(Context, &Base, Base);
+		const float InA = GetValue<float>(Context, &A, A);
+
+		float Result = 0;
+		if (InBase > 0.f)
+		{
+			Result = FMath::LogX(InBase, A);
+		}
+		SetValue<float>(Context, Result, &ReturnValue);
+	}
+}
+
+void FLogeDataflowNode::Evaluate(Dataflow::FContext& Context, const FDataflowOutput* Out) const
+{
+	if (Out->IsA<float>(&ReturnValue))
+	{
+		const float InA = GetValue<float>(Context, &A, A);
+
+		float Result = FMath::Loge(A);
+		SetValue<float>(Context, Result, &ReturnValue);
+	}
+}
+
+void FLerpDataflowNode::Evaluate(Dataflow::FContext& Context, const FDataflowOutput* Out) const
+{
+	if (Out->IsA<float>(&ReturnValue))
+	{
+		const float InA = GetValue<float>(Context, &A, A);
+		const float InB = GetValue<float>(Context, &B, B);
+		const float InAlpha = GetValue<float>(Context, &Alpha, Alpha);
+
+		float Result = InA;
+		if (InB > InA)
+		{
+			Result = FMath::Lerp(InA, InB, InAlpha);
+		}
+		SetValue<float>(Context, Result, &ReturnValue);
+	}
+}
+
+void FWrapDataflowNode::Evaluate(Dataflow::FContext& Context, const FDataflowOutput* Out) const
+{
+	if (Out->IsA<float>(&ReturnValue))
+	{
+		const float InFloat = GetValue<float>(Context, &Float, Float);
+		const float InMin = GetValue<float>(Context, &Min, Min);
+		const float InMax = GetValue<float>(Context, &Max, Max);
+
+		float Result = InFloat;
+		if (InMax > InMin)
+		{
+			Result = FMath::Wrap(InFloat, InMin, InMax);
+		}
+		SetValue<float>(Context, Result, &ReturnValue);
+	}
+}
+
+void FExpDataflowNode::Evaluate(Dataflow::FContext& Context, const FDataflowOutput* Out) const
+{
+	if (Out->IsA<float>(&ReturnValue))
+	{
+		const float InFloat = GetValue<float>(Context, &Float, Float);
+
+		float Result = FMath::Exp(InFloat);
+		SetValue<float>(Context, Result, &ReturnValue);
+	}
+}
+
+void FSinDataflowNode::Evaluate(Dataflow::FContext& Context, const FDataflowOutput* Out) const
+{
+	if (Out->IsA<float>(&ReturnValue))
+	{
+		const float InFloat = GetValue<float>(Context, &Float, Float);
+
+		float Result = FMath::Sin(InFloat);
+		SetValue<float>(Context, Result, &ReturnValue);
+	}
+}
+
+void FArcSinDataflowNode::Evaluate(Dataflow::FContext& Context, const FDataflowOutput* Out) const
+{
+	if (Out->IsA<float>(&ReturnValue))
+	{
+		const float InFloat = GetValue<float>(Context, &Float, Float);
+
+		float Result = 0.f;
+		if (InFloat >= -1.f && InFloat <= 1.f)
+		{
+			Result = FMath::Asin(InFloat);
+		}
+		SetValue<float>(Context, Result, &ReturnValue);
+	}
+}
+
+void FCosDataflowNode::Evaluate(Dataflow::FContext& Context, const FDataflowOutput* Out) const
+{
+	if (Out->IsA<float>(&ReturnValue))
+	{
+		const float InFloat = GetValue<float>(Context, &Float, Float);
+
+		float Result = FMath::Cos(InFloat);
+		SetValue<float>(Context, Result, &ReturnValue);
+	}
+}
+
+void FArcCosDataflowNode::Evaluate(Dataflow::FContext& Context, const FDataflowOutput* Out) const
+{
+	if (Out->IsA<float>(&ReturnValue))
+	{
+		const float InFloat = GetValue<float>(Context, &Float, Float);
+
+		float Result = 0.f;
+		if (InFloat >= -1.f && InFloat <= 1.f)
+		{
+			Result = FMath::Acos(InFloat);
+		}
+		SetValue<float>(Context, Result, &ReturnValue);
+	}
+}
+
+void FTanDataflowNode::Evaluate(Dataflow::FContext& Context, const FDataflowOutput* Out) const
+{
+	if (Out->IsA<float>(&ReturnValue))
+	{
+		const float InFloat = GetValue<float>(Context, &Float, Float);
+
+		float Result = FMath::Tan(InFloat);
+		SetValue<float>(Context, Result, &ReturnValue);
+	}
+}
+
+void FArcTanDataflowNode::Evaluate(Dataflow::FContext& Context, const FDataflowOutput* Out) const
+{
+	if (Out->IsA<float>(&ReturnValue))
+	{
+		const float InFloat = GetValue<float>(Context, &Float, Float);
+
+		float Result = FMath::Atan(InFloat);
+		SetValue<float>(Context, Result, &ReturnValue);
+	}
+}
+
+void FArcTan2DataflowNode::Evaluate(Dataflow::FContext& Context, const FDataflowOutput* Out) const
+{
+	if (Out->IsA<float>(&ReturnValue))
+	{
+		const float InY = GetValue<float>(Context, &Y, Y);
+		const float InX = GetValue<float>(Context, &X, X);
+
+		float Result = FMath::Atan2(InY, InX);
+		SetValue<float>(Context, Result, &ReturnValue);
+	}
+}
+
+void FNormalizeToRangeDataflowNode::Evaluate(Dataflow::FContext& Context, const FDataflowOutput* Out) const
+{
+	if (Out->IsA<float>(&ReturnValue))
+	{
+		const float InFloat = GetValue<float>(Context, &Float, Float);
+		float IRangeMin = GetValue<float>(Context, &RangeMin, RangeMin);
+		float InRangeMax = GetValue<float>(Context, &RangeMax, RangeMax);
+
+		float Result;
+		if (IRangeMin == InRangeMax)
+		{
+			if (InFloat < IRangeMin)
+			{
+				Result = 0.f;
+			}
+			else
+			{
+				Result = 1.f;
+			}
+		}
+
+		if (IRangeMin > InRangeMax)
+		{
+			Swap(InRangeMax, InRangeMax);
+		}
+		Result = (InFloat - IRangeMin) / (InRangeMax - IRangeMin);
+
+		SetValue<float>(Context, Result, &ReturnValue);
+	}
+}
+
+void FDotProductDataflowNode::Evaluate(Dataflow::FContext& Context, const FDataflowOutput* Out) const
+{
+	if (Out->IsA<float>(&ReturnValue))
+	{
+		const FVector InVectorA = GetValue<FVector>(Context, &VectorA, VectorA);
+		const FVector InVectorB = GetValue<FVector>(Context, &VectorB, VectorB);
+
+		float Result = FVector::DotProduct(InVectorA, InVectorB);
+		SetValue<float>(Context, Result, &ReturnValue);
+	}
+}
+
+void FCrossProductDataflowNode::Evaluate(Dataflow::FContext& Context, const FDataflowOutput* Out) const
+{
+	if (Out->IsA<FVector>(&ReturnValue))
+	{
+		const FVector InVectorA = GetValue<FVector>(Context, &VectorA, VectorA);
+		const FVector InVectorB = GetValue<FVector>(Context, &VectorB, VectorB);
+
+		FVector Result = FVector::CrossProduct(InVectorA, InVectorB);
+		SetValue<FVector>(Context, Result, &ReturnValue);
+	}
+}
+
+void FNormalizeDataflowNode::Evaluate(Dataflow::FContext& Context, const FDataflowOutput* Out) const
+{
+	if (Out->IsA<FVector>(&ReturnValue))
+	{
+		const FVector InVectorA = GetValue<FVector>(Context, &VectorA, VectorA);
+		const float InTolerance = GetValue<float>(Context, &Tolerance, Tolerance);
+
+		FVector Result = InVectorA.GetSafeNormal(Tolerance);
+		SetValue<FVector>(Context, Result, &ReturnValue);
+	}
+}
+
+void FLengthDataflowNode::Evaluate(Dataflow::FContext& Context, const FDataflowOutput* Out) const
+{
+	if (Out->IsA<float>(&ReturnValue))
+	{
+		const FVector InVector = GetValue<FVector>(Context, &Vector, Vector);
+
+		float Result = InVector.Length();
+		SetValue<float>(Context, Result, &ReturnValue);
+	}
+}
+
+void FDistanceDataflowNode::Evaluate(Dataflow::FContext& Context, const FDataflowOutput* Out) const
+{
+	if (Out->IsA<float>(&ReturnValue))
+	{
+		const FVector InPointA = GetValue<FVector>(Context, &PointA, PointA);
+		const FVector InPointB = GetValue<FVector>(Context, &PointB, PointB);
+
+		float Result = (InPointB - InPointA).Length();
+		SetValue<float>(Context, Result, &ReturnValue);
+	}
+}
+
+void FIsNearlyZeroDataflowNode::Evaluate(Dataflow::FContext& Context, const FDataflowOutput* Out) const
+{
+	if (Out->IsA<bool>(&ReturnValue))
+	{
+		const float InFloat = GetValue<float>(Context, &Float, Float);
+
+		bool Result = FMath::IsNearlyZero(InFloat);
+		SetValue<bool>(Context, Result, &ReturnValue);
+	}
+}
+
+void FRandomFloatDataflowNode::Evaluate(Dataflow::FContext& Context, const FDataflowOutput* Out) const
+{
+	if (Out->IsA<float>(&ReturnValue))
+	{
+		if (bDeterministic)
+		{
+			float RandomSeedVal = GetValue<float>(Context, &RandomSeed);
+
+			FRandomStream Stream(RandomSeedVal);
+			SetValue<float>(Context, Stream.FRand(), &ReturnValue);
+		}
+		else
+		{
+			SetValue<float>(Context, FMath::FRand(), &ReturnValue);
+		}
+	}
+}
+
+void FRandomFloatInRangeDataflowNode::Evaluate(Dataflow::FContext& Context, const FDataflowOutput* Out) const
+{
+	if (Out->IsA<float>(&ReturnValue))
+	{
+		float MinVal = GetValue<float>(Context, &Min);
+		float MaxVal = GetValue<float>(Context, &Max);
+
+		if (bDeterministic)
+		{
+			float RandomSeedVal = GetValue<float>(Context, &RandomSeed);
+
+			FRandomStream Stream(RandomSeedVal);
+			SetValue<float>(Context, Stream.FRandRange(MinVal, MaxVal), &ReturnValue);
+		}
+		else
+		{
+			SetValue<float>(Context, FMath::FRandRange(MinVal, MaxVal), &ReturnValue);
+		}
+	}
+}
+
+void FRandomUnitVectorDataflowNode::Evaluate(Dataflow::FContext& Context, const FDataflowOutput* Out) const
+{
+	if (Out->IsA<FVector>(&ReturnValue))
+	{
+		if (bDeterministic)
+		{
+			float RandomSeedVal = GetValue<float>(Context, &RandomSeed);
+
+			FRandomStream Stream(RandomSeedVal);
+			SetValue<FVector>(Context, Stream.VRand(), &ReturnValue);
+		}
+		else
+		{
+			SetValue<FVector>(Context, FMath::VRand(), &ReturnValue);
+		}
+	}
+}
+
+void FRandomUnitVectorInConeDataflowNode::Evaluate(Dataflow::FContext& Context, const FDataflowOutput* Out) const
+{
+	if (Out->IsA<FVector>(&ReturnValue))
+	{
+		FVector ConeDirectionVal = GetValue<FVector>(Context, &ConeDirection);
+		float ConeHalfAngleVal = GetValue<float>(Context, &ConeHalfAngle);
+
+		if (bDeterministic)
+		{
+			float RandomSeedVal = GetValue<float>(Context, &RandomSeed);
+
+			FRandomStream Stream(RandomSeedVal);
+			SetValue<FVector>(Context, Stream.VRandCone(ConeDirectionVal, ConeHalfAngleVal), &ReturnValue);
+		}
+		else
+		{
+			SetValue<FVector>(Context, FMath::VRandCone(ConeDirectionVal, ConeHalfAngleVal), &ReturnValue);
+		}
+	}
+}
+
+void FRadiansToDegreesDataflowNode::Evaluate(Dataflow::FContext& Context, const FDataflowOutput* Out) const
+{
+	if (Out->IsA<float>(&Degrees))
+	{
+		SetValue<float>(Context, FMath::RadiansToDegrees(GetValue<float>(Context, &Radians)), &Degrees);
+	}
+}
+
+void FDegreesToRadiansDataflowNode::Evaluate(Dataflow::FContext& Context, const FDataflowOutput* Out) const
+{
+	if (Out->IsA<float>(&Radians))
+	{
+		SetValue<float>(Context, FMath::DegreesToRadians(GetValue<float>(Context, &Degrees)), &Radians);
+	}
+}
+
+void FMathConstantsDataflowNode::Evaluate(Dataflow::FContext& Context, const FDataflowOutput* Out) const
+{
+	if (Out->IsA<float>(&ReturnValue))
+	{
+		if (Constant == EMathConstantsEnum::Dataflow_MathConstants_Pi)
+		{
+			SetValue<float>(Context, FMathf::Pi, &ReturnValue);
+		}
+		else if (Constant == EMathConstantsEnum::Dataflow_MathConstants_HalfPi)
+		{
+			SetValue<float>(Context, FMathf::HalfPi, &ReturnValue);
+		}
+		else if (Constant == EMathConstantsEnum::Dataflow_MathConstants_TwoPi)
+		{
+			SetValue<float>(Context, FMathf::TwoPi, &ReturnValue);
+		}
+		else if (Constant == EMathConstantsEnum::Dataflow_MathConstants_FourPi)
+		{
+			SetValue<float>(Context, FMathf::FourPi, &ReturnValue);
+		}
+		else if (Constant == EMathConstantsEnum::Dataflow_MathConstants_InvPi)
+		{
+			SetValue<float>(Context, FMathf::InvPi, &ReturnValue);
+		}
+		else if (Constant == EMathConstantsEnum::Dataflow_MathConstants_InvTwoPi)
+		{
+			SetValue<float>(Context, FMathf::InvTwoPi, &ReturnValue);
+		}
+		else if (Constant == EMathConstantsEnum::Dataflow_MathConstants_Sqrt2)
+		{
+			SetValue<float>(Context, FMathf::Sqrt2, &ReturnValue);
+		}
+		else if (Constant == EMathConstantsEnum::Dataflow_MathConstants_InvSqrt2)
+		{
+			SetValue<float>(Context, FMathf::InvSqrt2, &ReturnValue);
+		}
+		else if (Constant == EMathConstantsEnum::Dataflow_MathConstants_Sqrt3)
+		{
+			SetValue<float>(Context, FMathf::Sqrt3, &ReturnValue);
+		}
+		else if (Constant == EMathConstantsEnum::Dataflow_MathConstants_InvSqrt3)
+		{
+			SetValue<float>(Context, FMathf::InvSqrt3, &ReturnValue);
+		}
+		else if (Constant == EMathConstantsEnum::Dataflow_FloatToInt_Function_E)
+		{
+			SetValue<float>(Context, 2.71828182845904523536f, &ReturnValue);
+		}
+		else if (Constant == EMathConstantsEnum::Dataflow_FloatToInt_Function_Gamma)
+		{
+			SetValue<float>(Context, 0.577215664901532860606512090082f, &ReturnValue);
+		}
+		else if (Constant == EMathConstantsEnum::Dataflow_FloatToInt_Function_GoldenRatio)
+		{
+			SetValue<float>(Context, 1.618033988749894f, &ReturnValue);
+		}
+		else if (Constant == EMathConstantsEnum::Dataflow_FloatToInt_Function_ZeroTolerance)
+		{
+			SetValue<float>(Context, FMathf::ZeroTolerance, &ReturnValue);
+		}
+	}
+}
+
+
+
 
diff --git a/Engine/Plugins/Experimental/GeometryCollectionPlugin/Source/GeometryCollectionNodes/Private/Dataflow/GeometryCollectionNodes.cpp b/Engine/Plugins/Experimental/GeometryCollectionPlugin/Source/GeometryCollectionNodes/Private/Dataflow/GeometryCollectionNodes.cpp
--- a/Engine/Plugins/Experimental/GeometryCollectionPlugin/Source/GeometryCollectionNodes/Private/Dataflow/GeometryCollectionNodes.cpp
+++ b/Engine/Plugins/Experimental/GeometryCollectionPlugin/Source/GeometryCollectionNodes/Private/Dataflow/GeometryCollectionNodes.cpp
@@ -50,33 +50,12 @@ namespace Dataflow
 		DATAFLOW_NODE_REGISTER_CREATION_FACTORY(FAppendCollectionAssetsDataflowNode);
 		DATAFLOW_NODE_REGISTER_CREATION_FACTORY(FPrintStringDataflowNode);
 		DATAFLOW_NODE_REGISTER_CREATION_FACTORY(FLogStringDataflowNode);
-		DATAFLOW_NODE_REGISTER_CREATION_FACTORY(FMakeLiteralStringDataflowNode);
 		DATAFLOW_NODE_REGISTER_CREATION_FACTORY(FBoundingBoxDataflowNode);
 		DATAFLOW_NODE_REGISTER_CREATION_FACTORY(FExpandBoundingBoxDataflowNode);
-		DATAFLOW_NODE_REGISTER_CREATION_FACTORY(FVectorToStringDataflowNode);
-		DATAFLOW_NODE_REGISTER_CREATION_FACTORY(FFloatToStringDataflowNode);
-		DATAFLOW_NODE_REGISTER_CREATION_FACTORY(FMakePointsDataflowNode);
-		DATAFLOW_NODE_REGISTER_CREATION_FACTORY(FMakeBoxDataflowNode);
-		DATAFLOW_NODE_REGISTER_CREATION_FACTORY(FMakeSphereDataflowNode);
-		DATAFLOW_NODE_REGISTER_CREATION_FACTORY(FMakeLiteralFloatDataflowNode);
-		DATAFLOW_NODE_REGISTER_CREATION_FACTORY(FMakeLiteralIntDataflowNode);
-		DATAFLOW_NODE_REGISTER_CREATION_FACTORY(FMakeLiteralBoolDataflowNode);
-		DATAFLOW_NODE_REGISTER_CREATION_FACTORY(FMakeLiteralVectorDataflowNode);
-		DATAFLOW_NODE_REGISTER_CREATION_FACTORY(FIntToStringDataflowNode);
-		DATAFLOW_NODE_REGISTER_CREATION_FACTORY(FBoolToStringDataflowNode);
 		DATAFLOW_NODE_REGISTER_CREATION_FACTORY(FExpandVectorDataflowNode);
-		DATAFLOW_NODE_REGISTER_CREATION_FACTORY(FIntToFloatDataflowNode);
 		DATAFLOW_NODE_REGISTER_CREATION_FACTORY(FStringAppendDataflowNode);
-		DATAFLOW_NODE_REGISTER_CREATION_FACTORY(FRandomFloatDataflowNode);
-		DATAFLOW_NODE_REGISTER_CREATION_FACTORY(FRandomFloatInRangeDataflowNode);
-		DATAFLOW_NODE_REGISTER_CREATION_FACTORY(FRandomUnitVectorDataflowNode);
-		DATAFLOW_NODE_REGISTER_CREATION_FACTORY(FRandomUnitVectorInConeDataflowNode);
-		DATAFLOW_NODE_REGISTER_CREATION_FACTORY(FRadiansToDegreesDataflowNode);
-		DATAFLOW_NODE_REGISTER_CREATION_FACTORY(FDegreesToRadiansDataflowNode);
 		DATAFLOW_NODE_REGISTER_CREATION_FACTORY(FHashStringDataflowNode);
 		DATAFLOW_NODE_REGISTER_CREATION_FACTORY(FHashVectorDataflowNode);
-		DATAFLOW_NODE_REGISTER_CREATION_FACTORY(FFloatToIntDataflowNode);
-		DATAFLOW_NODE_REGISTER_CREATION_FACTORY(FMathConstantsDataflowNode);
 		DATAFLOW_NODE_REGISTER_CREATION_FACTORY(FGetArrayElementDataflowNode);
 		DATAFLOW_NODE_REGISTER_CREATION_FACTORY(FGetNumArrayElementsDataflowNode);
 		DATAFLOW_NODE_REGISTER_CREATION_FACTORY(FGetBoundingBoxesFromCollectionDataflowNode);
@@ -98,11 +77,9 @@ namespace Dataflow
 		DATAFLOW_NODE_REGISTER_CREATION_FACTORY(FBoolArrayToFaceSelectionDataflowNode);
 		DATAFLOW_NODE_REGISTER_CREATION_FACTORY(FFloatArrayToVertexSelectionDataflowNode);
 		DATAFLOW_NODE_REGISTER_CREATION_FACTORY(FSetVertexColorInCollectionFromVertexSelectionDataflowNode);
-		DATAFLOW_NODE_REGISTER_CREATION_FACTORY(FMakeTransformDataflowNode);
 		DATAFLOW_NODE_REGISTER_CREATION_FACTORY(FSetVertexColorInCollectionFromFloatArrayDataflowNode);
 		DATAFLOW_NODE_REGISTER_CREATION_FACTORY(FFloatArrayNormalizeDataflowNode);
 		DATAFLOW_NODE_REGISTER_CREATION_FACTORY(FVectorArrayNormalizeDataflowNode);
-		DATAFLOW_NODE_REGISTER_CREATION_FACTORY(FMakeQuaternionDataflowNode);
 		DATAFLOW_NODE_REGISTER_CREATION_FACTORY(FMultiplyTransformDataflowNode);
 		DATAFLOW_NODE_REGISTER_CREATION_FACTORY(FInvertTransformDataflowNode);
 
@@ -116,8 +93,6 @@ namespace Dataflow
 		DATAFLOW_NODE_REGISTER_CREATION_FACTORY_NODE_COLORS_BY_CATEGORY("Fracture", FLinearColor(1.f, 1.f, 0.8f), CDefaultNodeBodyTintColor);
 		// Utilities
 		DATAFLOW_NODE_REGISTER_CREATION_FACTORY_NODE_COLORS_BY_CATEGORY("Utilities", FLinearColor(1.f, 1.f, 0.f), CDefaultNodeBodyTintColor);
-		// Math
-		DATAFLOW_NODE_REGISTER_CREATION_FACTORY_NODE_COLORS_BY_CATEGORY("Math", FLinearColor(0.f, 0.4f, 0.8f), CDefaultNodeBodyTintColor);
 		// Generators
 		DATAFLOW_NODE_REGISTER_CREATION_FACTORY_NODE_COLORS_BY_CATEGORY("Generators", FLinearColor(.4f, 0.8f, 0.f), CDefaultNodeBodyTintColor);
 	}
@@ -182,13 +157,6 @@ void FLogStringDataflowNode::Evaluate(Dataflow::FContext& Context, const FDatafl
 	}
 }
 
-void FMakeLiteralStringDataflowNode::Evaluate(Dataflow::FContext& Context, const FDataflowOutput* Out) const
-{
-	if (Out->IsA<FString>(&String))
-	{
-		SetValue<FString>(Context, Value, &String);
-	}
-}
 
 
 void FBoundingBoxDataflowNode::Evaluate(Dataflow::FContext& Context, const FDataflowOutput* Out) const
@@ -231,116 +199,6 @@ void FExpandBoundingBoxDataflowNode::Evaluate(Dataflow::FContext& Context, const
 	}
 }
 
-void FVectorToStringDataflowNode::Evaluate(Dataflow::FContext& Context, const FDataflowOutput* Out) const
-{
-	if (Out->IsA<FString>(&String))
-	{
-		FString Value = GetValue<FVector>(Context, &Vector).ToString();
-		SetValue<FString>(Context, Value, &String);
-	}
-}
-
-void FFloatToStringDataflowNode::Evaluate(Dataflow::FContext& Context, const FDataflowOutput* Out) const
-{
-	if (Out->IsA<FString>(&String))
-	{
-		FString Value = FString::Printf(TEXT("%f"), GetValue<float>(Context, &Float));
-		SetValue<FString>(Context, Value, &String);
-	}
-}
-
-void FMakePointsDataflowNode::Evaluate(Dataflow::FContext& Context, const FDataflowOutput* Out) const
-{
-	if (Out->IsA<TArray<FVector>>(&Points))
-	{
-		SetValue<TArray<FVector>>(Context, Point, &Points);
-	}
-}
-
-void FMakeBoxDataflowNode::Evaluate(Dataflow::FContext& Context, const FDataflowOutput* Out) const
-{
-	if (Out->IsA<FBox>(&Box))
-	{
-		if (DataType == EMakeBoxDataTypeEnum::Dataflow_MakeBox_DataType_MinMax)
-		{
-			FVector MinVal = GetValue<FVector>(Context, &Min);
-			FVector MaxVal = GetValue<FVector>(Context, &Max);
-
-			SetValue<FBox>(Context, FBox(MinVal, MaxVal), &Box);
-		}
-		else if (DataType == EMakeBoxDataTypeEnum::Dataflow_MakeBox_DataType_CenterSize)
-		{
-			FVector CenterVal = GetValue<FVector>(Context, &Center);
-			FVector SizeVal = GetValue<FVector>(Context, &Size);
-
-			SetValue<FBox>(Context, FBox(CenterVal - 0.5 * SizeVal, CenterVal + 0.5 * SizeVal), &Box);
-		}
-	}
-}
-
-
-void FMakeSphereDataflowNode::Evaluate(Dataflow::FContext& Context, const FDataflowOutput* Out) const
-{
-	if (Out->IsA<FSphere>(&Sphere))
-	{
-		FVector CenterVal = GetValue<FVector>(Context, &Center);
-		float RadiusVal = GetValue<float>(Context, &Radius);
-
-		SetValue<FSphere>(Context, FSphere(CenterVal, RadiusVal), &Sphere);
-	}
-}
-
-
-void FMakeLiteralFloatDataflowNode::Evaluate(Dataflow::FContext& Context, const FDataflowOutput* Out) const
-{
-	if (Out->IsA<float>(&Float))
-	{
-		SetValue<float>(Context, Value, &Float);
-	}
-}
-
-void FMakeLiteralIntDataflowNode::Evaluate(Dataflow::FContext& Context, const FDataflowOutput* Out) const
-{
-	if (Out->IsA<int32>(&Int))
-	{
-		SetValue<int32>(Context, Value, &Int);
-	}
-}
-
-void FMakeLiteralBoolDataflowNode::Evaluate(Dataflow::FContext& Context, const FDataflowOutput* Out) const
-{
-	if (Out->IsA<bool>(&Bool))
-	{
-		SetValue<bool>(Context, Value, &Bool);
-	}
-}
-
-void FMakeLiteralVectorDataflowNode::Evaluate(Dataflow::FContext& Context, const FDataflowOutput* Out) const
-{
-	if (Out->IsA<FVector>(&Vector))
-	{
-		FVector Value(GetValue<float>(Context, &X, X), GetValue<float>(Context, &Y, Y), GetValue<float>(Context, &Z, Z));
-		SetValue<FVector>(Context, Value, &Vector);
-	}
-}
-
-void FIntToStringDataflowNode::Evaluate(Dataflow::FContext& Context, const FDataflowOutput* Out) const
-{
-	if (Out->IsA<FString>(&String))
-	{
-		FString Value = FString::Printf(TEXT("%d"), GetValue<int32>(Context, &Int));
-		SetValue<FString>(Context, Value, &String);
-	}
-}
-
-void FBoolToStringDataflowNode::Evaluate(Dataflow::FContext& Context, const FDataflowOutput* Out) const
-{
-	if (Out->IsA<FString>(&String))
-	{
-		FString Value = FString::Printf(TEXT("%s"), GetValue<bool>(Context, &Bool) ? TEXT("true") : TEXT("false"));
-		SetValue<FString>(Context, Value, &String);
-	}
-}
 
 void FExpandVectorDataflowNode::Evaluate(Dataflow::FContext& Context, const FDataflowOutput* Out) const
 {
@@ -360,16 +218,6 @@ void FExpandVectorDataflowNode::Evaluate(Dataflow::FContext& Context, const FDat
 	}
 }
 
-void FIntToFloatDataflowNode::Evaluate(Dataflow::FContext& Context, const FDataflowOutput* Out) const
-{
-	if (Out->IsA<float>(&Float))
-	{
-		float Value = float(GetValue<int32>(Context, &Int));
-		SetValue<float>(Context, Value, &Float);
-	}
-}
-
-
 void FStringAppendDataflowNode::Evaluate(Dataflow::FContext& Context, const FDataflowOutput* Out) const
 {
 	if (Out->IsA<FString>(&String))
@@ -379,100 +227,6 @@ void FStringAppendDataflowNode::Evaluate(Dataflow::FContext& Context, const FDat
 	}
 }
 
-void FRandomFloatDataflowNode::Evaluate(Dataflow::FContext& Context, const FDataflowOutput* Out) const
-{
-	if (Out->IsA<float>(&Float))
-	{
-		if (bDeterministic)
-		{
-			float RandomSeedVal = GetValue<float>(Context, &RandomSeed);
-
-			FRandomStream Stream(RandomSeedVal);
-			SetValue<float>(Context, Stream.FRand(), &Float);
-		}
-		else
-		{
-			SetValue<float>(Context, FMath::FRand(), &Float);
-		}
-	}
-}
-
-void FRandomFloatInRangeDataflowNode::Evaluate(Dataflow::FContext& Context, const FDataflowOutput* Out) const
-{
-	if (Out->IsA<float>(&Float))
-	{
-		float MinVal = GetValue<float>(Context, &Min);
-		float MaxVal = GetValue<float>(Context, &Max);
-
-		if (bDeterministic)
-		{
-			float RandomSeedVal = GetValue<float>(Context, &RandomSeed);
-
-			FRandomStream Stream(RandomSeedVal);
-			SetValue<float>(Context, Stream.FRandRange(MinVal, MaxVal), &Float);
-		}
-		else
-		{
-			SetValue<float>(Context, FMath::FRandRange(MinVal, MaxVal), &Float);
-		}
-	}
-}
-
-void FRandomUnitVectorDataflowNode::Evaluate(Dataflow::FContext& Context, const FDataflowOutput* Out) const
-{
-	if (Out->IsA<FVector>(&Vector))
-	{
-		if (bDeterministic)
-		{
-			float RandomSeedVal = GetValue<float>(Context, &RandomSeed);
-
-			FRandomStream Stream(RandomSeedVal);
-			SetValue<FVector>(Context, Stream.VRand(), &Vector);
-		}
-		else
-		{
-			SetValue<FVector>(Context, FMath::VRand(), &Vector);
-		}
-	}
-}
-
-void FRandomUnitVectorInConeDataflowNode::Evaluate(Dataflow::FContext& Context, const FDataflowOutput* Out) const
-{
-	if (Out->IsA<FVector>(&Vector))
-	{
-		FVector ConeDirectionVal = GetValue<FVector>(Context, &ConeDirection);
-		float ConeHalfAngleVal = GetValue<float>(Context, &ConeHalfAngle);
-
-		if (bDeterministic)
-		{
-			float RandomSeedVal = GetValue<float>(Context, &RandomSeed);
-
-			FRandomStream Stream(RandomSeedVal);
-			SetValue<FVector>(Context, Stream.VRandCone(ConeDirectionVal, ConeHalfAngleVal), &Vector);
-		}
-		else
-		{
-			SetValue<FVector>(Context, FMath::VRandCone(ConeDirectionVal, ConeHalfAngleVal), &Vector);
-		}
-	}
-}
-
-void FRadiansToDegreesDataflowNode::Evaluate(Dataflow::FContext& Context, const FDataflowOutput* Out) const
-{
-	if (Out->IsA<float>(&Degrees))
-	{
-		SetValue<float>(Context, FMath::RadiansToDegrees(GetValue<float>(Context, &Radians)), &Degrees);
-	}
-}
-
-void FDegreesToRadiansDataflowNode::Evaluate(Dataflow::FContext& Context, const FDataflowOutput* Out) const
-{
-	if (Out->IsA<float>(&Radians))
-	{
-		SetValue<float>(Context, FMath::DegreesToRadians(GetValue<float>(Context, &Degrees)), &Radians);
-	}
-}
-
 
 void FHashStringDataflowNode::Evaluate(Dataflow::FContext& Context, const FDataflowOutput* Out) const
 {
@@ -490,93 +244,6 @@ void FHashVectorDataflowNode::Evaluate(Dataflow::FContext& Context, const FDataf
 	}
 }
 
-void FFloatToIntDataflowNode::Evaluate(Dataflow::FContext& Context, const FDataflowOutput* Out) const
-{
-	if (Out->IsA<int32>(&Int))
-	{
-		float FloatVal = GetValue<float>(Context, &Float);
-		if (Function == EFloatToIntFunctionEnum::Dataflow_FloatToInt_Function_Floor)
-		{
-			SetValue<int32>(Context, FMath::FloorToInt32(FloatVal), &Int);
-		}
-		else if (Function == EFloatToIntFunctionEnum::Dataflow_FloatToInt_Function_Ceil)
-		{
-			SetValue<int32>(Context, FMath::CeilToInt32(FloatVal), &Int);
-		}
-		else if (Function == EFloatToIntFunctionEnum::Dataflow_FloatToInt_Function_Round)
-		{
-			SetValue<int32>(Context, FMath::RoundToInt32(FloatVal), &Int);
-		}
-		else if (Function == EFloatToIntFunctionEnum::Dataflow_FloatToInt_Function_Truncate)
-		{
-			SetValue<int32>(Context, int32(FMath::TruncToFloat(FloatVal)), &Int);
-		}
-	}
-}
-
-void FMathConstantsDataflowNode::Evaluate(Dataflow::FContext& Context, const FDataflowOutput* Out) const
-{
-	if (Out->IsA<float>(&Float))
-	{
-		if (Constant == EMathConstantsEnum::Dataflow_MathConstants_Pi)
-		{
-			SetValue<float>(Context, FMathf::Pi, &Float);
-		}
-		else if (Constant == EMathConstantsEnum::Dataflow_MathConstants_HalfPi)
-		{
-			SetValue<float>(Context, FMathf::HalfPi, &Float);
-		}
-		else if (Constant == EMathConstantsEnum::Dataflow_MathConstants_TwoPi)
-		{
-			SetValue<float>(Context, FMathf::TwoPi, &Float);
-		}
-		else if (Constant == EMathConstantsEnum::Dataflow_MathConstants_FourPi)
-		{
-			SetValue<float>(Context, FMathf::FourPi, &Float);
-		}
-		else if (Constant == EMathConstantsEnum::Dataflow_MathConstants_InvPi)
-		{
-			SetValue<float>(Context, FMathf::InvPi, &Float);
-		}
-		else if (Constant == EMathConstantsEnum::Dataflow_MathConstants_InvTwoPi)
-		{
-			SetValue<float>(Context, FMathf::InvTwoPi, &Float);
-		}
-		else if (Constant == EMathConstantsEnum::Dataflow_MathConstants_Sqrt2)
-		{
-			SetValue<float>(Context, FMathf::Sqrt2, &Float);
-		}
-		else if (Constant == EMathConstantsEnum::Dataflow_MathConstants_InvSqrt2)
-		{
-			SetValue<float>(Context, FMathf::InvSqrt2, &Float);
-		}
-		else if (Constant == EMathConstantsEnum::Dataflow_MathConstants_Sqrt3)
-		{
-			SetValue<float>(Context, FMathf::Sqrt3, &Float);
-		}
-		else if (Constant == EMathConstantsEnum::Dataflow_MathConstants_InvSqrt3)
-		{
-			SetValue<float>(Context, FMathf::InvSqrt3, &Float);
-		}
-		else if (Constant == EMathConstantsEnum::Dataflow_FloatToInt_Function_E)
-		{
-			SetValue<float>(Context, 2.71828182845904523536f, &Float);
-		}
-		else if (Constant == EMathConstantsEnum::Dataflow_FloatToInt_Function_Gamma)
-		{
-			SetValue<float>(Context, 0.577215664901532860606512090082f, &Float);
-		}
-		else if (Constant == EMathConstantsEnum::Dataflow_FloatToInt_Function_GoldenRatio)
-		{
-			SetValue<float>(Context, 1.618033988749894f, &Float);
-		}
-		else if (Constant == EMathConstantsEnum::Dataflow_FloatToInt_Function_ZeroTolerance)
-		{
-			SetValue<float>(Context, FMathf::ZeroTolerance, &Float);
-		}
-	}
-}
-
 void FGetArrayElementDataflowNode::Evaluate(Dataflow::FContext& Context, const FDataflowOutput* Out) const
 {
 	if (Out->IsA<FVector>(&Point))
@@ -1437,19 +1104,6 @@ void FSetVertexColorInCollectionFromVertexSelectionDataflowNode::Evaluate(Datafl
 }
 
 
-void FMakeTransformDataflowNode::Evaluate(Dataflow::FContext& Context, const FDataflowOutput* Out) const
-{
-	if (Out->IsA<FTransform>(&OutTransform))
-	{
-		SetValue<FTransform>(Context,
-			FTransform(FQuat::MakeFromEuler(GetValue<FVector>(Context, &InRotation))
-			, GetValue<FVector>(Context, &InTranslation)
-			, GetValue<FVector>(Context, &InScale))
-			, &OutTransform);
-	}
-}
-
-
 void FSetVertexColorInCollectionFromFloatArrayDataflowNode::Evaluate(Dataflow::FContext& Context, const FDataflowOutput* Out) const
 {
 	if (Out->IsA<FManagedArrayCollection>(&Collection))
@@ -1616,16 +1270,6 @@ void FVectorArrayNormalizeDataflowNode::Evaluate(Dataflow::FContext& Context, co
 	}
 }
 
-
-void FMakeQuaternionDataflowNode::Evaluate(Dataflow::FContext& Context, const FDataflowOutput* Out) const
-{
-	if (Out->IsA<FQuat>(&Quaternion))
-	{
-		FQuat Value(GetValue<float>(Context, &X, X), GetValue<float>(Context, &Y, Y), GetValue<float>(Context, &Z, Z), GetValue<float>(Context, &W, W));
-		SetValue<FQuat>(Context, Value, &Quaternion);
-	}
-}
-
 void FMultiplyTransformDataflowNode::Evaluate(Dataflow::FContext& Context, const FDataflowOutput* Out) const
 {
 	if (Out->IsA<FTransform>(&OutTransform))
diff --git a/Engine/Plugins/Experimental/GeometryCollectionPlugin/Source/GeometryCollectionNodes/Private/Dataflow/GeometryCollectionNodesPlugin.cpp b/Engine/Plugins/Experimental/GeometryCollectionPlugin/Source/GeometryCollectionNodes/Private/Dataflow/GeometryCollectionNodesPlugin.cpp
--- a/Engine/Plugins/Experimental/GeometryCollectionPlugin/Source/GeometryCollectionNodes/Private/Dataflow/GeometryCollectionNodesPlugin.cpp
+++ b/Engine/Plugins/Experimental/GeometryCollectionPlugin/Source/GeometryCollectionNodes/Private/Dataflow/GeometryCollectionNodesPlugin.cpp
@@ -17,6 +17,10 @@
 #include "Dataflow/GeometryCollectionUtilityNodes.h"
 #include "Dataflow/GeometryCollectionMaterialNodes.h"
 #include "Dataflow/GeometryCollectionFieldNodes.h"
+//#include "Dataflow/GeometryCollectionOverrideNodes.h"
+#include "Dataflow/GeometryCollectionMakeNodes.h"
+#include "Dataflow/GeometryCollectionMathNodes.h"
+#include "Dataflow/GeometryCollectionConversionNodes.h"
 #include "Dataflow/GeometryCollectionVerticesNodes.h"
 
 #define LOCTEXT_NAMESPACE "DataflowNodes"
@@ -36,6 +40,10 @@ void IGeometryCollectionNodesPlugin::StartupModule()
 	Dataflow::GeometryCollectionUtilityNodes();
 	Dataflow::GeometryCollectionMaterialNodes();
 	Dataflow::GeometryCollectionFieldNodes();
+	//Dataflow::GeometryCollectionOverrideNodes();
+	Dataflow::GeometryCollectionMakeNodes();
+	Dataflow::GeometryCollectionMathNodes();
+	Dataflow::GeometryCollectionConversionNodes();
 	Dataflow::GeometryCollectionVerticesNodes();
 }
 
diff --git a/Engine/Plugins/Experimental/GeometryCollectionPlugin/Source/GeometryCollectionNodes/Public/Dataflow/GeometryCollectionConversionNodes.h b/Engine/Plugins/Experimental/GeometryCollectionPlugin/Source/GeometryCollectionNodes/Public/Dataflow/GeometryCollectionConversionNodes.h
new file mode 100644
--- /dev/null
+++ b/Engine/Plugins/Experimental/GeometryCollectionPlugin/Source/GeometryCollectionNodes/Public/Dataflow/GeometryCollectionConversionNodes.h
@@ -0,0 +1,263 @@
+// Copyright Epic Games, Inc. All Rights Reserved.
+
+#pragma once 
+
+#include "CoreMinimal.h"
+#include "Dataflow/DataflowEngine.h"
+
+#include "GeometryCollectionConversionNodes.generated.h"
+
+/**
+ *
+ * Description for this node
+ *
+ */
+USTRUCT(meta = (DataflowGeometryCollection))
+struct FVectorToStringDataflowNode : public FDataflowNode
+{
+	GENERATED_USTRUCT_BODY()
+	DATAFLOW_NODE_DEFINE_INTERNAL(FVectorToStringDataflowNode, "VectorToString", "Utilities|String", "")
+
+public:
+	UPROPERTY(EditAnywhere, Category = "Vector", meta = (DataflowInput))
+	FVector Vector = FVector(0.0);
+
+	UPROPERTY(meta = (DataflowOutput))
+	FString String = FString("");
+
+	FVectorToStringDataflowNode(const Dataflow::FNodeParameters& InParam, FGuid InGuid = FGuid::NewGuid())
+		: FDataflowNode(InParam, InGuid)
+	{
+		RegisterInputConnection(&Vector);
+		RegisterOutputConnection(&String);
+	}
+
+	virtual void Evaluate(Dataflow::FContext& Context, const FDataflowOutput* Out) const override;
+
+};
+
+/**
+ *
+ * Description for this node
+ *
+ */
+USTRUCT(meta = (DataflowGeometryCollection))
+struct FFloatToStringDataflowNode : public FDataflowNode
+{
+	GENERATED_USTRUCT_BODY()
+	DATAFLOW_NODE_DEFINE_INTERNAL(FFloatToStringDataflowNode, "FloatToString", "Utilities|String", "")
+
+public:
+	UPROPERTY(EditAnywhere, Category = "Float", meta = (DataflowInput))
+	float Float = 0.f;
+
+	UPROPERTY(meta = (DataflowOutput))
+	FString String = FString("");
+
+	FFloatToStringDataflowNode(const Dataflow::FNodeParameters& InParam, FGuid InGuid = FGuid::NewGuid())
+		: FDataflowNode(InParam, InGuid)
+	{
+		RegisterInputConnection(&Float);
+		RegisterOutputConnection(&String);
+	}
+
+	virtual void Evaluate(Dataflow::FContext& Context, const FDataflowOutput* Out) const override;
+
+};
+
+/**
+ *
+ * Converts an Int to a String
+ *
+ */
+USTRUCT(meta = (DataflowGeometryCollection))
+struct FIntToStringDataflowNode : public FDataflowNode
+{
+	GENERATED_USTRUCT_BODY()
+	DATAFLOW_NODE_DEFINE_INTERNAL(FIntToStringDataflowNode, "IntToString", "Utilities|String", "")
+
+public:
+	UPROPERTY(EditAnywhere, Category = "Int", meta = (DataflowInput))
+	int32 Int = 0;
+
+	UPROPERTY(meta = (DataflowOutput))
+	FString String = FString("");
+
+	FIntToStringDataflowNode(const Dataflow::FNodeParameters& InParam, FGuid InGuid = FGuid::NewGuid())
+		: FDataflowNode(InParam, InGuid)
+	{
+		RegisterInputConnection(&Int);
+		RegisterOutputConnection(&String);
+	}
+
+	virtual void Evaluate(Dataflow::FContext& Context, const FDataflowOutput* Out) const override;
+
+};
+
+/**
+ *
+ * Converts a Bool to a String in a form of ("true", "false")
+ *
+ */
+USTRUCT(meta = (DataflowGeometryCollection))
+struct FBoolToStringDataflowNode : public FDataflowNode
+{
+	GENERATED_USTRUCT_BODY()
+	DATAFLOW_NODE_DEFINE_INTERNAL(FBoolToStringDataflowNode, "BoolToString", "Utilities|String", "")
+
+public:
+	UPROPERTY(EditAnywhere, Category = "Bool", meta = (DataflowInput))
+	bool Bool = false;
+
+	UPROPERTY(meta = (DataflowOutput))
+	FString String = FString("");
+
+	FBoolToStringDataflowNode(const Dataflow::FNodeParameters& InParam, FGuid InGuid = FGuid::NewGuid())
+		: FDataflowNode(InParam, InGuid)
+	{
+		RegisterInputConnection(&Bool);
+		RegisterOutputConnection(&String);
+	}
+
+	virtual void Evaluate(Dataflow::FContext& Context, const FDataflowOutput* Out) const override;
+
+};
+
+/**
+ *
+ * Converts an Int to a Float
+ *
+ */
+USTRUCT(meta = (DataflowGeometryCollection))
+struct FIntToFloatDataflowNode : public FDataflowNode
+{
+	GENERATED_USTRUCT_BODY()
+	DATAFLOW_NODE_DEFINE_INTERNAL(FIntToFloatDataflowNode, "IntToFloat", "Math|Conversions", "")
+
+public:
+	UPROPERTY(EditAnywhere, Category = "Int", meta = (DataflowInput))
+	int32 Int = 0;
+
+	UPROPERTY(meta = (DataflowOutput))
+	float Float = 0.f;
+
+	FIntToFloatDataflowNode(const Dataflow::FNodeParameters& InParam, FGuid InGuid = FGuid::NewGuid())
+		: FDataflowNode(InParam, InGuid)
+	{
+		RegisterInputConnection(&Int);
+		RegisterOutputConnection(&Float);
+	}
+
+	virtual void Evaluate(Dataflow::FContext& Context, const FDataflowOutput* Out) const override;
+
+};
+
+UENUM(BlueprintType)
+enum class EFloatToIntFunctionEnum : uint8
+{
+	Dataflow_FloatToInt_Function_Floor UMETA(DisplayName = "Floor()"),
+	Dataflow_FloatToInt_Function_Ceil UMETA(DisplayName = "Ceil()"),
+	Dataflow_FloatToInt_Function_Round UMETA(DisplayName = "Round()"),
+	Dataflow_FloatToInt_Function_Truncate UMETA(DisplayName = "Truncate()"),
+	//~~~
+	//256th entry
+	Dataflow_Max                UMETA(Hidden)
+};
+
+/**
+ *
+ * Converts a Float to Int using the specified method
+ *
+ */
+USTRUCT(meta = (DataflowGeometryCollection))
+struct FFloatToIntDataflowNode : public FDataflowNode
+{
+	GENERATED_USTRUCT_BODY()
+	DATAFLOW_NODE_DEFINE_INTERNAL(FFloatToIntDataflowNode, "FloatToInt", "Math|Conversions", "")
+
+public:
+	/** Method to convert */
+	UPROPERTY(EditAnywhere, Category = "Float");
+	EFloatToIntFunctionEnum Function = EFloatToIntFunctionEnum::Dataflow_FloatToInt_Function_Round;
+
+	/** Float value to convert */
+	UPROPERTY(EditAnywhere, Category = "Float", meta = (DataflowInput))
+	float Float = 0.f;
+
+	/** Int output */
+	UPROPERTY(meta = (DataflowOutput))
+	int32 Int = 0;
+
+	FFloatToIntDataflowNode(const Dataflow::FNodeParameters& InParam, FGuid InGuid = FGuid::NewGuid())
+		: FDataflowNode(InParam, InGuid)
+	{
+		RegisterInputConnection(&Float);
+		RegisterOutputConnection(&Int);
+	}
+
+	virtual void Evaluate(Dataflow::FContext& Context, const FDataflowOutput* Out) const override;
+
+};
+
+/**
+ *
+ * Converts an Int to a Bool
+ *
+ */
+USTRUCT(meta = (DataflowGeometryCollection))
+struct FIntToBoolDataflowNode : public FDataflowNode
+{
+	GENERATED_USTRUCT_BODY()
+	DATAFLOW_NODE_DEFINE_INTERNAL(FIntToBoolDataflowNode, "IntToBool", "Math|Conversions", "")
+
+public:
+	UPROPERTY(EditAnywhere, Category = "Int", meta = (DataflowInput))
+	int32 Int = 0;
+
+	UPROPERTY(meta = (DataflowOutput))
+	bool Bool = false;
+
+	FIntToBoolDataflowNode(const Dataflow::FNodeParameters& InParam, FGuid InGuid = FGuid::NewGuid())
+		: FDataflowNode(InParam, InGuid)
+	{
+		RegisterInputConnection(&Int);
+		RegisterOutputConnection(&Bool);
+	}
+
+	virtual void Evaluate(Dataflow::FContext& Context, const FDataflowOutput* Out) const override;
+
+};
+
+/**
+ *
+ * Converts a Bool to an Int
+ *
+ */
+USTRUCT(meta = (DataflowGeometryCollection))
+struct FBoolToIntDataflowNode : public FDataflowNode
+{
+	GENERATED_USTRUCT_BODY()
+	DATAFLOW_NODE_DEFINE_INTERNAL(FBoolToIntDataflowNode, "BoolToInt", "Math|Conversions", "")
+
+public:
+	UPROPERTY(EditAnywhere, Category = "Bool", meta = (DataflowInput))
+	bool Bool = false;
+
+	UPROPERTY(meta = (DataflowOutput))
+	int32 Int = 0;
+
+	FBoolToIntDataflowNode(const Dataflow::FNodeParameters& InParam, FGuid InGuid = FGuid::NewGuid())
+		: FDataflowNode(InParam, InGuid)
+	{
+		RegisterInputConnection(&Bool);
+		RegisterOutputConnection(&Int);
+	}
+
+	virtual void Evaluate(Dataflow::FContext& Context, const FDataflowOutput* Out) const override;
+
+};
+
+namespace Dataflow
+{
+	void GeometryCollectionConversionNodes();
+}
diff --git a/Engine/Plugins/Experimental/GeometryCollectionPlugin/Source/GeometryCollectionNodes/Public/Dataflow/GeometryCollectionMakeNodes.h b/Engine/Plugins/Experimental/GeometryCollectionPlugin/Source/GeometryCollectionNodes/Public/Dataflow/GeometryCollectionMakeNodes.h
new file mode 100644
--- /dev/null
+++ b/Engine/Plugins/Experimental/GeometryCollectionPlugin/Source/GeometryCollectionNodes/Public/Dataflow/GeometryCollectionMakeNodes.h
@@ -0,0 +1,369 @@
+// Copyright Epic Games, Inc. All Rights Reserved.
+
+#pragma once 
+
+#include "CoreMinimal.h"
+#include "Dataflow/DataflowEngine.h"
+#include "GeometryCollection/ManagedArrayCollection.h"
+#include "DynamicMesh/DynamicMesh3.h"
+#include "UDynamicMesh.h"
+#include "Dataflow/DataflowSelection.h"
+#include "Math/MathFwd.h"
+#include "DynamicMesh/DynamicMesh3.h"
+
+#include "GeometryCollectionMakeNodes.generated.h"
+
+/**
+ *
+ * Description for this node
+ *
+ */
+USTRUCT()
+struct FMakeLiteralStringDataflowNode : public FDataflowNode
+{
+	GENERATED_USTRUCT_BODY()
+	DATAFLOW_NODE_DEFINE_INTERNAL(FMakeLiteralStringDataflowNode, "MakeLiteralString", "Utilities|String", "")
+
+public:
+	UPROPERTY(EditAnywhere, Category = "String");
+	FString Value = FString("");
+
+	UPROPERTY(meta = (DataflowOutput))
+	FString String;
+
+	FMakeLiteralStringDataflowNode(const Dataflow::FNodeParameters& InParam, FGuid InGuid = FGuid::NewGuid())
+		: FDataflowNode(InParam, InGuid)
+	{
+		RegisterOutputConnection(&String);
+	}
+
+	virtual void Evaluate(Dataflow::FContext& Context, const FDataflowOutput* Out) const override;
+
+};
+
+/**
+ *
+ * Description for this node
+ *
+ */
+USTRUCT()
+struct FMakePointsDataflowNode : public FDataflowNode
+{
+	GENERATED_USTRUCT_BODY()
+	DATAFLOW_NODE_DEFINE_INTERNAL(FMakePointsDataflowNode, "MakePoints", "Generators|Point", "")
+
+public:
+	UPROPERTY(EditAnywhere, Category = "Point")
+	TArray<FVector> Point;
+
+	UPROPERTY(meta = (DataflowOutput))
+	TArray<FVector> Points;
+
+	FMakePointsDataflowNode(const Dataflow::FNodeParameters& InParam, FGuid InGuid = FGuid::NewGuid())
+		: FDataflowNode(InParam, InGuid)
+	{
+		RegisterOutputConnection(&Points);
+	}
+
+	virtual void Evaluate(Dataflow::FContext& Context, const FDataflowOutput* Out) const override;
+
+};
+
+
+UENUM(BlueprintType)
+enum class EMakeBoxDataTypeEnum : uint8
+{
+	Dataflow_MakeBox_DataType_MinMax UMETA(DisplayName = "Min/Max"),
+	Dataflow_MakeBox_DataType_CenterSize UMETA(DisplayName = "Center/Size"),
+	//~~~
+	//256th entry
+	Dataflow_Max                UMETA(Hidden)
+};
+
+/**
+ *
+ * Description for this node
+ *
+ */
+USTRUCT()
+struct FMakeBoxDataflowNode : public FDataflowNode
+{
+	GENERATED_USTRUCT_BODY()
+	DATAFLOW_NODE_DEFINE_INTERNAL(FMakeBoxDataflowNode, "MakeBox", "Generators|Box", "")
+	DATAFLOW_NODE_RENDER_TYPE(FName("FBox"), "Box")
+
+public:
+	UPROPERTY(EditAnywhere, Category = "Box", meta = (DisplayName = "Input Data Type"));
+	EMakeBoxDataTypeEnum DataType = EMakeBoxDataTypeEnum::Dataflow_MakeBox_DataType_MinMax;
+
+	UPROPERTY(EditAnywhere, Category = "Box", meta = (DataflowInput, EditCondition = "DataType == EMakeBoxDataTypeEnum::Dataflow_MakeBox_DataType_MinMax", EditConditionHides));
+	FVector Min = FVector(0.0);
+
+	UPROPERTY(EditAnywhere, Category = "Box", meta = (DataflowInput, EditCondition = "DataType == EMakeBoxDataTypeEnum::Dataflow_MakeBox_DataType_MinMax", EditConditionHides));
+	FVector Max = FVector(10.0);
+
+	UPROPERTY(EditAnywhere, Category = "Box", meta = (DataflowInput, EditCondition = "DataType == EMakeBoxDataTypeEnum::Dataflow_MakeBox_DataType_CenterSize", EditConditionHides));
+	FVector Center = FVector(0.0);
+
+	UPROPERTY(EditAnywhere, Category = "Box", meta = (DataflowInput, EditCondition = "DataType == EMakeBoxDataTypeEnum::Dataflow_MakeBox_DataType_CenterSize", EditConditionHides));
+	FVector Size = FVector(10.0);
+
+	UPROPERTY(meta = (DataflowOutput));
+	FBox Box = FBox(ForceInit);
+
+	FMakeBoxDataflowNode(const Dataflow::FNodeParameters& InParam, FGuid InGuid = FGuid::NewGuid())
+		: FDataflowNode(InParam, InGuid)
+	{
+		RegisterInputConnection(&Min);
+		RegisterInputConnection(&Max);
+		RegisterInputConnection(&Center);
+		RegisterInputConnection(&Size);
+		RegisterOutputConnection(&Box);
+	}
+
+	virtual void Evaluate(Dataflow::FContext& Context, const FDataflowOutput* Out) const override;
+
+};
+
+
+/**
+ *
+ * Description for this node
+ *
+ */
+USTRUCT()
+struct FMakeSphereDataflowNode : public FDataflowNode
+{
+	GENERATED_USTRUCT_BODY()
+	DATAFLOW_NODE_DEFINE_INTERNAL(FMakeSphereDataflowNode, "MakeSphere", "Generators|Sphere", "")
+
+public:
+	UPROPERTY(EditAnywhere, Category = "Sphere", meta = (DataflowInput));
+	FVector Center = FVector(0.f);
+
+	UPROPERTY(EditAnywhere, Category = "Sphere", meta = (DataflowInput));
+	float Radius = 10.f;
+
+	UPROPERTY(meta = (DataflowOutput));
+	FSphere Sphere = FSphere(ForceInit);
+
+	FMakeSphereDataflowNode(const Dataflow::FNodeParameters& InParam, FGuid InGuid = FGuid::NewGuid())
+		: FDataflowNode(InParam, InGuid)
+	{
+		RegisterInputConnection(&Center);
+		RegisterInputConnection(&Radius);
+		RegisterOutputConnection(&Sphere);
+	}
+
+	virtual void Evaluate(Dataflow::FContext& Context, const FDataflowOutput* Out) const override;
+
+};
+
+
+/**
+ *
+ * Description for this node
+ *
+ */
+USTRUCT()
+struct FMakeLiteralFloatDataflowNode : public FDataflowNode
+{
+	GENERATED_USTRUCT_BODY()
+	DATAFLOW_NODE_DEFINE_INTERNAL(FMakeLiteralFloatDataflowNode, "MakeLiteralFloat", "Math|Float", "")
+
+public:
+	UPROPERTY(EditAnywhere, Category = "Float");
+	float Value = 0.f;
+
+	UPROPERTY(meta = (DataflowOutput))
+	float Float = 0.f;
+
+	FMakeLiteralFloatDataflowNode(const Dataflow::FNodeParameters& InParam, FGuid InGuid = FGuid::NewGuid())
+		: FDataflowNode(InParam, InGuid)
+	{
+		RegisterOutputConnection(&Float);
+	}
+
+	virtual void Evaluate(Dataflow::FContext& Context, const FDataflowOutput* Out) const override;
+
+};
+
+
+/**
+ *
+ * Description for this node
+ *
+ */
+USTRUCT()
+struct FMakeLiteralIntDataflowNode : public FDataflowNode
+{
+	GENERATED_USTRUCT_BODY()
+	DATAFLOW_NODE_DEFINE_INTERNAL(FMakeLiteralIntDataflowNode, "MakeLiteralInt", "Math|Int", "")
+
+public:
+	UPROPERTY(EditAnywhere, Category = "Int");
+	int32 Value = 0;
+
+	UPROPERTY(meta = (DataflowOutput, DisplayName = "Int"))
+	int32 Int = 0;
+
+	FMakeLiteralIntDataflowNode(const Dataflow::FNodeParameters& InParam, FGuid InGuid = FGuid::NewGuid())
+		: FDataflowNode(InParam, InGuid)
+	{
+		RegisterOutputConnection(&Int);
+	}
+
+	virtual void Evaluate(Dataflow::FContext& Context, const FDataflowOutput* Out) const override;
+
+};
+
+/**
+ *
+ * Description for this node
+ *
+ */
+USTRUCT()
+struct FMakeLiteralBoolDataflowNode : public FDataflowNode
+{
+	GENERATED_USTRUCT_BODY()
+	DATAFLOW_NODE_DEFINE_INTERNAL(FMakeLiteralBoolDataflowNode, "MakeLiteralBool", "Math|Boolean", "")
+
+public:
+	UPROPERTY(EditAnywhere, Category = "Bool");
+	bool Value = false;
+
+	UPROPERTY(meta = (DataflowOutput))
+	bool Bool = false;
+
+	FMakeLiteralBoolDataflowNode(const Dataflow::FNodeParameters& InParam, FGuid InGuid = FGuid::NewGuid())
+		: FDataflowNode(InParam, InGuid)
+	{
+		RegisterOutputConnection(&Bool);
+	}
+
+	virtual void Evaluate(Dataflow::FContext& Context, const FDataflowOutput* Out) const override;
+
+};
+
+
+/**
+ *
+ *
+ *
+ */
+USTRUCT()
+struct FMakeLiteralVectorDataflowNode : public FDataflowNode
+{
+	GENERATED_USTRUCT_BODY()
+	DATAFLOW_NODE_DEFINE_INTERNAL(FMakeLiteralVectorDataflowNode, "MakeLiteralVector", "Math|Vector", "")
+
+public:
+	UPROPERTY(EditAnywhere, Category = "Vector", meta = (DataflowInput));
+	float X = float(0.0);
+
+	UPROPERTY(EditAnywhere, Category = "Vector", meta = (DataflowInput));
+	float Y = float(0.0);
+
+	UPROPERTY(EditAnywhere, Category = "Vector", meta = (DataflowInput));
+	float Z = float(0.0);
+
+	UPROPERTY(meta = (DataflowOutput, DisplayName = "Vector"))
+	FVector Vector = FVector(0.0);
+
+	FMakeLiteralVectorDataflowNode(const Dataflow::FNodeParameters& InParam, FGuid InGuid = FGuid::NewGuid())
+		: FDataflowNode(InParam, InGuid)
+	{
+		RegisterInputConnection(&X);
+		RegisterInputConnection(&Y);
+		RegisterInputConnection(&Z);
+		RegisterOutputConnection(&Vector);
+	}
+
+	virtual void Evaluate(Dataflow::FContext& Context, const FDataflowOutput* Out) const override;
+
+};
+
+/**
+ *
+ * Description for this node
+ *
+ */
+USTRUCT()
+struct FMakeTransformDataflowNode : public FDataflowNode
+{
+	GENERATED_USTRUCT_BODY()
+	DATAFLOW_NODE_DEFINE_INTERNAL(FMakeTransformDataflowNode, "MakeTransform", "Generators|Transform", "")
+
+public:
+
+	UPROPERTY(EditAnywhere, Category = "Transform", meta = (DataflowInput, DisplayName = "Translation"));
+	FVector InTranslation = FVector(0, 0, 0);
+
+	UPROPERTY(EditAnywhere, Category = "Transform", meta = (DataflowInput, DisplayName = "Rotation"));
+	FVector InRotation = FVector(0, 0, 0);
+
+	UPROPERTY(EditAnywhere, Category = "Transform", meta = (DataflowInput, DisplayName = "Scale"));
+	FVector InScale = FVector(1, 1, 1);
+
+	UPROPERTY(meta = (DataflowOutput, DisplayName = "Transform"));
+	FTransform OutTransform = FTransform::Identity;
+
+	FMakeTransformDataflowNode(const Dataflow::FNodeParameters& InParam, FGuid InGuid = FGuid::NewGuid())
+		: FDataflowNode(InParam, InGuid)
+	{
+		RegisterInputConnection(&InTranslation);
+		RegisterInputConnection(&InRotation);
+		RegisterInputConnection(&InScale);
+		RegisterOutputConnection(&OutTransform);
+	}
+
+	virtual void Evaluate(Dataflow::FContext& Context, const FDataflowOutput* Out) const override;
+
+};
+
+/**
+ *
+ *
+ *
+ */
+USTRUCT()
+struct FMakeQuaternionDataflowNode : public FDataflowNode
+{
+	GENERATED_USTRUCT_BODY()
+	DATAFLOW_NODE_DEFINE_INTERNAL(FMakeQuaternionDataflowNode, "MakeQuaternion", "Math|Vector", "")
+
+public:
+	UPROPERTY(EditAnywhere, Category = "Quaternion ", meta = (DataflowInput));
+	float X = float(0.0);
+
+	UPROPERTY(EditAnywhere, Category = "Quaternion", meta = (DataflowInput));
+	float Y = float(0.0);
+
+	UPROPERTY(EditAnywhere, Category = "Quaternion", meta = (DataflowInput));
+	float Z = float(0.0);
+
+	UPROPERTY(EditAnywhere, Category = "Quaternion", meta = (DataflowInput));
+	float W = float(0.0);
+
+	UPROPERTY(meta = (DataflowOutput, DisplayName = "Quaternion"))
+	FQuat Quaternion = FQuat(ForceInitToZero);
+
+	FMakeQuaternionDataflowNode(const Dataflow::FNodeParameters& InParam, FGuid InGuid = FGuid::NewGuid())
+		: FDataflowNode(InParam, InGuid)
+	{
+		RegisterInputConnection(&X);
+		RegisterInputConnection(&Y);
+		RegisterInputConnection(&Z);
+		RegisterInputConnection(&W);
+		RegisterOutputConnection(&Quaternion);
+	}
+
+	virtual void Evaluate(Dataflow::FContext& Context, const FDataflowOutput* Out) const override;
+
+};
+
+
+namespace Dataflow
+{
+	void GeometryCollectionMakeNodes();
+}
diff --git a/Engine/Plugins/Experimental/GeometryCollectionPlugin/Source/GeometryCollectionNodes/Public/Dataflow/GeometryCollectionMathNodes.h b/Engine/Plugins/Experimental/GeometryCollectionPlugin/Source/GeometryCollectionNodes/Public/Dataflow/GeometryCollectionMathNodes.h
--- a/Engine/Plugins/Experimental/GeometryCollectionPlugin/Source/GeometryCollectionNodes/Public/Dataflow/GeometryCollectionMathNodes.h
+++ b/Engine/Plugins/Experimental/GeometryCollectionPlugin/Source/GeometryCollectionNodes/Public/Dataflow/GeometryCollectionMathNodes.h
@@ -4,13 +4,1728 @@
 
 #include "CoreMinimal.h"
 #include "Dataflow/DataflowEngine.h"
-#include "GeometryCollection/ManagedArrayCollection.h"
-#include "Dataflow/DataflowSelection.h"
 
-//#include "GeometryCollectionMathNodes.generated.h"
+#include "GeometryCollectionMathNodes.generated.h"
+
+/**
+*
+* 
+*
+*/
+USTRUCT(meta = (DataflowGeometryCollection))
+struct FAddDataflowNode : public FDataflowNode
+{
+	GENERATED_USTRUCT_BODY()
+	DATAFLOW_NODE_DEFINE_INTERNAL(FAddDataflowNode, "Add", "Math|Float", "")
+
+public:
+	UPROPERTY(EditAnywhere, Category = "Add", meta = (DataflowInput));
+	float FloatA = 0.f;
+
+	UPROPERTY(EditAnywhere, Category = "Add", meta = (DataflowInput));
+	float FloatB = 0.f;
+	
+	UPROPERTY(meta = (DataflowOutput))
+	float ReturnValue = 0.f;
+
+	FAddDataflowNode(const Dataflow::FNodeParameters& InParam, FGuid InGuid = FGuid::NewGuid())
+		: FDataflowNode(InParam, InGuid)
+	{
+		RegisterInputConnection(&FloatA);
+		RegisterInputConnection(&FloatB);
+		RegisterOutputConnection(&ReturnValue);
+	}
+
+	virtual void Evaluate(Dataflow::FContext& Context, const FDataflowOutput* Out) const override;
+};
+
+/**
+*
+*
+*
+*/
+USTRUCT(meta = (DataflowGeometryCollection))
+struct FSubtractDataflowNode : public FDataflowNode
+{
+	GENERATED_USTRUCT_BODY()
+	DATAFLOW_NODE_DEFINE_INTERNAL(FSubtractDataflowNode, "Subtract", "Math|Float", "")
+
+public:
+	UPROPERTY(EditAnywhere, Category = "Subtract", meta = (DataflowInput));
+	float FloatA = 0.f;
+
+	UPROPERTY(EditAnywhere, Category = "Subtract", meta = (DataflowInput));
+	float FloatB = 0.f;
+
+	UPROPERTY(meta = (DataflowOutput))
+	float ReturnValue = 0.f;
+
+	FSubtractDataflowNode(const Dataflow::FNodeParameters& InParam, FGuid InGuid = FGuid::NewGuid())
+		: FDataflowNode(InParam, InGuid)
+	{
+		RegisterInputConnection(&FloatA);
+		RegisterInputConnection(&FloatB);
+		RegisterOutputConnection(&ReturnValue);
+	}
+
+	virtual void Evaluate(Dataflow::FContext& Context, const FDataflowOutput* Out) const override;
+};
+
+/**
+*
+*
+*
+*/
+USTRUCT(meta = (DataflowGeometryCollection))
+struct FMultiplyDataflowNode : public FDataflowNode
+{
+	GENERATED_USTRUCT_BODY()
+	DATAFLOW_NODE_DEFINE_INTERNAL(FMultiplyDataflowNode, "Multiply", "Math|Float", "")
+
+public:
+	UPROPERTY(EditAnywhere, Category = "Multiply", meta = (DataflowInput));
+	float FloatA = 0.f;
+
+	UPROPERTY(EditAnywhere, Category = "Multiply", meta = (DataflowInput));
+	float FloatB = 0.f;
+
+	UPROPERTY(meta = (DataflowOutput))
+	float ReturnValue = 0.f;
+
+	FMultiplyDataflowNode(const Dataflow::FNodeParameters& InParam, FGuid InGuid = FGuid::NewGuid())
+		: FDataflowNode(InParam, InGuid)
+	{
+		RegisterInputConnection(&FloatA);
+		RegisterInputConnection(&FloatB);
+		RegisterOutputConnection(&ReturnValue);
+	}
+
+	virtual void Evaluate(Dataflow::FContext& Context, const FDataflowOutput* Out) const override;
+};
+
+/**
+*
+*
+*
+*/
+USTRUCT(meta = (DataflowGeometryCollection))
+struct FSafeDivideDataflowNode : public FDataflowNode
+{
+	GENERATED_USTRUCT_BODY()
+	DATAFLOW_NODE_DEFINE_INTERNAL(FSafeDivideDataflowNode, "SafeDivide", "Math|Float", "")
+
+public:
+	UPROPERTY(EditAnywhere, Category = "SafeDivide", meta = (DataflowInput));
+	float FloatA = 0.f;
+
+	UPROPERTY(EditAnywhere, Category = "SafeDivide", meta = (DataflowInput));
+	float FloatB = 0.f;
+
+	UPROPERTY(meta = (DataflowOutput))
+	float ReturnValue = 0.f;
+
+	FSafeDivideDataflowNode(const Dataflow::FNodeParameters& InParam, FGuid InGuid = FGuid::NewGuid())
+		: FDataflowNode(InParam, InGuid)
+	{
+		RegisterInputConnection(&FloatA);
+		RegisterInputConnection(&FloatB);
+		RegisterOutputConnection(&ReturnValue);
+	}
+
+	virtual void Evaluate(Dataflow::FContext& Context, const FDataflowOutput* Out) const override;
+};
+
+/**
+*
+*
+*
+*/
+USTRUCT(meta = (DataflowGeometryCollection))
+struct FDivisionDataflowNode : public FDataflowNode
+{
+	GENERATED_USTRUCT_BODY()
+	DATAFLOW_NODE_DEFINE_INTERNAL(FDivisionDataflowNode, "Division (Whole and Remainder)", "Math|Float", "")
+
+public:
+	UPROPERTY(EditAnywhere, Category = "Division", meta = (DataflowInput));
+	float Dividend = 0.f;
+
+	UPROPERTY(EditAnywhere, Category = "Division", meta = (DataflowInput));
+	float Divisor = 0.f;
+
+	UPROPERTY(meta = (DataflowOutput))
+	float Remainder = 0.f;
+
+	UPROPERTY(meta = (DataflowOutput))
+	int32 ReturnValue = 0.f;
+
+	FDivisionDataflowNode(const Dataflow::FNodeParameters& InParam, FGuid InGuid = FGuid::NewGuid())
+		: FDataflowNode(InParam, InGuid)
+	{
+		RegisterInputConnection(&Dividend);
+		RegisterInputConnection(&Divisor);
+		RegisterOutputConnection(&Remainder);
+		RegisterOutputConnection(&ReturnValue);
+	}
+
+	virtual void Evaluate(Dataflow::FContext& Context, const FDataflowOutput* Out) const override;
+};
+
+/**
+*
+*
+*
+*/
+USTRUCT(meta = (DataflowGeometryCollection))
+struct FSafeReciprocalDataflowNode : public FDataflowNode
+{
+	GENERATED_USTRUCT_BODY()
+	DATAFLOW_NODE_DEFINE_INTERNAL(FSafeReciprocalDataflowNode, "SafeReciprocal", "Math|Float", "")
+
+public:
+	UPROPERTY(EditAnywhere, Category = "SafeReciprocal", meta = (DataflowInput));
+	float Float = 0.f;
+
+	UPROPERTY(meta = (DataflowOutput))
+	float ReturnValue = 0.f;
+
+	FSafeReciprocalDataflowNode(const Dataflow::FNodeParameters& InParam, FGuid InGuid = FGuid::NewGuid())
+		: FDataflowNode(InParam, InGuid)
+	{
+		RegisterInputConnection(&Float);
+		RegisterOutputConnection(&ReturnValue);
+	}
+
+	virtual void Evaluate(Dataflow::FContext& Context, const FDataflowOutput* Out) const override;
+};
+
+/**
+*
+*
+*
+*/
+USTRUCT(meta = (DataflowGeometryCollection))
+struct FSquareDataflowNode : public FDataflowNode
+{
+	GENERATED_USTRUCT_BODY()
+	DATAFLOW_NODE_DEFINE_INTERNAL(FSquareDataflowNode, "Square", "Math|Float", "")
+
+public:
+	UPROPERTY(EditAnywhere, Category = "Square", meta = (DataflowInput));
+	float Float = 0.f;
+
+	UPROPERTY(meta = (DataflowOutput))
+	float ReturnValue = 0.f;
+
+	FSquareDataflowNode(const Dataflow::FNodeParameters& InParam, FGuid InGuid = FGuid::NewGuid())
+		: FDataflowNode(InParam, InGuid)
+	{
+		RegisterInputConnection(&Float);
+		RegisterOutputConnection(&ReturnValue);
+	}
+
+	virtual void Evaluate(Dataflow::FContext& Context, const FDataflowOutput* Out) const override;
+};
+
+/**
+*
+*
+*
+*/
+USTRUCT(meta = (DataflowGeometryCollection))
+struct FSquareRootDataflowNode : public FDataflowNode
+{
+	GENERATED_USTRUCT_BODY()
+	DATAFLOW_NODE_DEFINE_INTERNAL(FSquareRootDataflowNode, "SquareRoot", "Math|Float", "")
+
+public:
+	UPROPERTY(EditAnywhere, Category = "SquareRoot", meta = (DataflowInput));
+	float Float = 0.f;
+
+	UPROPERTY(meta = (DataflowOutput))
+	float ReturnValue = 0.f;
+
+	FSquareRootDataflowNode(const Dataflow::FNodeParameters& InParam, FGuid InGuid = FGuid::NewGuid())
+		: FDataflowNode(InParam, InGuid)
+	{
+		RegisterInputConnection(&Float);
+		RegisterOutputConnection(&ReturnValue);
+	}
+
+	virtual void Evaluate(Dataflow::FContext& Context, const FDataflowOutput* Out) const override;
+};
+
+/**
+*
+*
+*
+*/
+USTRUCT(meta = (DataflowGeometryCollection))
+struct FInverseSqrtDataflowNode : public FDataflowNode
+{
+	GENERATED_USTRUCT_BODY()
+	DATAFLOW_NODE_DEFINE_INTERNAL(FInverseSqrtDataflowNode, "InverseSqrt", "Math|Float", "")
+
+public:
+	UPROPERTY(EditAnywhere, Category = "InverseSqrt", meta = (DataflowInput));
+	float Float = 0.f;
+
+	UPROPERTY(meta = (DataflowOutput))
+	float ReturnValue = 0.f;
+
+	FInverseSqrtDataflowNode(const Dataflow::FNodeParameters& InParam, FGuid InGuid = FGuid::NewGuid())
+		: FDataflowNode(InParam, InGuid)
+	{
+		RegisterInputConnection(&Float);
+		RegisterOutputConnection(&ReturnValue);
+	}
+
+	virtual void Evaluate(Dataflow::FContext& Context, const FDataflowOutput* Out) const override;
+};
+
+/**
+*
+*
+*
+*/
+USTRUCT(meta = (DataflowGeometryCollection))
+struct FCubeDataflowNode : public FDataflowNode
+{
+	GENERATED_USTRUCT_BODY()
+	DATAFLOW_NODE_DEFINE_INTERNAL(FCubeDataflowNode, "Cube", "Math|Float", "")
+
+public:
+	UPROPERTY(EditAnywhere, Category = "Cube", meta = (DataflowInput));
+	float Float = 0.f;
+
+	UPROPERTY(meta = (DataflowOutput))
+	float ReturnValue = 0.f;
+
+	FCubeDataflowNode(const Dataflow::FNodeParameters& InParam, FGuid InGuid = FGuid::NewGuid())
+		: FDataflowNode(InParam, InGuid)
+	{
+		RegisterInputConnection(&Float);
+		RegisterOutputConnection(&ReturnValue);
+	}
+
+	virtual void Evaluate(Dataflow::FContext& Context, const FDataflowOutput* Out) const override;
+};
+
+/**
+*
+*
+*
+*/
+USTRUCT(meta = (DataflowGeometryCollection))
+struct FNegateDataflowNode : public FDataflowNode
+{
+	GENERATED_USTRUCT_BODY()
+	DATAFLOW_NODE_DEFINE_INTERNAL(FNegateDataflowNode, "Negate", "Math|Float", "")
+
+public:
+	UPROPERTY(EditAnywhere, Category = "Negate", meta = (DataflowInput));
+	float Float = 0.f;
+
+	UPROPERTY(meta = (DataflowOutput))
+	float ReturnValue = 0.f;
+
+	FNegateDataflowNode(const Dataflow::FNodeParameters& InParam, FGuid InGuid = FGuid::NewGuid())
+		: FDataflowNode(InParam, InGuid)
+	{
+		RegisterInputConnection(&Float);
+		RegisterOutputConnection(&ReturnValue);
+	}
+
+	virtual void Evaluate(Dataflow::FContext& Context, const FDataflowOutput* Out) const override;
+};
+
+/**
+*
+*
+*
+*/
+USTRUCT(meta = (DataflowGeometryCollection))
+struct FAbsDataflowNode : public FDataflowNode
+{
+	GENERATED_USTRUCT_BODY()
+	DATAFLOW_NODE_DEFINE_INTERNAL(FAbsDataflowNode, "Abs", "Math|Float", "")
+
+public:
+	UPROPERTY(EditAnywhere, Category = "Abs", meta = (DataflowInput));
+	float Float = 0.f;
+
+	UPROPERTY(meta = (DataflowOutput))
+	float ReturnValue = 0.f;
+
+	FAbsDataflowNode(const Dataflow::FNodeParameters& InParam, FGuid InGuid = FGuid::NewGuid())
+		: FDataflowNode(InParam, InGuid)
+	{
+		RegisterInputConnection(&Float);
+		RegisterOutputConnection(&ReturnValue);
+	}
+
+	virtual void Evaluate(Dataflow::FContext& Context, const FDataflowOutput* Out) const override;
+};
+
+/**
+*
+* Converts a float to a nearest less or equal integer
+*
+*/
+USTRUCT(meta = (DataflowGeometryCollection))
+struct FFloorDataflowNode : public FDataflowNode
+{
+	GENERATED_USTRUCT_BODY()
+	DATAFLOW_NODE_DEFINE_INTERNAL(FFloorDataflowNode, "Floor", "Math|Float", "")
+
+public:
+	UPROPERTY(EditAnywhere, Category = "Floor", meta = (DataflowInput));
+	float Float = 0.f;
+
+	UPROPERTY(meta = (DataflowOutput))
+	float ReturnValue = 0.f;
+
+	FFloorDataflowNode(const Dataflow::FNodeParameters& InParam, FGuid InGuid = FGuid::NewGuid())
+		: FDataflowNode(InParam, InGuid)
+	{
+		RegisterInputConnection(&Float);
+		RegisterOutputConnection(&ReturnValue);
+	}
+
+	virtual void Evaluate(Dataflow::FContext& Context, const FDataflowOutput* Out) const override;
+};
+
+/**
+*
+* Converts a float to the nearest greater or equal integer
+*
+*/
+USTRUCT(meta = (DataflowGeometryCollection))
+struct FCeilDataflowNode : public FDataflowNode
+{
+	GENERATED_USTRUCT_BODY()
+	DATAFLOW_NODE_DEFINE_INTERNAL(FCeilDataflowNode, "Ceil", "Math|Float", "")
+
+public:
+	UPROPERTY(EditAnywhere, Category = "Ceil", meta = (DataflowInput));
+	float Float = 0.f;
+
+	UPROPERTY(meta = (DataflowOutput))
+	float ReturnValue = 0.f;
+
+	FCeilDataflowNode(const Dataflow::FNodeParameters& InParam, FGuid InGuid = FGuid::NewGuid())
+		: FDataflowNode(InParam, InGuid)
+	{
+		RegisterInputConnection(&Float);
+		RegisterOutputConnection(&ReturnValue);
+	}
+
+	virtual void Evaluate(Dataflow::FContext& Context, const FDataflowOutput* Out) const override;
+};
+
+/**
+*
+* Converts a float to the nearest integer. Rounds up when the fraction is .5
+*
+*/
+USTRUCT(meta = (DataflowGeometryCollection))
+struct FRoundDataflowNode : public FDataflowNode
+{
+	GENERATED_USTRUCT_BODY()
+	DATAFLOW_NODE_DEFINE_INTERNAL(FRoundDataflowNode, "Round", "Math|Float", "")
+
+public:
+	UPROPERTY(EditAnywhere, Category = "Round", meta = (DataflowInput));
+	float Float = 0.f;
+
+	UPROPERTY(meta = (DataflowOutput))
+	float ReturnValue = 0.f;
+
+	FRoundDataflowNode(const Dataflow::FNodeParameters& InParam, FGuid InGuid = FGuid::NewGuid())
+		: FDataflowNode(InParam, InGuid)
+	{
+		RegisterInputConnection(&Float);
+		RegisterOutputConnection(&ReturnValue);
+	}
+
+	virtual void Evaluate(Dataflow::FContext& Context, const FDataflowOutput* Out) const override;
+};
+
+/**
+*
+* Converts a float to the nearest integer. Rounds up when the fraction is .5
+*
+*/
+USTRUCT(meta = (DataflowGeometryCollection))
+struct FTruncDataflowNode : public FDataflowNode
+{
+	GENERATED_USTRUCT_BODY()
+	DATAFLOW_NODE_DEFINE_INTERNAL(FTruncDataflowNode, "Trunc", "Math|Float", "")
+
+public:
+	UPROPERTY(EditAnywhere, Category = "Trunc", meta = (DataflowInput));
+	float Float = 0.f;
+
+	UPROPERTY(meta = (DataflowOutput))
+	float ReturnValue = 0.f;
+
+	FTruncDataflowNode(const Dataflow::FNodeParameters& InParam, FGuid InGuid = FGuid::NewGuid())
+		: FDataflowNode(InParam, InGuid)
+	{
+		RegisterInputConnection(&Float);
+		RegisterOutputConnection(&ReturnValue);
+	}
+
+	virtual void Evaluate(Dataflow::FContext& Context, const FDataflowOutput* Out) const override;
+};
+
+/**
+*
+* Returns the fractional part of a float
+*
+*/
+USTRUCT(meta = (DataflowGeometryCollection))
+struct FFracDataflowNode : public FDataflowNode
+{
+	GENERATED_USTRUCT_BODY()
+	DATAFLOW_NODE_DEFINE_INTERNAL(FFracDataflowNode, "Frac", "Math|Float", "")
+
+public:
+	UPROPERTY(EditAnywhere, Category = "Frac", meta = (DataflowInput));
+	float Float = 0.f;
+
+	UPROPERTY(meta = (DataflowOutput))
+	float ReturnValue = 0.f;
+
+	FFracDataflowNode(const Dataflow::FNodeParameters& InParam, FGuid InGuid = FGuid::NewGuid())
+		: FDataflowNode(InParam, InGuid)
+	{
+		RegisterInputConnection(&Float);
+		RegisterOutputConnection(&ReturnValue);
+	}
+
+	virtual void Evaluate(Dataflow::FContext& Context, const FDataflowOutput* Out) const override;
+};
+
+/**
+*
+*
+*
+*/
+USTRUCT(meta = (DataflowGeometryCollection))
+struct FMinDataflowNode : public FDataflowNode
+{
+	GENERATED_USTRUCT_BODY()
+	DATAFLOW_NODE_DEFINE_INTERNAL(FMinDataflowNode, "Min", "Math|Float", "")
+
+public:
+	UPROPERTY(EditAnywhere, Category = "Min", meta = (DataflowInput));
+	float FloatA = 0.f;
+
+	UPROPERTY(EditAnywhere, Category = "Min", meta = (DataflowInput));
+	float FloatB = 0.f;
+
+	UPROPERTY(meta = (DataflowOutput))
+	float ReturnValue = 0.f;
+
+	FMinDataflowNode(const Dataflow::FNodeParameters& InParam, FGuid InGuid = FGuid::NewGuid())
+		: FDataflowNode(InParam, InGuid)
+	{
+		RegisterInputConnection(&FloatA);
+		RegisterInputConnection(&FloatB);
+		RegisterOutputConnection(&ReturnValue);
+	}
+
+	virtual void Evaluate(Dataflow::FContext& Context, const FDataflowOutput* Out) const override;
+};
+
+/**
+*
+*
+*
+*/
+USTRUCT(meta = (DataflowGeometryCollection))
+struct FMaxDataflowNode : public FDataflowNode
+{
+	GENERATED_USTRUCT_BODY()
+	DATAFLOW_NODE_DEFINE_INTERNAL(FMaxDataflowNode, "Max", "Math|Float", "")
+
+public:
+	UPROPERTY(EditAnywhere, Category = "Max", meta = (DataflowInput));
+	float FloatA = 0.f;
+
+	UPROPERTY(EditAnywhere, Category = "Max", meta = (DataflowInput));
+	float FloatB = 0.f;
+
+	UPROPERTY(meta = (DataflowOutput))
+		float ReturnValue = 0.f;
+
+	FMaxDataflowNode(const Dataflow::FNodeParameters& InParam, FGuid InGuid = FGuid::NewGuid())
+		: FDataflowNode(InParam, InGuid)
+	{
+		RegisterInputConnection(&FloatA);
+		RegisterInputConnection(&FloatB);
+		RegisterOutputConnection(&ReturnValue);
+	}
+
+	virtual void Evaluate(Dataflow::FContext& Context, const FDataflowOutput* Out) const override;
+};
+
+/**
+*
+*
+*
+*/
+USTRUCT(meta = (DataflowGeometryCollection))
+struct FMin3DataflowNode : public FDataflowNode
+{
+	GENERATED_USTRUCT_BODY()
+	DATAFLOW_NODE_DEFINE_INTERNAL(FMin3DataflowNode, "Min3", "Math|Float", "")
+
+public:
+	UPROPERTY(EditAnywhere, Category = "Min3", meta = (DataflowInput));
+	float FloatA = 0.f;
+
+	UPROPERTY(EditAnywhere, Category = "Min3", meta = (DataflowInput));
+	float FloatB = 0.f;
+
+	UPROPERTY(EditAnywhere, Category = "Min3", meta = (DataflowInput));
+	float FloatC = 0.f;
+
+	UPROPERTY(meta = (DataflowOutput))
+		float ReturnValue = 0.f;
+
+	FMin3DataflowNode(const Dataflow::FNodeParameters& InParam, FGuid InGuid = FGuid::NewGuid())
+		: FDataflowNode(InParam, InGuid)
+	{
+		RegisterInputConnection(&FloatA);
+		RegisterInputConnection(&FloatB);
+		RegisterInputConnection(&FloatC);
+		RegisterOutputConnection(&ReturnValue);
+	}
+
+	virtual void Evaluate(Dataflow::FContext& Context, const FDataflowOutput* Out) const override;
+};
+
+/**
+*
+*
+*
+*/
+USTRUCT(meta = (DataflowGeometryCollection))
+struct FMax3DataflowNode : public FDataflowNode
+{
+	GENERATED_USTRUCT_BODY()
+	DATAFLOW_NODE_DEFINE_INTERNAL(FMax3DataflowNode, "Max3", "Math|Float", "")
+
+public:
+	UPROPERTY(EditAnywhere, Category = "Max3", meta = (DataflowInput));
+	float FloatA = 0.f;
+
+	UPROPERTY(EditAnywhere, Category = "Max3", meta = (DataflowInput));
+	float FloatB = 0.f;
+
+	UPROPERTY(EditAnywhere, Category = "Max3", meta = (DataflowInput));
+	float FloatC = 0.f;
+
+	UPROPERTY(meta = (DataflowOutput))
+	float ReturnValue = 0.f;
+
+	FMax3DataflowNode(const Dataflow::FNodeParameters& InParam, FGuid InGuid = FGuid::NewGuid())
+		: FDataflowNode(InParam, InGuid)
+	{
+		RegisterInputConnection(&FloatA);
+		RegisterInputConnection(&FloatB);
+		RegisterInputConnection(&FloatC);
+		RegisterOutputConnection(&ReturnValue);
+	}
+
+	virtual void Evaluate(Dataflow::FContext& Context, const FDataflowOutput* Out) const override;
+};
+
+/**
+*
+* Returns 1, 0, or -1 depending on relation of Float to 0
+*
+*/
+USTRUCT(meta = (DataflowGeometryCollection))
+struct FSignDataflowNode : public FDataflowNode
+{
+	GENERATED_USTRUCT_BODY()
+	DATAFLOW_NODE_DEFINE_INTERNAL(FSignDataflowNode, "Sign", "Math|Float", "")
+
+public:
+	UPROPERTY(EditAnywhere, Category = "Sign", meta = (DataflowInput));
+	float Float = 0.f;
+
+	UPROPERTY(meta = (DataflowOutput))
+	float ReturnValue = 0.f;
+
+	FSignDataflowNode(const Dataflow::FNodeParameters& InParam, FGuid InGuid = FGuid::NewGuid())
+		: FDataflowNode(InParam, InGuid)
+	{
+		RegisterInputConnection(&Float);
+		RegisterOutputConnection(&ReturnValue);
+	}
+
+	virtual void Evaluate(Dataflow::FContext& Context, const FDataflowOutput* Out) const override;
+};
+
+/**
+*
+* Clamps X to be between Min and Max, inclusive
+*
+*/
+USTRUCT(meta = (DataflowGeometryCollection))
+struct FClampDataflowNode : public FDataflowNode
+{
+	GENERATED_USTRUCT_BODY()
+	DATAFLOW_NODE_DEFINE_INTERNAL(FClampDataflowNode, "Clamp", "Math|Float", "")
+
+public:
+	UPROPERTY(EditAnywhere, Category = "Clamp", meta = (DataflowInput));
+	float Float = 0.f;
+
+	UPROPERTY(EditAnywhere, Category = "Clamp", meta = (DataflowInput));
+	float Min = 0.f;
+
+	UPROPERTY(EditAnywhere, Category = "Clamp", meta = (DataflowInput));
+	float Max = 1.f;
+
+	UPROPERTY(meta = (DataflowOutput))
+	float ReturnValue = 0.f;
+
+	FClampDataflowNode(const Dataflow::FNodeParameters& InParam, FGuid InGuid = FGuid::NewGuid())
+		: FDataflowNode(InParam, InGuid)
+	{
+		RegisterInputConnection(&Float);
+		RegisterInputConnection(&Min);
+		RegisterInputConnection(&Max);
+		RegisterOutputConnection(&ReturnValue);
+	}
+
+	virtual void Evaluate(Dataflow::FContext& Context, const FDataflowOutput* Out) const override;
+};
+
+/**
+*
+* Fits a value from one range to another
+* Returns a number between newmin and newmax that is relative to num in the range between oldmin and oldmax. 
+* If the value is outside the old range, it will be clamped to the new range.
+*/
+USTRUCT(meta = (DataflowGeometryCollection))
+struct FFitDataflowNode : public FDataflowNode
+{
+	GENERATED_USTRUCT_BODY()
+	DATAFLOW_NODE_DEFINE_INTERNAL(FFitDataflowNode, "Fit", "Math|Float", "")
+
+public:
+	UPROPERTY(EditAnywhere, Category = "Fit", meta = (DataflowInput));
+	float Float = 0.f;
+
+	UPROPERTY(EditAnywhere, Category = "Fit", meta = (DataflowInput));
+	float OldMin = 0.f;
+
+	UPROPERTY(EditAnywhere, Category = "Fit", meta = (DataflowInput));
+	float OldMax = 1.f;
+
+	UPROPERTY(EditAnywhere, Category = "Fit", meta = (DataflowInput));
+	float NewMin = 0.f;
+
+	UPROPERTY(EditAnywhere, Category = "Fit", meta = (DataflowInput));
+	float NewMax = 1.f;
+
+	UPROPERTY(meta = (DataflowOutput))
+	float ReturnValue = 0.f;
+
+	FFitDataflowNode(const Dataflow::FNodeParameters& InParam, FGuid InGuid = FGuid::NewGuid())
+		: FDataflowNode(InParam, InGuid)
+	{
+		RegisterInputConnection(&Float);
+		RegisterInputConnection(&OldMin);
+		RegisterInputConnection(&OldMax);
+		RegisterInputConnection(&NewMin);
+		RegisterInputConnection(&NewMax);
+		RegisterOutputConnection(&ReturnValue);
+	}
+
+	virtual void Evaluate(Dataflow::FContext& Context, const FDataflowOutput* Out) const override;
+};
+
+/**
+*
+* Fits a value from one range to another
+* 
+* Takes the value num from the range (oldmin, oldmax) and shifts it to the corresponding value in the new range (newmin, newmax). Unlike fit, this function does not clamp values in the given range.
+*/
+USTRUCT(meta = (DataflowGeometryCollection))
+struct FEFitDataflowNode : public FDataflowNode
+{
+	GENERATED_USTRUCT_BODY()
+	DATAFLOW_NODE_DEFINE_INTERNAL(FEFitDataflowNode, "EFit", "Math|Float", "")
+
+public:
+	UPROPERTY(EditAnywhere, Category = "EFit", meta = (DataflowInput));
+	float Float = 0.f;
+
+	UPROPERTY(EditAnywhere, Category = "EFit", meta = (DataflowInput));
+	float OldMin = 0.f;
+
+	UPROPERTY(EditAnywhere, Category = "EFit", meta = (DataflowInput));
+	float OldMax = 1.f;
+
+	UPROPERTY(EditAnywhere, Category = "EFit", meta = (DataflowInput));
+	float NewMin = 0.f;
+
+	UPROPERTY(EditAnywhere, Category = "EFit", meta = (DataflowInput));
+	float NewMax = 1.f;
+
+	UPROPERTY(meta = (DataflowOutput))
+	float ReturnValue = 0.f;
+
+	FEFitDataflowNode(const Dataflow::FNodeParameters& InParam, FGuid InGuid = FGuid::NewGuid())
+		: FDataflowNode(InParam, InGuid)
+	{
+		RegisterInputConnection(&Float);
+		RegisterInputConnection(&OldMin);
+		RegisterInputConnection(&OldMax);
+		RegisterInputConnection(&NewMin);
+		RegisterInputConnection(&NewMax);
+		RegisterOutputConnection(&ReturnValue);
+	}
+
+	virtual void Evaluate(Dataflow::FContext& Context, const FDataflowOutput* Out) const override;
+};
+
+/**
+*
+* 
+* 
+*/
+USTRUCT(meta = (DataflowGeometryCollection))
+struct FPowDataflowNode : public FDataflowNode
+{
+	GENERATED_USTRUCT_BODY()
+	DATAFLOW_NODE_DEFINE_INTERNAL(FPowDataflowNode, "Pow", "Math|Float", "")
+
+public:
+	UPROPERTY(EditAnywhere, Category = "Pow", meta = (DataflowInput));
+	float Base = 0.f;
+
+	UPROPERTY(EditAnywhere, Category = "Pow", meta = (DataflowInput));
+	float Exp = 0.f;
+
+	UPROPERTY(meta = (DataflowOutput))
+	float ReturnValue = 0.f;
+
+	FPowDataflowNode(const Dataflow::FNodeParameters& InParam, FGuid InGuid = FGuid::NewGuid())
+		: FDataflowNode(InParam, InGuid)
+	{
+		RegisterInputConnection(&Base);
+		RegisterInputConnection(&Exp);
+		RegisterOutputConnection(&ReturnValue);
+	}
+
+	virtual void Evaluate(Dataflow::FContext& Context, const FDataflowOutput* Out) const override;
+};
+
+/**
+*
+* Returns log of A base Base (if Base ^ R == A, where R is ReturnValue)
+*
+*/
+USTRUCT(meta = (DataflowGeometryCollection))
+struct FLogDataflowNode : public FDataflowNode
+{
+	GENERATED_USTRUCT_BODY()
+	DATAFLOW_NODE_DEFINE_INTERNAL(FLogDataflowNode, "Log", "Math|Float", "")
+
+public:
+	UPROPERTY(EditAnywhere, Category = "Log", meta = (DataflowInput));
+	float Base = 1.f;
+
+	UPROPERTY(EditAnywhere, Category = "Log", meta = (DataflowInput));
+	float A = 0.f;
+
+	UPROPERTY(meta = (DataflowOutput))
+	float ReturnValue = 0.f;
+
+	FLogDataflowNode(const Dataflow::FNodeParameters& InParam, FGuid InGuid = FGuid::NewGuid())
+		: FDataflowNode(InParam, InGuid)
+	{
+		RegisterInputConnection(&A);
+		RegisterInputConnection(&Base);
+		RegisterOutputConnection(&ReturnValue);
+	}
+
+	virtual void Evaluate(Dataflow::FContext& Context, const FDataflowOutput* Out) const override;
+};
+
+/**
+*
+* Returns log of A base Base (if Base ^ R == A, where R is ReturnValue)
+*
+*/
+USTRUCT(meta = (DataflowGeometryCollection))
+struct FLogeDataflowNode : public FDataflowNode
+{
+	GENERATED_USTRUCT_BODY()
+	DATAFLOW_NODE_DEFINE_INTERNAL(FLogeDataflowNode, "Loge", "Math|Float", "")
+
+public:
+	UPROPERTY(EditAnywhere, Category = "Loge", meta = (DataflowInput));
+	float A = 0.f;
+
+	UPROPERTY(meta = (DataflowOutput))
+	float ReturnValue = 0.f;
+
+	FLogeDataflowNode(const Dataflow::FNodeParameters& InParam, FGuid InGuid = FGuid::NewGuid())
+		: FDataflowNode(InParam, InGuid)
+	{
+		RegisterInputConnection(&A);
+		RegisterOutputConnection(&ReturnValue);
+	}
+
+	virtual void Evaluate(Dataflow::FContext& Context, const FDataflowOutput* Out) const override;
+};
+
+/**
+*
+* Linearly interpolates between A and B based on Alpha. (100% of A when Alpha = 0 and 100% of B when Alpha = 100)
+*
+*/
+USTRUCT(meta = (DataflowGeometryCollection))
+struct FLerpDataflowNode : public FDataflowNode
+{
+	GENERATED_USTRUCT_BODY()
+	DATAFLOW_NODE_DEFINE_INTERNAL(FLerpDataflowNode, "Lerp", "Math|Float", "")
+
+public:
+	UPROPERTY(EditAnywhere, Category = "Lerp", meta = (DataflowInput));
+	float A = 0.f;
+
+	UPROPERTY(EditAnywhere, Category = "Lerp", meta = (DataflowInput));
+	float B = 1.f;
+
+	UPROPERTY(EditAnywhere, Category = "Lerp", meta = (DataflowInput, UIMin = 0.f, UIMax = 100.f));
+	float Alpha = 0.f;
+
+	UPROPERTY(meta = (DataflowOutput))
+	float ReturnValue = 0.f;
+
+	FLerpDataflowNode(const Dataflow::FNodeParameters& InParam, FGuid InGuid = FGuid::NewGuid())
+		: FDataflowNode(InParam, InGuid)
+	{
+		RegisterInputConnection(&A);
+		RegisterInputConnection(&B);
+		RegisterInputConnection(&Alpha);
+		RegisterOutputConnection(&ReturnValue);
+	}
+
+	virtual void Evaluate(Dataflow::FContext& Context, const FDataflowOutput* Out) const override;
+};
+
+/**
+*
+* 
+*
+*/
+USTRUCT(meta = (DataflowGeometryCollection))
+struct FWrapDataflowNode : public FDataflowNode
+{
+	GENERATED_USTRUCT_BODY()
+	DATAFLOW_NODE_DEFINE_INTERNAL(FWrapDataflowNode, "Wrap", "Math|Float", "")
+
+public:
+	UPROPERTY(EditAnywhere, Category = "Wrap", meta = (DataflowInput));
+	float Float = 0.f;
+
+	UPROPERTY(EditAnywhere, Category = "Wrap", meta = (DataflowInput));
+	float Min = 0.f;
+
+	UPROPERTY(EditAnywhere, Category = "Wrap", meta = (DataflowInput));
+	float Max = 1.f;
+
+	UPROPERTY(meta = (DataflowOutput))
+	float ReturnValue = 0.f;
+
+	FWrapDataflowNode(const Dataflow::FNodeParameters& InParam, FGuid InGuid = FGuid::NewGuid())
+		: FDataflowNode(InParam, InGuid)
+	{
+		RegisterInputConnection(&Float);
+		RegisterInputConnection(&Min);
+		RegisterInputConnection(&Max);
+		RegisterOutputConnection(&ReturnValue);
+	}
+
+	virtual void Evaluate(Dataflow::FContext& Context, const FDataflowOutput* Out) const override;
+};
+
+/**
+*
+*
+*
+*/
+USTRUCT(meta = (DataflowGeometryCollection))
+struct FExpDataflowNode : public FDataflowNode
+{
+	GENERATED_USTRUCT_BODY()
+	DATAFLOW_NODE_DEFINE_INTERNAL(FExpDataflowNode, "Exp", "Math|Float", "")
+
+public:
+	UPROPERTY(EditAnywhere, Category = "Exp", meta = (DataflowInput));
+	float Float = 0.f;
+
+	UPROPERTY(meta = (DataflowOutput))
+	float ReturnValue = 0.f;
+
+	FExpDataflowNode(const Dataflow::FNodeParameters& InParam, FGuid InGuid = FGuid::NewGuid())
+		: FDataflowNode(InParam, InGuid)
+	{
+		RegisterInputConnection(&Float);
+		RegisterOutputConnection(&ReturnValue);
+	}
+
+	virtual void Evaluate(Dataflow::FContext& Context, const FDataflowOutput* Out) const override;
+};
+
+/**
+*
+* Returns the sine of A. (Expects Radians)
+*
+*/
+USTRUCT(meta = (DataflowGeometryCollection))
+struct FSinDataflowNode : public FDataflowNode
+{
+	GENERATED_USTRUCT_BODY()
+	DATAFLOW_NODE_DEFINE_INTERNAL(FSinDataflowNode, "Sin", "Math|Trigonometry", "")
+
+public:
+	UPROPERTY(EditAnywhere, Category = "Sin", meta = (DataflowInput));
+	float Float = 0.f;
+
+	UPROPERTY(meta = (DataflowOutput))
+	float ReturnValue = 0.f;
+
+	FSinDataflowNode(const Dataflow::FNodeParameters& InParam, FGuid InGuid = FGuid::NewGuid())
+		: FDataflowNode(InParam, InGuid)
+	{
+		RegisterInputConnection(&Float);
+		RegisterOutputConnection(&ReturnValue);
+	}
+
+	virtual void Evaluate(Dataflow::FContext& Context, const FDataflowOutput* Out) const override;
+};
+
+/**
+*
+* Returns the arcsine of A. (ReturnValue is in Radians)
+*
+*/
+USTRUCT(meta = (DataflowGeometryCollection))
+struct FArcSinDataflowNode : public FDataflowNode
+{
+	GENERATED_USTRUCT_BODY()
+	DATAFLOW_NODE_DEFINE_INTERNAL(FArcSinDataflowNode, "ArcSin", "Math|Trigonometry", "")
+
+public:
+	UPROPERTY(EditAnywhere, Category = "ArcSin", meta = (DataflowInput, UIMin = -1.f, UIMax = 1.f));
+	float Float = 0.f;
+
+	UPROPERTY(meta = (DataflowOutput))
+	float ReturnValue = 0.f;
+
+	FArcSinDataflowNode(const Dataflow::FNodeParameters& InParam, FGuid InGuid = FGuid::NewGuid())
+		: FDataflowNode(InParam, InGuid)
+	{
+		RegisterInputConnection(&Float);
+		RegisterOutputConnection(&ReturnValue);
+	}
+
+	virtual void Evaluate(Dataflow::FContext& Context, const FDataflowOutput* Out) const override;
+};
+
+/**
+*
+* Returns the cosine of A. (Expects Radians)
+*
+*/
+USTRUCT(meta = (DataflowGeometryCollection))
+struct FCosDataflowNode : public FDataflowNode
+{
+	GENERATED_USTRUCT_BODY()
+	DATAFLOW_NODE_DEFINE_INTERNAL(FCosDataflowNode, "Cos", "Math|Trigonometry", "")
+
+public:
+	UPROPERTY(EditAnywhere, Category = "Cos", meta = (DataflowInput));
+	float Float = 0.f;
+
+	UPROPERTY(meta = (DataflowOutput))
+	float ReturnValue = 0.f;
+
+	FCosDataflowNode(const Dataflow::FNodeParameters& InParam, FGuid InGuid = FGuid::NewGuid())
+		: FDataflowNode(InParam, InGuid)
+	{
+		RegisterInputConnection(&Float);
+		RegisterOutputConnection(&ReturnValue);
+	}
+
+	virtual void Evaluate(Dataflow::FContext& Context, const FDataflowOutput* Out) const override;
+};
+
+/**
+*
+* Returns the arccosine of A. (ReturnValue is in Radians)
+*
+*/
+USTRUCT(meta = (DataflowGeometryCollection))
+struct FArcCosDataflowNode : public FDataflowNode
+{
+	GENERATED_USTRUCT_BODY()
+	DATAFLOW_NODE_DEFINE_INTERNAL(FArcCosDataflowNode, "ArcCos", "Math|Trigonometry", "")
+
+public:
+	UPROPERTY(EditAnywhere, Category = "ArcCos", meta = (DataflowInput, UIMin = -1.f, UIMax = 1.f));
+	float Float = 0.f;
+
+	UPROPERTY(meta = (DataflowOutput))
+	float ReturnValue = 0.f;
+
+	FArcCosDataflowNode(const Dataflow::FNodeParameters& InParam, FGuid InGuid = FGuid::NewGuid())
+		: FDataflowNode(InParam, InGuid)
+	{
+		RegisterInputConnection(&Float);
+		RegisterOutputConnection(&ReturnValue);
+	}
+
+	virtual void Evaluate(Dataflow::FContext& Context, const FDataflowOutput* Out) const override;
+};
+
+/**
+*
+* Returns the cosine of A. (Expects Radians)
+*
+*/
+USTRUCT(meta = (DataflowGeometryCollection))
+struct FTanDataflowNode : public FDataflowNode
+{
+	GENERATED_USTRUCT_BODY()
+	DATAFLOW_NODE_DEFINE_INTERNAL(FTanDataflowNode, "Tan", "Math|Trigonometry", "")
+
+public:
+	UPROPERTY(EditAnywhere, Category = "Tan", meta = (DataflowInput));
+	float Float = 0.f;
+
+	UPROPERTY(meta = (DataflowOutput))
+	float ReturnValue = 0.f;
+
+	FTanDataflowNode(const Dataflow::FNodeParameters& InParam, FGuid InGuid = FGuid::NewGuid())
+		: FDataflowNode(InParam, InGuid)
+	{
+		RegisterInputConnection(&Float);
+		RegisterOutputConnection(&ReturnValue);
+	}
+
+	virtual void Evaluate(Dataflow::FContext& Context, const FDataflowOutput* Out) const override;
+};
+
+/**
+*
+* Returns the arccosine of A. (ReturnValue is in Radians)
+*
+*/
+USTRUCT(meta = (DataflowGeometryCollection))
+struct FArcTanDataflowNode : public FDataflowNode
+{
+	GENERATED_USTRUCT_BODY()
+	DATAFLOW_NODE_DEFINE_INTERNAL(FArcTanDataflowNode, "ArcTan", "Math|Trigonometry", "")
+
+public:
+	UPROPERTY(EditAnywhere, Category = "ArcTan", meta = (DataflowInput));
+	float Float = 0.f;
+
+	UPROPERTY(meta = (DataflowOutput))
+	float ReturnValue = 0.f;
+
+	FArcTanDataflowNode(const Dataflow::FNodeParameters& InParam, FGuid InGuid = FGuid::NewGuid())
+		: FDataflowNode(InParam, InGuid)
+	{
+		RegisterInputConnection(&Float);
+		RegisterOutputConnection(&ReturnValue);
+	}
+
+	virtual void Evaluate(Dataflow::FContext& Context, const FDataflowOutput* Out) const override;
+};
+
+/**
+*
+* Returns the arccosine of A. (ReturnValue is in Radians)
+*
+*/
+USTRUCT(meta = (DataflowGeometryCollection))
+struct FArcTan2DataflowNode : public FDataflowNode
+{
+	GENERATED_USTRUCT_BODY()
+	DATAFLOW_NODE_DEFINE_INTERNAL(FArcTan2DataflowNode, "ArcTan2", "Math|Trigonometry", "")
+
+public:
+	UPROPERTY(EditAnywhere, Category = "ArcTan2", meta = (DataflowInput));
+	float Y = 0.f;
+
+	UPROPERTY(EditAnywhere, Category = "ArcTan2", meta = (DataflowInput));
+	float X = 0.f;
+
+	UPROPERTY(meta = (DataflowOutput))
+	float ReturnValue = 0.f;
+
+	FArcTan2DataflowNode(const Dataflow::FNodeParameters& InParam, FGuid InGuid = FGuid::NewGuid())
+		: FDataflowNode(InParam, InGuid)
+	{
+		RegisterInputConnection(&Y);
+		RegisterInputConnection(&X);
+		RegisterOutputConnection(&ReturnValue);
+	}
+
+	virtual void Evaluate(Dataflow::FContext& Context, const FDataflowOutput* Out) const override;
+};
+
+/**
+*
+* Returns Float normalized to the given range
+*
+*/
+USTRUCT(meta = (DataflowGeometryCollection))
+struct FNormalizeToRangeDataflowNode : public FDataflowNode
+{
+	GENERATED_USTRUCT_BODY()
+	DATAFLOW_NODE_DEFINE_INTERNAL(FNormalizeToRangeDataflowNode, "NormalizeToRange", "Math|Float", "")
+
+public:
+	UPROPERTY(EditAnywhere, Category = "NormalizeToRange", meta = (DataflowInput));
+	float Float = 0.f;
+
+	UPROPERTY(EditAnywhere, Category = "NormalizeToRange", meta = (DataflowInput));
+	float RangeMin = 0.f;
+
+	UPROPERTY(EditAnywhere, Category = "NormalizeToRange", meta = (DataflowInput));
+	float RangeMax = 1.f;
+
+	UPROPERTY(meta = (DataflowOutput))
+	float ReturnValue = 0.f;
+
+	FNormalizeToRangeDataflowNode(const Dataflow::FNodeParameters& InParam, FGuid InGuid = FGuid::NewGuid())
+		: FDataflowNode(InParam, InGuid)
+	{
+		RegisterInputConnection(&Float);
+		RegisterInputConnection(&RangeMin);
+		RegisterInputConnection(&RangeMax);
+		RegisterOutputConnection(&ReturnValue);
+	}
+
+	virtual void Evaluate(Dataflow::FContext& Context, const FDataflowOutput* Out) const override;
+};
+
+/**
+*
+*
+*
+*/
+USTRUCT(meta = (DataflowGeometryCollection))
+struct FDotProductDataflowNode : public FDataflowNode
+{
+	GENERATED_USTRUCT_BODY()
+	DATAFLOW_NODE_DEFINE_INTERNAL(FDotProductDataflowNode, "DotProduct", "Math|Vector", "")
+
+public:
+	UPROPERTY(EditAnywhere, Category = "DotProduct", meta = (DataflowInput));
+	FVector VectorA = FVector(0.f);
+
+	UPROPERTY(EditAnywhere, Category = "DotProduct", meta = (DataflowInput));
+	FVector VectorB = FVector(0.f);
+
+	UPROPERTY(meta = (DataflowOutput))
+	float ReturnValue = 0.f;
+
+	FDotProductDataflowNode(const Dataflow::FNodeParameters& InParam, FGuid InGuid = FGuid::NewGuid())
+		: FDataflowNode(InParam, InGuid)
+	{
+		RegisterInputConnection(&VectorA);
+		RegisterInputConnection(&VectorB);
+		RegisterOutputConnection(&ReturnValue);
+	}
+
+	virtual void Evaluate(Dataflow::FContext& Context, const FDataflowOutput* Out) const override;
+};
+
+/**
+*
+*
+*
+*/
+USTRUCT(meta = (DataflowGeometryCollection))
+struct FCrossProductDataflowNode : public FDataflowNode
+{
+	GENERATED_USTRUCT_BODY()
+	DATAFLOW_NODE_DEFINE_INTERNAL(FCrossProductDataflowNode, "CrossProduct", "Math|Vector", "")
+
+public:
+	UPROPERTY(EditAnywhere, Category = "CrossProduct", meta = (DataflowInput));
+	FVector VectorA = FVector(0.f);
+
+	UPROPERTY(EditAnywhere, Category = "CrossProduct", meta = (DataflowInput));
+	FVector VectorB = FVector(0.f);
+
+	UPROPERTY(meta = (DataflowOutput))
+	FVector ReturnValue = FVector(0.f);
+
+	FCrossProductDataflowNode(const Dataflow::FNodeParameters& InParam, FGuid InGuid = FGuid::NewGuid())
+		: FDataflowNode(InParam, InGuid)
+	{
+		RegisterInputConnection(&VectorA);
+		RegisterInputConnection(&VectorB);
+		RegisterOutputConnection(&ReturnValue);
+	}
+
+	virtual void Evaluate(Dataflow::FContext& Context, const FDataflowOutput* Out) const override;
+};
+
+/**
+*
+*
+*
+*/
+USTRUCT(meta = (DataflowGeometryCollection))
+struct FNormalizeDataflowNode : public FDataflowNode
+{
+	GENERATED_USTRUCT_BODY()
+	DATAFLOW_NODE_DEFINE_INTERNAL(FNormalizeDataflowNode, "Normalize", "Math|Vector", "")
+
+public:
+	UPROPERTY(EditAnywhere, Category = "Normalize", meta = (DataflowInput));
+	FVector VectorA = FVector(0.f);
+
+	UPROPERTY(EditAnywhere, Category = "Normalize", meta = (DataflowInput));
+	float Tolerance = 0.0001;
+
+	UPROPERTY(meta = (DataflowOutput))
+	FVector ReturnValue = FVector(0.f);
+
+	FNormalizeDataflowNode(const Dataflow::FNodeParameters& InParam, FGuid InGuid = FGuid::NewGuid())
+		: FDataflowNode(InParam, InGuid)
+	{
+		RegisterInputConnection(&VectorA);
+		RegisterInputConnection(&Tolerance);
+		RegisterOutputConnection(&ReturnValue);
+	}
+
+	virtual void Evaluate(Dataflow::FContext& Context, const FDataflowOutput* Out) const override;
+};
+
+/**
+*
+*
+*
+*/
+USTRUCT(meta = (DataflowGeometryCollection))
+struct FLengthDataflowNode : public FDataflowNode
+{
+	GENERATED_USTRUCT_BODY()
+	DATAFLOW_NODE_DEFINE_INTERNAL(FLengthDataflowNode, "Length", "Math|Vector", "")
+
+public:
+	UPROPERTY(EditAnywhere, Category = "Length", meta = (DataflowInput));
+	FVector Vector = FVector(0.f);
+
+	UPROPERTY(meta = (DataflowOutput))
+	float ReturnValue = 0.f;
+
+	FLengthDataflowNode(const Dataflow::FNodeParameters& InParam, FGuid InGuid = FGuid::NewGuid())
+		: FDataflowNode(InParam, InGuid)
+	{
+		RegisterInputConnection(&Vector);
+		RegisterOutputConnection(&ReturnValue);
+	}
+
+	virtual void Evaluate(Dataflow::FContext& Context, const FDataflowOutput* Out) const override;
+};
+
+/**
+*
+*
+*
+*/
+USTRUCT(meta = (DataflowGeometryCollection))
+struct FDistanceDataflowNode : public FDataflowNode
+{
+	GENERATED_USTRUCT_BODY()
+	DATAFLOW_NODE_DEFINE_INTERNAL(FDistanceDataflowNode, "Distance", "Math|Points", "")
+
+public:
+	UPROPERTY(EditAnywhere, Category = "Distance", meta = (DataflowInput));
+	FVector PointA = FVector(0.f);
+
+	UPROPERTY(EditAnywhere, Category = "Distance", meta = (DataflowInput));
+	FVector PointB = FVector(0.f);
+
+	UPROPERTY(meta = (DataflowOutput))
+	float ReturnValue = 0.f;
+
+	FDistanceDataflowNode(const Dataflow::FNodeParameters& InParam, FGuid InGuid = FGuid::NewGuid())
+		: FDataflowNode(InParam, InGuid)
+	{
+		RegisterInputConnection(&PointA);
+		RegisterInputConnection(&PointB);
+		RegisterOutputConnection(&ReturnValue);
+	}
+
+	virtual void Evaluate(Dataflow::FContext& Context, const FDataflowOutput* Out) const override;
+};
+
+/**
+*
+* 
+*
+*/
+USTRUCT(meta = (DataflowGeometryCollection))
+struct FIsNearlyZeroDataflowNode : public FDataflowNode
+{
+	GENERATED_USTRUCT_BODY()
+	DATAFLOW_NODE_DEFINE_INTERNAL(FIsNearlyZeroDataflowNode, "IsNearlyZero", "Math|Utilities", "")
+
+public:
+	UPROPERTY(EditAnywhere, Category = "IsNearlyZero", meta = (DataflowInput));
+	float Float = 0.f;
+
+	UPROPERTY(meta = (DataflowOutput))
+	bool ReturnValue = false;
+
+	FIsNearlyZeroDataflowNode(const Dataflow::FNodeParameters& InParam, FGuid InGuid = FGuid::NewGuid())
+		: FDataflowNode(InParam, InGuid)
+	{
+		RegisterInputConnection(&Float);
+		RegisterOutputConnection(&ReturnValue);
+	}
+
+	virtual void Evaluate(Dataflow::FContext& Context, const FDataflowOutput* Out) const override;
+};
+
+/**
+ *
+ * Generates a random float
+ *
+ */
+USTRUCT(meta = (DataflowGeometryCollection))
+struct FRandomFloatDataflowNode : public FDataflowNode
+{
+	GENERATED_USTRUCT_BODY()
+	DATAFLOW_NODE_DEFINE_INTERNAL(FRandomFloatDataflowNode, "RandomFloat", "Math|Random", "")
+
+public:
+	UPROPERTY(EditAnywhere, Category = "Seed")
+	bool bDeterministic = false;
+
+	UPROPERTY(EditAnywhere, Category = "Seed", meta = (DataflowInput, EditCondition = "Deterministic"))
+	float RandomSeed = 0.f;
+
+	UPROPERTY(meta = (DataflowOutput))
+	float ReturnValue = 0.f;
+
+	FRandomFloatDataflowNode(const Dataflow::FNodeParameters& InParam, FGuid InGuid = FGuid::NewGuid())
+		: FDataflowNode(InParam, InGuid)
+	{
+		RandomSeed = FMath::FRandRange(-1e5, 1e5);
+		RegisterInputConnection(&RandomSeed);
+		RegisterOutputConnection(&ReturnValue);
+	}
+
+	virtual void Evaluate(Dataflow::FContext& Context, const FDataflowOutput* Out) const override;
+
+};
+
+/**
+ *
+ * Generates a random float between Min and Max
+ *
+ */
+USTRUCT(meta = (DataflowGeometryCollection))
+struct FRandomFloatInRangeDataflowNode : public FDataflowNode
+{
+	GENERATED_USTRUCT_BODY()
+	DATAFLOW_NODE_DEFINE_INTERNAL(FRandomFloatInRangeDataflowNode, "RandomFloatInRange", "Math|Random", "")
+
+public:
+	UPROPERTY(EditAnywhere, Category = "Seed")
+	bool bDeterministic = false;
+
+	UPROPERTY(EditAnywhere, Category = "Seed", meta = (DataflowInput, EditCondition = "Deterministic"))
+	float RandomSeed = 0.f;
+
+	UPROPERTY(EditAnywhere, Category = "Random", meta = (DataflowInput))
+	float Min = 0.f;
+
+	UPROPERTY(EditAnywhere, Category = "Random", meta = (DataflowInput))
+	float Max = 1.f;
+
+	UPROPERTY(meta = (DataflowOutput))
+	float ReturnValue = 0.f;
+
+	FRandomFloatInRangeDataflowNode(const Dataflow::FNodeParameters& InParam, FGuid InGuid = FGuid::NewGuid())
+		: FDataflowNode(InParam, InGuid)
+	{
+		RandomSeed = FMath::FRandRange(-1e5, 1e5);
+		RegisterInputConnection(&RandomSeed);
+		RegisterInputConnection(&Min);
+		RegisterInputConnection(&Max);
+		RegisterOutputConnection(&ReturnValue);
+	}
+
+	virtual void Evaluate(Dataflow::FContext& Context, const FDataflowOutput* Out) const override;
+
+};
+
+
+/**
+ *
+ * Returns a random vector with length of 1
+ *
+ */
+USTRUCT(meta = (DataflowGeometryCollection))
+struct FRandomUnitVectorDataflowNode : public FDataflowNode
+{
+	GENERATED_USTRUCT_BODY()
+	DATAFLOW_NODE_DEFINE_INTERNAL(FRandomUnitVectorDataflowNode, "RandomUnitVector", "Math|Random", "")
+
+public:
+	UPROPERTY(EditAnywhere, Category = "Seed")
+	bool bDeterministic = false;
+
+	UPROPERTY(EditAnywhere, Category = "Seed", meta = (DataflowInput, EditCondition = "Deterministic"))
+	float RandomSeed = 0.f;
+
+	UPROPERTY(meta = (DataflowOutput))
+	FVector ReturnValue = FVector(0.0);
+
+	FRandomUnitVectorDataflowNode(const Dataflow::FNodeParameters& InParam, FGuid InGuid = FGuid::NewGuid())
+		: FDataflowNode(InParam, InGuid)
+	{
+		RandomSeed = FMath::FRandRange(-1e5, 1e5);
+		RegisterInputConnection(&RandomSeed);
+		RegisterOutputConnection(&ReturnValue);
+	}
+
+	virtual void Evaluate(Dataflow::FContext& Context, const FDataflowOutput* Out) const override;
+
+};
+
+
+/**
+ *
+ * Returns a random vector with length of 1, within the specified cone, with uniform random distribution
+ *
+ */
+USTRUCT(meta = (DataflowGeometryCollection))
+struct FRandomUnitVectorInConeDataflowNode : public FDataflowNode
+{
+	GENERATED_USTRUCT_BODY()
+	DATAFLOW_NODE_DEFINE_INTERNAL(FRandomUnitVectorInConeDataflowNode, "RandomUnitVectorInCone", "Math|Random", "")
+
+public:
+	UPROPERTY(EditAnywhere, Category = "Seed")
+	bool bDeterministic = false;
+
+	UPROPERTY(EditAnywhere, Category = "Seed", meta = (DataflowInput, EditCondition = "Deterministic"))
+	float RandomSeed = 0.f;
+
+	/** The base "center" direction of the cone */
+	UPROPERTY(EditAnywhere, Category = "Random", meta = (DataflowInput))
+	FVector ConeDirection = FVector(0.0, 0.0, 1.0);
+
+	/** The half-angle of the cone (from ConeDir to edge), in degrees */
+	UPROPERTY(EditAnywhere, Category = "Random", meta = (DataflowInput))
+	float ConeHalfAngle = PI / 4.f;
+
+	UPROPERTY(meta = (DataflowOutput))
+	FVector ReturnValue = FVector(0.0);
+
+	FRandomUnitVectorInConeDataflowNode(const Dataflow::FNodeParameters& InParam, FGuid InGuid = FGuid::NewGuid())
+		: FDataflowNode(InParam, InGuid)
+	{
+		RandomSeed = FMath::FRandRange(-1e5, 1e5);
+		RegisterInputConnection(&RandomSeed);
+		RegisterInputConnection(&ConeDirection);
+		RegisterInputConnection(&ConeHalfAngle);
+		RegisterOutputConnection(&ReturnValue);
+	}
+
+	virtual void Evaluate(Dataflow::FContext& Context, const FDataflowOutput* Out) const override;
+
+};
+
+
+/**
+ *
+ * Converts radians to degrees
+ *
+ */
+USTRUCT(meta = (DataflowGeometryCollection))
+struct FRadiansToDegreesDataflowNode : public FDataflowNode
+{
+	GENERATED_USTRUCT_BODY()
+	DATAFLOW_NODE_DEFINE_INTERNAL(FRadiansToDegreesDataflowNode, "RadiansToDegrees", "Math|Trigonometry", "")
+
+public:
+	UPROPERTY(EditAnywhere, Category = "Radians", meta = (DataflowInput))
+	float Radians = 0.f;
+
+	UPROPERTY(meta = (DataflowOutput))
+	float Degrees = 0.f;
+
+	FRadiansToDegreesDataflowNode(const Dataflow::FNodeParameters& InParam, FGuid InGuid = FGuid::NewGuid())
+		: FDataflowNode(InParam, InGuid)
+	{
+		RegisterInputConnection(&Radians);
+		RegisterOutputConnection(&Degrees);
+	}
+
+	virtual void Evaluate(Dataflow::FContext& Context, const FDataflowOutput* Out) const override;
+
+};
+
+
+/**
+ *
+ * Converts degrees to radians
+ *
+ */
+USTRUCT(meta = (DataflowGeometryCollection))
+struct FDegreesToRadiansDataflowNode : public FDataflowNode
+{
+	GENERATED_USTRUCT_BODY()
+	DATAFLOW_NODE_DEFINE_INTERNAL(FDegreesToRadiansDataflowNode, "DegreesToRadians", "Math|Trigonometry", "")
+
+public:
+	UPROPERTY(EditAnywhere, Category = "Degrees", meta = (DataflowInput))
+	float Degrees = 0.f;
+
+	UPROPERTY(meta = (DataflowOutput))
+	float Radians = 0.f;
+
+	FDegreesToRadiansDataflowNode(const Dataflow::FNodeParameters& InParam, FGuid InGuid = FGuid::NewGuid())
+		: FDataflowNode(InParam, InGuid)
+	{
+		RegisterInputConnection(&Degrees);
+		RegisterOutputConnection(&Radians);
+	}
+
+	virtual void Evaluate(Dataflow::FContext& Context, const FDataflowOutput* Out) const override;
+
+};
+
+UENUM(BlueprintType)
+enum class EMathConstantsEnum : uint8
+{
+	Dataflow_MathConstants_Pi UMETA(DisplayName = "Pi"),
+	Dataflow_MathConstants_HalfPi UMETA(DisplayName = "HalfPi"),
+	Dataflow_MathConstants_TwoPi UMETA(DisplayName = "TwoPi"),
+	Dataflow_MathConstants_FourPi UMETA(DisplayName = "FourPi"),
+	Dataflow_MathConstants_InvPi UMETA(DisplayName = "InvPi"),
+	Dataflow_MathConstants_InvTwoPi UMETA(DisplayName = "InvTwoPi"),
+	Dataflow_MathConstants_Sqrt2 UMETA(DisplayName = "Sqrt2"),
+	Dataflow_MathConstants_InvSqrt2 UMETA(DisplayName = "InvSqrt2"),
+	Dataflow_MathConstants_Sqrt3 UMETA(DisplayName = "Sqrt3"),
+	Dataflow_MathConstants_InvSqrt3 UMETA(DisplayName = "InvSqrt3"),
+	Dataflow_FloatToInt_Function_E UMETA(DisplayName = "e"),
+	Dataflow_FloatToInt_Function_Gamma UMETA(DisplayName = "Gamma"),
+	Dataflow_FloatToInt_Function_GoldenRatio UMETA(DisplayName = "GoldenRatio"),
+	Dataflow_FloatToInt_Function_ZeroTolerance UMETA(DisplayName = "ZeroTolerance"),
+	//~~~
+	//256th entry
+	Dataflow_Max                UMETA(Hidden)
+};
+
+/**
+ *
+ * Offers a selection of Math constants
+ *
+ */
+USTRUCT(meta = (DataflowGeometryCollection))
+struct FMathConstantsDataflowNode : public FDataflowNode
+{
+	GENERATED_USTRUCT_BODY()
+	DATAFLOW_NODE_DEFINE_INTERNAL(FMathConstantsDataflowNode, "MathConstants", "Math|Utilities", "")
+
+public:
+	/** Math constant to output */
+	UPROPERTY(EditAnywhere, Category = "Constants");
+	EMathConstantsEnum Constant = EMathConstantsEnum::Dataflow_MathConstants_Pi;
+
+	/** Selected Math constant */
+	UPROPERTY(meta = (DataflowOutput))
+	float ReturnValue = 0;
+
+	FMathConstantsDataflowNode(const Dataflow::FNodeParameters& InParam, FGuid InGuid = FGuid::NewGuid())
+		: FDataflowNode(InParam, InGuid)
+	{
+		RegisterOutputConnection(&ReturnValue);
+	}
+
+	virtual void Evaluate(Dataflow::FContext& Context, const FDataflowOutput* Out) const override;
+
+};
+
+
+
+
+
+
+
+
+
+
+
+
+
+
+
+
+
+
+
+
+
+
+
+
+
+
+
+
+
+
+
+
+
+
+
+
+
+
+
+
+
+
+
+
+
+
+
+
+
+
+
+
 
 
-class FGeometryCollection;
 
 
 
diff --git a/Engine/Plugins/Experimental/GeometryCollectionPlugin/Source/GeometryCollectionNodes/Public/Dataflow/GeometryCollectionNodes.h b/Engine/Plugins/Experimental/GeometryCollectionPlugin/Source/GeometryCollectionNodes/Public/Dataflow/GeometryCollectionNodes.h
--- a/Engine/Plugins/Experimental/GeometryCollectionPlugin/Source/GeometryCollectionNodes/Public/Dataflow/GeometryCollectionNodes.h
+++ b/Engine/Plugins/Experimental/GeometryCollectionPlugin/Source/GeometryCollectionNodes/Public/Dataflow/GeometryCollectionNodes.h
@@ -151,537 +151,33 @@ public:
 };
 
 
-/**
- *
- * Description for this node
- *
- */
-USTRUCT()
-struct FMakeLiteralStringDataflowNode : public FDataflowNode
-{
-	GENERATED_USTRUCT_BODY()
-	DATAFLOW_NODE_DEFINE_INTERNAL(FMakeLiteralStringDataflowNode, "MakeLiteralString", "Utilities|String", "")
-
-public:
-	UPROPERTY(EditAnywhere, Category = "String");
-	FString Value = FString("");
-
-	UPROPERTY(meta = (DataflowOutput))
-	FString String;
-
-	FMakeLiteralStringDataflowNode(const Dataflow::FNodeParameters& InParam, FGuid InGuid = FGuid::NewGuid())
-		: FDataflowNode(InParam, InGuid)
-	{
-		RegisterOutputConnection(&String);
-	}
-
-	virtual void Evaluate(Dataflow::FContext& Context, const FDataflowOutput* Out) const override;
-
-};
-
-
-/**
- *
- * Description for this node
- *
- */
-USTRUCT(meta = (DataflowGeometryCollection))
-struct FBoundingBoxDataflowNode : public FDataflowNode
-{
-	GENERATED_USTRUCT_BODY()
-	DATAFLOW_NODE_DEFINE_INTERNAL(FBoundingBoxDataflowNode, "BoundingBox", "Utilities|Box", "")
-	DATAFLOW_NODE_RENDER_TYPE(FName("FBox"), "BoundingBox")
-
-public:
-	UPROPERTY(meta = (DataflowInput))
-	FManagedArrayCollection Collection;
-
-	UPROPERTY(meta = (DataflowOutput))
-	FBox BoundingBox = FBox(ForceInit);
-
-	FBoundingBoxDataflowNode(const Dataflow::FNodeParameters& InParam, FGuid InGuid = FGuid::NewGuid())
-		: FDataflowNode(InParam, InGuid)
-	{
-		RegisterInputConnection(&Collection);
-		RegisterOutputConnection(&BoundingBox);
-	}
-
-	virtual void Evaluate(Dataflow::FContext& Context, const FDataflowOutput* Out) const override;
-
-};
-
-
-/**
- *
- * Description for this node
- *
- */
-USTRUCT()
-struct FExpandBoundingBoxDataflowNode : public FDataflowNode
-{
-	GENERATED_USTRUCT_BODY()
-	DATAFLOW_NODE_DEFINE_INTERNAL(FExpandBoundingBoxDataflowNode, "ExpandBoundingBox", "Utilities|Box", "")
-
-public:
-	UPROPERTY(meta = (DataflowInput))
-	FBox BoundingBox = FBox(ForceInit);;
-
-	UPROPERTY(meta = (DataflowOutput))
-	FVector Min = FVector(0.0);
-
-	UPROPERTY(meta = (DataflowOutput))
-	FVector Max = FVector(0.0);
-
-	UPROPERTY(meta = (DataflowOutput))
-	FVector Center = FVector(0.0);
-
-	UPROPERTY(meta = (DataflowOutput))
-	FVector HalfExtents = FVector(0.0);
-
-	UPROPERTY(meta = (DataflowOutput))
-	float Volume = 0.0;
-
-	FExpandBoundingBoxDataflowNode(const Dataflow::FNodeParameters& InParam, FGuid InGuid = FGuid::NewGuid())
-		: FDataflowNode(InParam, InGuid)
-	{
-		RegisterInputConnection(&BoundingBox);
-		RegisterOutputConnection(&Min);
-		RegisterOutputConnection(&Max);
-		RegisterOutputConnection(&Center);
-		RegisterOutputConnection(&HalfExtents);
-		RegisterOutputConnection(&Volume);
-	}
-
-	virtual void Evaluate(Dataflow::FContext& Context, const FDataflowOutput* Out) const override;
-
-};
-
-
-/**
- *
- * Description for this node
- *
- */
-USTRUCT()
-struct FVectorToStringDataflowNode : public FDataflowNode
-{
-	GENERATED_USTRUCT_BODY()
-	DATAFLOW_NODE_DEFINE_INTERNAL(FVectorToStringDataflowNode, "VectorToString", "Utilities|String", "")
-
-public:
-	UPROPERTY(EditAnywhere, Category = "Vector", meta = (DataflowInput))
-	FVector Vector = FVector(0.0);
-
-	UPROPERTY(meta = (DataflowOutput))
-	FString String = FString("");
-
-	FVectorToStringDataflowNode(const Dataflow::FNodeParameters& InParam, FGuid InGuid = FGuid::NewGuid())
-		: FDataflowNode(InParam, InGuid)
-	{
-		RegisterInputConnection(&Vector);
-		RegisterOutputConnection(&String);
-	}
-
-	virtual void Evaluate(Dataflow::FContext& Context, const FDataflowOutput* Out) const override;
-
-};
-
-
-/**
- *
- * Description for this node
- *
- */
-USTRUCT()
-struct FFloatToStringDataflowNode : public FDataflowNode
-{
-	GENERATED_USTRUCT_BODY()
-	DATAFLOW_NODE_DEFINE_INTERNAL(FFloatToStringDataflowNode, "FloatToString", "Utilities|String", "")
-
-public:
-	UPROPERTY(EditAnywhere, Category = "Float", meta = (DataflowInput))
-	float Float = 0.f;
-
-	UPROPERTY(meta = (DataflowOutput))
-	FString String = FString("");
-
-	FFloatToStringDataflowNode(const Dataflow::FNodeParameters& InParam, FGuid InGuid = FGuid::NewGuid())
-		: FDataflowNode(InParam, InGuid)
-	{
-		RegisterInputConnection(&Float);
-		RegisterOutputConnection(&String);
-	}
-
-	virtual void Evaluate(Dataflow::FContext& Context, const FDataflowOutput* Out) const override;
-
-};
-
-
-/**
- *
- * Description for this node
- *
- */
-USTRUCT()
-struct FMakePointsDataflowNode : public FDataflowNode
-{
-	GENERATED_USTRUCT_BODY()
-	DATAFLOW_NODE_DEFINE_INTERNAL(FMakePointsDataflowNode, "MakePoints", "Generators|Point", "")
-
-public:
-	UPROPERTY(EditAnywhere, Category = "Point")
-	TArray<FVector> Point;
-
-	UPROPERTY(meta = (DataflowOutput))
-	TArray<FVector> Points;
-
-	FMakePointsDataflowNode(const Dataflow::FNodeParameters& InParam, FGuid InGuid = FGuid::NewGuid())
-		: FDataflowNode(InParam, InGuid)
-	{
-		RegisterOutputConnection(&Points);
-	}
-
-	virtual void Evaluate(Dataflow::FContext& Context, const FDataflowOutput* Out) const override;
-
-};
-
-
-UENUM(BlueprintType)
-enum class EMakeBoxDataTypeEnum : uint8
-{
-	Dataflow_MakeBox_DataType_MinMax UMETA(DisplayName = "Min/Max"),
-	Dataflow_MakeBox_DataType_CenterSize UMETA(DisplayName = "Center/Size"),
-	//~~~
-	//256th entry
-	Dataflow_Max                UMETA(Hidden)
-};
-
-/**
- *
- * Description for this node
- *
- */
-USTRUCT()
-struct FMakeBoxDataflowNode : public FDataflowNode
-{
-	GENERATED_USTRUCT_BODY()
-	DATAFLOW_NODE_DEFINE_INTERNAL(FMakeBoxDataflowNode, "MakeBox", "Generators|Box", "")
-	DATAFLOW_NODE_RENDER_TYPE(FName("FBox"), "Box")
-
-public:
-	UPROPERTY(EditAnywhere, Category = "Box", meta = (DisplayName = "Input Data Type"));
-	EMakeBoxDataTypeEnum DataType = EMakeBoxDataTypeEnum::Dataflow_MakeBox_DataType_MinMax;
-
-	UPROPERTY(EditAnywhere, Category = "Box", meta = (DataflowInput, EditCondition = "DataType == EMakeBoxDataTypeEnum::Dataflow_MakeBox_DataType_MinMax", EditConditionHides));
-	FVector Min = FVector(0.0);
-
-	UPROPERTY(EditAnywhere, Category = "Box", meta = (DataflowInput, EditCondition = "DataType == EMakeBoxDataTypeEnum::Dataflow_MakeBox_DataType_MinMax", EditConditionHides));
-	FVector Max = FVector(10.0);
-
-	UPROPERTY(EditAnywhere, Category = "Box", meta = (DataflowInput, EditCondition = "DataType == EMakeBoxDataTypeEnum::Dataflow_MakeBox_DataType_CenterSize", EditConditionHides));
-	FVector Center = FVector(0.0);
-
-	UPROPERTY(EditAnywhere, Category = "Box", meta = (DataflowInput, EditCondition = "DataType == EMakeBoxDataTypeEnum::Dataflow_MakeBox_DataType_CenterSize", EditConditionHides));
-	FVector Size = FVector(10.0);
-
-	UPROPERTY(meta = (DataflowOutput));
-	FBox Box = FBox(ForceInit);
-
-	FMakeBoxDataflowNode(const Dataflow::FNodeParameters& InParam, FGuid InGuid = FGuid::NewGuid())
-		: FDataflowNode(InParam, InGuid)
-	{
-		RegisterInputConnection(&Min);
-		RegisterInputConnection(&Max);
-		RegisterInputConnection(&Center);
-		RegisterInputConnection(&Size);
-		RegisterOutputConnection(&Box);
-	}
-
-	virtual void Evaluate(Dataflow::FContext& Context, const FDataflowOutput* Out) const override;
-
-};
-
-
-/**
- *
- * Description for this node
- *
- */
-USTRUCT()
-struct FMakeSphereDataflowNode : public FDataflowNode
-{
-	GENERATED_USTRUCT_BODY()
-	DATAFLOW_NODE_DEFINE_INTERNAL(FMakeSphereDataflowNode, "MakeSphere", "Generators|Sphere", "")
-
-public:
-	UPROPERTY(EditAnywhere, Category = "Sphere", meta = (DataflowInput));
-	FVector Center = FVector(0.f);
-
-	UPROPERTY(EditAnywhere, Category = "Sphere", meta = (DataflowInput));
-	float Radius = 10.f;
-
-	UPROPERTY(meta = (DataflowOutput));
-	FSphere Sphere = FSphere(ForceInit);
-
-	FMakeSphereDataflowNode(const Dataflow::FNodeParameters& InParam, FGuid InGuid = FGuid::NewGuid())
-		: FDataflowNode(InParam, InGuid)
-	{
-		RegisterInputConnection(&Center);
-		RegisterInputConnection(&Radius);
-		RegisterOutputConnection(&Sphere);
-	}
-
-	virtual void Evaluate(Dataflow::FContext& Context, const FDataflowOutput* Out) const override;
-
-};
-
-
-/**
- *
- * Description for this node
- *
- */
-USTRUCT()
-struct FMakeLiteralFloatDataflowNode : public FDataflowNode
-{
-	GENERATED_USTRUCT_BODY()
-	DATAFLOW_NODE_DEFINE_INTERNAL(FMakeLiteralFloatDataflowNode, "MakeLiteralFloat", "Math|Float", "")
-
-public:
-	UPROPERTY(EditAnywhere, Category = "Float");
-	float Value = 0.f;
-
-	UPROPERTY(meta = (DataflowOutput))
-	float Float = 0.f;
-
-	FMakeLiteralFloatDataflowNode(const Dataflow::FNodeParameters& InParam, FGuid InGuid = FGuid::NewGuid())
-		: FDataflowNode(InParam, InGuid)
-	{
-		RegisterOutputConnection(&Float);
-	}
-
-	virtual void Evaluate(Dataflow::FContext& Context, const FDataflowOutput* Out) const override;
-
-};
-
-
-/**
- *
- * Description for this node
- *
- */
-USTRUCT()
-struct FMakeLiteralIntDataflowNode : public FDataflowNode
-{
-	GENERATED_USTRUCT_BODY()
-	DATAFLOW_NODE_DEFINE_INTERNAL(FMakeLiteralIntDataflowNode, "MakeLiteralInt", "Math|Int", "")
-
-public:
-	UPROPERTY(EditAnywhere, Category = "Int");
-	int32 Value = 0;
-
-	UPROPERTY(meta = (DataflowOutput, DisplayName = "Int"))
-	int32 Int = 0;
-
-	FMakeLiteralIntDataflowNode(const Dataflow::FNodeParameters& InParam, FGuid InGuid = FGuid::NewGuid())
-		: FDataflowNode(InParam, InGuid)
-	{
-		RegisterOutputConnection(&Int);
-	}
-
-	virtual void Evaluate(Dataflow::FContext& Context, const FDataflowOutput* Out) const override;
-
-};
-
-/**
- *
- * Description for this node
- *
- */
-USTRUCT()
-struct FMakeLiteralBoolDataflowNode : public FDataflowNode
-{
-	GENERATED_USTRUCT_BODY()
-	DATAFLOW_NODE_DEFINE_INTERNAL(FMakeLiteralBoolDataflowNode, "MakeLiteralBool", "Math|Boolean", "")
-
-public:
-	UPROPERTY(EditAnywhere, Category = "Bool");
-	bool Value = false;
-
-	UPROPERTY(meta = (DataflowOutput))
-	bool Bool = false;
-
-	FMakeLiteralBoolDataflowNode(const Dataflow::FNodeParameters& InParam, FGuid InGuid = FGuid::NewGuid())
-		: FDataflowNode(InParam, InGuid)
-	{
-		RegisterOutputConnection(&Bool);
-	}
-
-	virtual void Evaluate(Dataflow::FContext& Context, const FDataflowOutput* Out) const override;
-
-};
-
-
-/**
- *
- * 
- *
- */
-USTRUCT()
-struct FMakeLiteralVectorDataflowNode : public FDataflowNode
-{
-	GENERATED_USTRUCT_BODY()
-	DATAFLOW_NODE_DEFINE_INTERNAL(FMakeLiteralVectorDataflowNode, "MakeLiteralVector", "Math|Vector", "")
-
-public:
-	UPROPERTY(EditAnywhere, Category = "Vector", meta = (DataflowInput));
-	float X = float(0.0);
-
-	UPROPERTY(EditAnywhere, Category = "Vector", meta = (DataflowInput));
-	float Y = float(0.0);
-
-	UPROPERTY(EditAnywhere, Category = "Vector", meta = (DataflowInput));
-	float Z = float(0.0);
-
-	UPROPERTY(meta = (DataflowOutput, DisplayName = "Vector"))
-	FVector Vector = FVector(0.0);
-
-	FMakeLiteralVectorDataflowNode(const Dataflow::FNodeParameters& InParam, FGuid InGuid = FGuid::NewGuid())
-		: FDataflowNode(InParam, InGuid)
-	{
-		RegisterInputConnection(&X);
-		RegisterInputConnection(&Y);
-		RegisterInputConnection(&Z);
-		RegisterOutputConnection(&Vector);
-	}
-
-	virtual void Evaluate(Dataflow::FContext& Context, const FDataflowOutput* Out) const override;
-
-};
-
-
-/**
- *
- * Converts an Int to a String
- *
- */
-USTRUCT()
-struct FIntToStringDataflowNode : public FDataflowNode
-{
-	GENERATED_USTRUCT_BODY()
-	DATAFLOW_NODE_DEFINE_INTERNAL(FIntToStringDataflowNode, "IntToString", "Utilities|String", "")
-
-public:
-	UPROPERTY(EditAnywhere, Category = "Int", meta = (DataflowInput))
-	int32 Int = 0;
-
-	UPROPERTY(meta = (DataflowOutput))
-	FString String = FString("");
-
-	FIntToStringDataflowNode(const Dataflow::FNodeParameters& InParam, FGuid InGuid = FGuid::NewGuid())
-		: FDataflowNode(InParam, InGuid)
-	{
-		RegisterInputConnection(&Int);
-		RegisterOutputConnection(&String);
-	}
-
-	virtual void Evaluate(Dataflow::FContext& Context, const FDataflowOutput* Out) const override;
-
-};
-
-
-/**
- *
- * Converts a Bool to a String in a form of ("true", "false")
- *
- */
-USTRUCT()
-struct FBoolToStringDataflowNode : public FDataflowNode
-{
-	GENERATED_USTRUCT_BODY()
-	DATAFLOW_NODE_DEFINE_INTERNAL(FBoolToStringDataflowNode, "BoolToString", "Utilities|String", "")
-
-public:
-	UPROPERTY(EditAnywhere, Category = "Bool", meta = (DataflowInput))
-	bool Bool = false;
-
-	UPROPERTY(meta = (DataflowOutput))
-	FString String = FString("");
-
-	FBoolToStringDataflowNode(const Dataflow::FNodeParameters& InParam, FGuid InGuid = FGuid::NewGuid())
-		: FDataflowNode(InParam, InGuid)
-	{
-		RegisterInputConnection(&Bool);
-		RegisterOutputConnection(&String);
-	}
-
-	virtual void Evaluate(Dataflow::FContext& Context, const FDataflowOutput* Out) const override;
-
-};
-
-
-/**
- *
- * Expands a Vector into X, Y, Z components
- *
- */
-USTRUCT()
-struct FExpandVectorDataflowNode : public FDataflowNode
-{
-	GENERATED_USTRUCT_BODY()
-	DATAFLOW_NODE_DEFINE_INTERNAL(FExpandVectorDataflowNode, "ExpandVector", "Utilities|Vector", "")
-
-public:
-	UPROPERTY(meta = (DataflowInput))
-	FVector Vector = FVector(0.0);
-
-	UPROPERTY(meta = (DataflowOutput))
-	float X = 0.f;
 
-	UPROPERTY(meta = (DataflowOutput))
-	float Y = 0.f;
-
-	UPROPERTY(meta = (DataflowOutput))
-	float Z = 0.f;
-
-
-	FExpandVectorDataflowNode(const Dataflow::FNodeParameters& InParam, FGuid InGuid = FGuid::NewGuid())
-		: FDataflowNode(InParam, InGuid)
-	{
-		RegisterInputConnection(&Vector);
-		RegisterOutputConnection(&X);
-		RegisterOutputConnection(&Y);
-		RegisterOutputConnection(&Z);
-	}
-
-	virtual void Evaluate(Dataflow::FContext& Context, const FDataflowOutput* Out) const override;
-
-};
 
 
 /**
  *
- * Converts an Int to a Float
+ * Description for this node
  *
  */
-USTRUCT()
-struct FIntToFloatDataflowNode : public FDataflowNode
+USTRUCT(meta = (DataflowGeometryCollection))
+struct FBoundingBoxDataflowNode : public FDataflowNode
 {
 	GENERATED_USTRUCT_BODY()
-	DATAFLOW_NODE_DEFINE_INTERNAL(FIntToFloatDataflowNode, "IntToFloat", "Math|Conversions", "")
+	DATAFLOW_NODE_DEFINE_INTERNAL(FBoundingBoxDataflowNode, "BoundingBox", "Utilities|Box", "")
+	DATAFLOW_NODE_RENDER_TYPE(FName("FBox"), "BoundingBox")
 
 public:
-	UPROPERTY(EditAnywhere, Category = "Int", meta = (DataflowInput))
-	int32 Int = 0;
+	UPROPERTY(meta = (DataflowInput))
+	FManagedArrayCollection Collection;
 
 	UPROPERTY(meta = (DataflowOutput))
-	float Float = 0.f;
+	FBox BoundingBox = FBox(ForceInit);
 
-	FIntToFloatDataflowNode(const Dataflow::FNodeParameters& InParam, FGuid InGuid = FGuid::NewGuid())
+	FBoundingBoxDataflowNode(const Dataflow::FNodeParameters& InParam, FGuid InGuid = FGuid::NewGuid())
 		: FDataflowNode(InParam, InGuid)
 	{
-		RegisterInputConnection(&Int);
-		RegisterOutputConnection(&Float);
+		RegisterInputConnection(&Collection);
+		RegisterOutputConnection(&BoundingBox);
 	}
 
 	virtual void Evaluate(Dataflow::FContext& Context, const FDataflowOutput* Out) const override;
@@ -691,106 +187,43 @@ public:
 
 /**
  *
- * Concatenates two strings together to make a new string
+ * Description for this node
  *
  */
 USTRUCT()
-struct FStringAppendDataflowNode : public FDataflowNode
+struct FExpandBoundingBoxDataflowNode : public FDataflowNode
 {
 	GENERATED_USTRUCT_BODY()
-	DATAFLOW_NODE_DEFINE_INTERNAL(FStringAppendDataflowNode, "StringAppend", "Utilities|String", "")
+	DATAFLOW_NODE_DEFINE_INTERNAL(FExpandBoundingBoxDataflowNode, "ExpandBoundingBox", "Utilities|Box", "")
 
 public:
-	UPROPERTY(EditAnywhere, Category = "String", meta = (DataflowInput))
-	FString String1 = FString("");
-
-	UPROPERTY(EditAnywhere, Category = "String", meta = (DataflowInput))
-	FString String2 = FString("");
+	UPROPERTY(meta = (DataflowInput))
+	FBox BoundingBox = FBox(ForceInit);;
 
 	UPROPERTY(meta = (DataflowOutput))
-	FString String = FString("");
-
-	FStringAppendDataflowNode(const Dataflow::FNodeParameters& InParam, FGuid InGuid = FGuid::NewGuid())
-		: FDataflowNode(InParam, InGuid)
-	{
-		RegisterInputConnection(&String1);
-		RegisterInputConnection(&String2);
-		RegisterOutputConnection(&String);
-	}
-
-	virtual void Evaluate(Dataflow::FContext& Context, const FDataflowOutput* Out) const override;
-
-};
-
-
-/**
- *
- * Generates a random float
- *
- */
-USTRUCT()
-struct FRandomFloatDataflowNode : public FDataflowNode
-{
-	GENERATED_USTRUCT_BODY()
-	DATAFLOW_NODE_DEFINE_INTERNAL(FRandomFloatDataflowNode, "RandomFloat", "Math|Random", "")
-
-public:
-	UPROPERTY(EditAnywhere, Category = "Seed")
-	bool bDeterministic = false;
-
-	UPROPERTY(EditAnywhere, Category = "Seed", meta = (DataflowInput, EditCondition = "Deterministic"))
-	float RandomSeed = 0.f;
+	FVector Min = FVector(0.0);
 
 	UPROPERTY(meta = (DataflowOutput))
-	float Float = 0.f;
-
-	FRandomFloatDataflowNode(const Dataflow::FNodeParameters& InParam, FGuid InGuid = FGuid::NewGuid())
-		: FDataflowNode(InParam, InGuid)
-	{
-		RandomSeed = FMath::FRandRange(-1e5, 1e5);
-		RegisterInputConnection(&RandomSeed);
-		RegisterOutputConnection(&Float);
-	}
-
-	virtual void Evaluate(Dataflow::FContext& Context, const FDataflowOutput* Out) const override;
-
-};
-
-/**
- *
- * Generates a random float between Min and Max
- *
- */
-USTRUCT()
-struct FRandomFloatInRangeDataflowNode : public FDataflowNode
-{
-	GENERATED_USTRUCT_BODY()
-	DATAFLOW_NODE_DEFINE_INTERNAL(FRandomFloatInRangeDataflowNode, "RandomFloatInRange", "Math|Random", "")
-
-public:
-	UPROPERTY(EditAnywhere, Category = "Seed")
-	bool bDeterministic = false;
+	FVector Max = FVector(0.0);
 
-	UPROPERTY(EditAnywhere, Category = "Seed", meta = (DataflowInput, EditCondition = "Deterministic"))
-	float RandomSeed = 0.f;
-	
-	UPROPERTY(EditAnywhere, Category = "Random", meta = (DataflowInput))
-	float Min = 0.f;
+	UPROPERTY(meta = (DataflowOutput))
+	FVector Center = FVector(0.0);
 
-	UPROPERTY(EditAnywhere, Category = "Random", meta = (DataflowInput))
-	float Max = 1.f;
+	UPROPERTY(meta = (DataflowOutput))
+	FVector HalfExtents = FVector(0.0);
 
 	UPROPERTY(meta = (DataflowOutput))
-	float Float = 0.f;
+	float Volume = 0.0;
 
-	FRandomFloatInRangeDataflowNode(const Dataflow::FNodeParameters& InParam, FGuid InGuid = FGuid::NewGuid())
+	FExpandBoundingBoxDataflowNode(const Dataflow::FNodeParameters& InParam, FGuid InGuid = FGuid::NewGuid())
 		: FDataflowNode(InParam, InGuid)
 	{
-		RandomSeed = FMath::FRandRange(-1e5, 1e5);
-		RegisterInputConnection(&RandomSeed);
-		RegisterInputConnection(&Min);
-		RegisterInputConnection(&Max);
-		RegisterOutputConnection(&Float);
+		RegisterInputConnection(&BoundingBox);
+		RegisterOutputConnection(&Min);
+		RegisterOutputConnection(&Max);
+		RegisterOutputConnection(&Center);
+		RegisterOutputConnection(&HalfExtents);
+		RegisterOutputConnection(&Volume);
 	}
 
 	virtual void Evaluate(Dataflow::FContext& Context, const FDataflowOutput* Out) const override;
@@ -798,144 +231,79 @@ public:
 };
 
 
-/**
- *
- * Returns a random vector with length of 1
- *
- */
-USTRUCT()
-struct FRandomUnitVectorDataflowNode : public FDataflowNode
-{
-	GENERATED_USTRUCT_BODY()
-	DATAFLOW_NODE_DEFINE_INTERNAL(FRandomUnitVectorDataflowNode, "RandomUnitVector", "Math|Random", "")
-
-public:
-	UPROPERTY(EditAnywhere, Category = "Seed")
-	bool bDeterministic = false;
-
-	UPROPERTY(EditAnywhere, Category = "Seed", meta = (DataflowInput, EditCondition = "Deterministic"))
-	float RandomSeed = 0.f;
-	
-	UPROPERTY(meta = (DataflowOutput))
-	FVector Vector = FVector(0.0);
-
-	FRandomUnitVectorDataflowNode(const Dataflow::FNodeParameters& InParam, FGuid InGuid = FGuid::NewGuid())
-		: FDataflowNode(InParam, InGuid)
-	{
-		RandomSeed = FMath::FRandRange(-1e5, 1e5);
-		RegisterInputConnection(&RandomSeed);
-		RegisterOutputConnection(&Vector);
-	}
-
-	virtual void Evaluate(Dataflow::FContext& Context, const FDataflowOutput* Out) const override;
-
-};
 
 
 /**
  *
- * Returns a random vector with length of 1, within the specified cone, with uniform random distribution 
+ * Expands a Vector into X, Y, Z components
  *
  */
 USTRUCT()
-struct FRandomUnitVectorInConeDataflowNode : public FDataflowNode
+struct FExpandVectorDataflowNode : public FDataflowNode
 {
 	GENERATED_USTRUCT_BODY()
-	DATAFLOW_NODE_DEFINE_INTERNAL(FRandomUnitVectorInConeDataflowNode, "RandomUnitVectorInCone", "Math|Random", "")
+	DATAFLOW_NODE_DEFINE_INTERNAL(FExpandVectorDataflowNode, "ExpandVector", "Utilities|Vector", "")
 
 public:
-	UPROPERTY(EditAnywhere, Category = "Seed")
-	bool bDeterministic = false;
-
-	UPROPERTY(EditAnywhere, Category = "Seed", meta = (DataflowInput, EditCondition = "Deterministic"))
-	float RandomSeed = 0.f;
-
-	/** The base "center" direction of the cone */
-	UPROPERTY(EditAnywhere, Category = "Random", meta = (DataflowInput))
-	FVector ConeDirection = FVector(0.0, 0.0, 1.0);
-
-	/** The half-angle of the cone (from ConeDir to edge), in degrees */
-	UPROPERTY(EditAnywhere, Category = "Random", meta = (DataflowInput))
-	float ConeHalfAngle = PI / 4.f;
-
-	UPROPERTY(meta = (DataflowOutput))
+	UPROPERTY(meta = (DataflowInput))
 	FVector Vector = FVector(0.0);
 
-	FRandomUnitVectorInConeDataflowNode(const Dataflow::FNodeParameters& InParam, FGuid InGuid = FGuid::NewGuid())
-		: FDataflowNode(InParam, InGuid)
-	{
-		RandomSeed = FMath::FRandRange(-1e5, 1e5);
-		RegisterInputConnection(&RandomSeed);
-		RegisterInputConnection(&ConeDirection);
-		RegisterInputConnection(&ConeHalfAngle);
-		RegisterOutputConnection(&Vector);
-	}
-
-	virtual void Evaluate(Dataflow::FContext& Context, const FDataflowOutput* Out) const override;
-
-};
-
-
-/**
- *
- * Converts radians to degrees
- *
- */
-USTRUCT()
-struct FRadiansToDegreesDataflowNode : public FDataflowNode
-{
-	GENERATED_USTRUCT_BODY()
-	DATAFLOW_NODE_DEFINE_INTERNAL(FRadiansToDegreesDataflowNode, "RadiansToDegrees", "Math|Trigonometry", "")
+	UPROPERTY(meta = (DataflowOutput))
+	float X = 0.f;
 
-public:
-	UPROPERTY(EditAnywhere, Category = "Radians", meta = (DataflowInput))
-	float Radians = 0.f;
+	UPROPERTY(meta = (DataflowOutput))
+	float Y = 0.f;
 
 	UPROPERTY(meta = (DataflowOutput))
-	float Degrees = 0.f;
+	float Z = 0.f;
 
-	FRadiansToDegreesDataflowNode(const Dataflow::FNodeParameters& InParam, FGuid InGuid = FGuid::NewGuid())
+
+	FExpandVectorDataflowNode(const Dataflow::FNodeParameters& InParam, FGuid InGuid = FGuid::NewGuid())
 		: FDataflowNode(InParam, InGuid)
 	{
-		RegisterInputConnection(&Radians);
-		RegisterOutputConnection(&Degrees);
+		RegisterInputConnection(&Vector);
+		RegisterOutputConnection(&X);
+		RegisterOutputConnection(&Y);
+		RegisterOutputConnection(&Z);
 	}
 
 	virtual void Evaluate(Dataflow::FContext& Context, const FDataflowOutput* Out) const override;
 
 };
 
-
 /**
  *
- * Converts degrees to radians
+ * Concatenates two strings together to make a new string
  *
  */
 USTRUCT()
-struct FDegreesToRadiansDataflowNode : public FDataflowNode
+struct FStringAppendDataflowNode : public FDataflowNode
 {
 	GENERATED_USTRUCT_BODY()
-	DATAFLOW_NODE_DEFINE_INTERNAL(FDegreesToRadiansDataflowNode, "DegreesToRadians", "Math|Trigonometry", "")
+	DATAFLOW_NODE_DEFINE_INTERNAL(FStringAppendDataflowNode, "StringAppend", "Utilities|String", "")
 
 public:
-	UPROPERTY(EditAnywhere, Category = "Degrees", meta = (DataflowInput))
-	float Degrees = 0.f;
+	UPROPERTY(EditAnywhere, Category = "String", meta = (DataflowInput))
+	FString String1 = FString("");
+
+	UPROPERTY(EditAnywhere, Category = "String", meta = (DataflowInput))
+	FString String2 = FString("");
 
 	UPROPERTY(meta = (DataflowOutput))
-	float Radians = 0.f;
+	FString String = FString("");
 
-	FDegreesToRadiansDataflowNode(const Dataflow::FNodeParameters& InParam, FGuid InGuid = FGuid::NewGuid())
+	FStringAppendDataflowNode(const Dataflow::FNodeParameters& InParam, FGuid InGuid = FGuid::NewGuid())
 		: FDataflowNode(InParam, InGuid)
 	{
-		RegisterInputConnection(&Degrees);
-		RegisterOutputConnection(&Radians);
+		RegisterInputConnection(&String1);
+		RegisterInputConnection(&String2);
+		RegisterOutputConnection(&String);
 	}
 
 	virtual void Evaluate(Dataflow::FContext& Context, const FDataflowOutput* Out) const override;
 
 };
 
-
 /**
  *
  * Generates a hash value from a string
@@ -1000,107 +368,6 @@ public:
 };
 
 
-UENUM(BlueprintType)
-enum class EFloatToIntFunctionEnum : uint8
-{
-	Dataflow_FloatToInt_Function_Floor UMETA(DisplayName = "Floor()"),
-	Dataflow_FloatToInt_Function_Ceil UMETA(DisplayName = "Ceil()"),
-	Dataflow_FloatToInt_Function_Round UMETA(DisplayName = "Round()"),
-	Dataflow_FloatToInt_Function_Truncate UMETA(DisplayName = "Truncate()"),
-	//~~~
-	//256th entry
-	Dataflow_Max                UMETA(Hidden)
-};
-
-/**
- *
- * Converts a Float to Int using the specified method
- *
- */
-USTRUCT()
-struct FFloatToIntDataflowNode : public FDataflowNode
-{
-	GENERATED_USTRUCT_BODY()
-	DATAFLOW_NODE_DEFINE_INTERNAL(FFloatToIntDataflowNode, "FloatToInt", "Math|Conversions", "")
-
-public:
-	/** Method to convert */
-	UPROPERTY(EditAnywhere, Category = "Float");
-	EFloatToIntFunctionEnum Function = EFloatToIntFunctionEnum::Dataflow_FloatToInt_Function_Round;
-
-	/** Float value to convert */
-	UPROPERTY(EditAnywhere, Category = "Float", meta = (DataflowInput))
-	float Float = 0.f;
-
-	/** Int output */
-	UPROPERTY(meta = (DataflowOutput))
-	int32 Int = 0;
-
-	FFloatToIntDataflowNode(const Dataflow::FNodeParameters& InParam, FGuid InGuid = FGuid::NewGuid())
-		: FDataflowNode(InParam, InGuid)
-	{
-		RegisterInputConnection(&Float);
-		RegisterOutputConnection(&Int);
-	}
-
-	virtual void Evaluate(Dataflow::FContext& Context, const FDataflowOutput* Out) const override;
-
-};
-
-
-UENUM(BlueprintType)
-enum class EMathConstantsEnum : uint8
-{
-	Dataflow_MathConstants_Pi UMETA(DisplayName = "Pi"),
-	Dataflow_MathConstants_HalfPi UMETA(DisplayName = "HalfPi"),
-	Dataflow_MathConstants_TwoPi UMETA(DisplayName = "TwoPi"),
-	Dataflow_MathConstants_FourPi UMETA(DisplayName = "FourPi"),
-	Dataflow_MathConstants_InvPi UMETA(DisplayName = "InvPi"),
-	Dataflow_MathConstants_InvTwoPi UMETA(DisplayName = "InvTwoPi"),
-	Dataflow_MathConstants_Sqrt2 UMETA(DisplayName = "Sqrt2"),
-	Dataflow_MathConstants_InvSqrt2 UMETA(DisplayName = "InvSqrt2"),
-	Dataflow_MathConstants_Sqrt3 UMETA(DisplayName = "Sqrt3"),
-	Dataflow_MathConstants_InvSqrt3 UMETA(DisplayName = "InvSqrt3"),
-	Dataflow_FloatToInt_Function_E UMETA(DisplayName = "e"),
-	Dataflow_FloatToInt_Function_Gamma UMETA(DisplayName = "Gamma"),
-	Dataflow_FloatToInt_Function_GoldenRatio UMETA(DisplayName = "GoldenRatio"),
-	Dataflow_FloatToInt_Function_ZeroTolerance UMETA(DisplayName = "ZeroTolerance"),
-	//~~~
-	//256th entry
-	Dataflow_Max                UMETA(Hidden)
-};
-
-/**
- *
- * Offers a selection of Math constants
- *
- */
-USTRUCT()
-struct FMathConstantsDataflowNode : public FDataflowNode
-{
-	GENERATED_USTRUCT_BODY()
-	DATAFLOW_NODE_DEFINE_INTERNAL(FMathConstantsDataflowNode, "MathConstants", "Math|Utilities", "")
-
-public:
-	/** Math constant to output */
-	UPROPERTY(EditAnywhere, Category = "Constants");
-	EMathConstantsEnum Constant = EMathConstantsEnum::Dataflow_MathConstants_Pi;
-
-	/** Selected Math constant */
-	UPROPERTY(meta = (DataflowOutput))
-	float Float = 0;
-
-	FMathConstantsDataflowNode(const Dataflow::FNodeParameters& InParam, FGuid InGuid = FGuid::NewGuid())
-		: FDataflowNode(InParam, InGuid)
-	{
-		RegisterOutputConnection(&Float);
-	}
-
-	virtual void Evaluate(Dataflow::FContext& Context, const FDataflowOutput* Out) const override;
-
-};
-
-
 /**
  *
  * Returns the specified element from an array
@@ -2135,45 +1402,6 @@ public:
 };
 
 
-/**
- *
- * Description for this node
- *
- */
-USTRUCT()
-struct FMakeTransformDataflowNode : public FDataflowNode
-{
-	GENERATED_USTRUCT_BODY()
-	DATAFLOW_NODE_DEFINE_INTERNAL(FMakeTransformDataflowNode, "MakeTransform", "Generators|Transform", "")
-
-public:
-
-	UPROPERTY(EditAnywhere, Category = "Transform", meta = (DataflowInput, DisplayName = "Translation"));
-	FVector InTranslation = FVector(0, 0, 0);
-
-	UPROPERTY(EditAnywhere, Category = "Transform", meta = (DataflowInput, DisplayName = "Rotation"));
-	FVector InRotation = FVector(0, 0, 0);
-
-	UPROPERTY(EditAnywhere, Category = "Transform", meta = (DataflowInput, DisplayName = "Scale"));
-	FVector InScale = FVector(1, 1, 1);
-
-	UPROPERTY(meta = (DataflowOutput, DisplayName = "Transform"));
-	FTransform OutTransform = FTransform::Identity;
-
-	FMakeTransformDataflowNode(const Dataflow::FNodeParameters& InParam, FGuid InGuid = FGuid::NewGuid())
-		: FDataflowNode(InParam, InGuid)
-	{
-		RegisterInputConnection(&InTranslation);
-		RegisterInputConnection(&InRotation);
-		RegisterInputConnection(&InScale);
-		RegisterOutputConnection(&OutTransform);
-	}
-
-	virtual void Evaluate(Dataflow::FContext& Context, const FDataflowOutput* Out) const override;
-
-};
-
-
 /**
  *
  *
