Fix merge of 24362704
- remove moved class FMakeQuaternionDataflowNode from GeometryCollectionNodes.h

diff --git a/Engine/Plugins/Experimental/GeometryCollectionPlugin/Source/GeometryCollectionNodes/Public/Dataflow/GeometryCollectionNodes.h b/Engine/Plugins/Experimental/GeometryCollectionPlugin/Source/GeometryCollectionNodes/Public/Dataflow/GeometryCollectionNodes.h
--- a/Engine/Plugins/Experimental/GeometryCollectionPlugin/Source/GeometryCollectionNodes/Public/Dataflow/GeometryCollectionNodes.h
+++ b/Engine/Plugins/Experimental/GeometryCollectionPlugin/Source/GeometryCollectionNodes/Public/Dataflow/GeometryCollectionNodes.h
@@ -1589,48 +1589,6 @@ public:
 };
 
 
-/**
- *
- *
- *
- */
-USTRUCT()
-struct FMakeQuaternionDataflowNode : public FDataflowNode
-{
-	GENERATED_USTRUCT_BODY()
-	DATAFLOW_NODE_DEFINE_INTERNAL(FMakeQuaternionDataflowNode, "MakeQuaternion", "Math|Vector", "")
-
-public:
-	UPROPERTY(EditAnywhere, Category = "Quaternion ", meta = (DataflowInput));
-	float X = float(0.0);
-
-	UPROPERTY(EditAnywhere, Category = "Quaternion", meta = (DataflowInput));
-	float Y = float(0.0);
-
-	UPROPERTY(EditAnywhere, Category = "Quaternion", meta = (DataflowInput));
-	float Z = float(0.0);
-
-	UPROPERTY(EditAnywhere, Category = "Quaternion", meta = (DataflowInput));
-	float W = float(0.0);
-
-	UPROPERTY(meta = (DataflowOutput, DisplayName = "Quaternion"))
-	FQuat Quaternion = FQuat(ForceInitToZero);
-
-	FMakeQuaternionDataflowNode(const Dataflow::FNodeParameters& InParam, FGuid InGuid = FGuid::NewGuid())
-		: FDataflowNode(InParam, InGuid)
-	{
-		RegisterInputConnection(&X);
-		RegisterInputConnection(&Y);
-		RegisterInputConnection(&Z);
-		RegisterInputConnection(&W);
-		RegisterOutputConnection(&Quaternion);
-	}
-
-	virtual void Evaluate(Dataflow::FContext& Context, const FDataflowOutput* Out) const override;
-
-};
-
-
 namespace Dataflow
 {
 	void GeometryCollectionEngineNodes();
