FloatArrayToIntArrayDataflowNode - dataflow node that converts from TArray<float> to TArray<int32>.

diff --git a/Engine/Plugins/Experimental/GeometryCollectionPlugin/Source/GeometryCollectionNodes/Private/Dataflow/GeometryCollectionNodes.cpp b/Engine/Plugins/Experimental/GeometryCollectionPlugin/Source/GeometryCollectionNodes/Private/Dataflow/GeometryCollectionNodes.cpp
--- a/Engine/Plugins/Experimental/GeometryCollectionPlugin/Source/GeometryCollectionNodes/Private/Dataflow/GeometryCollectionNodes.cpp
+++ b/Engine/Plugins/Experimental/GeometryCollectionPlugin/Source/GeometryCollectionNodes/Private/Dataflow/GeometryCollectionNodes.cpp
@@ -56,6 +56,7 @@ namespace Dataflow
 		DATAFLOW_NODE_REGISTER_CREATION_FACTORY(FStringAppendDataflowNode);
 		DATAFLOW_NODE_REGISTER_CREATION_FACTORY(FHashStringDataflowNode);
 		DATAFLOW_NODE_REGISTER_CREATION_FACTORY(FHashVectorDataflowNode);
+		DATAFLOW_NODE_REGISTER_CREATION_FACTORY(FFloatArrayToIntArrayDataflowNode);
 		DATAFLOW_NODE_REGISTER_CREATION_FACTORY(FGetArrayElementDataflowNode);
 		DATAFLOW_NODE_REGISTER_CREATION_FACTORY(FGetNumArrayElementsDataflowNode);
 		DATAFLOW_NODE_REGISTER_CREATION_FACTORY(FGetBoundingBoxesFromCollectionDataflowNode);
@@ -244,6 +245,72 @@ void FHashVectorDataflowNode::Evaluate(Dataflow::FContext& Context, const FDataf
 	}
 }
 
+void FFloatArrayToIntArrayDataflowNode::Evaluate(Dataflow::FContext& Context, const FDataflowOutput* Out) const
+{
+	if (Out->IsA<TArray<int32>>(&IntArray))
+	{
+		TArray<float> FloatVal = GetValue<TArray<float>>(Context, &FloatArray);
+		TArray<int32> RetVal; RetVal.SetNumUninitialized(FloatVal.Num());
+		if (Function == EFloatArrayToIntArrayFunctionEnum::Dataflow_FloatToInt_Function_Floor)
+		{
+			for (int32 i = 0; i < FloatVal.Num(); i++)
+			{
+				RetVal[i] = FMath::FloorToInt32(FloatVal[i]);
+			}
+		}
+		else if (Function == EFloatArrayToIntArrayFunctionEnum::Dataflow_FloatToInt_Function_Ceil)
+		{
+			for (int32 i = 0; i < FloatVal.Num(); i++)
+			{
+				RetVal[i] = FMath::CeilToInt32(FloatVal[i]);
+			}
+		}
+		else if (Function == EFloatArrayToIntArrayFunctionEnum::Dataflow_FloatToInt_Function_Round)
+		{
+			for (int32 i = 0; i < FloatVal.Num(); i++)
+			{
+				RetVal[i] = FMath::RoundToInt32(FloatVal[i]);
+			}
+		}
+		else if (Function == EFloatArrayToIntArrayFunctionEnum::Dataflow_FloatToInt_Function_Truncate)
+		{
+			for (int32 i = 0; i < FloatVal.Num(); i++)
+			{
+				RetVal[i] = int32(FMath::TruncToFloat(FloatVal[i]));
+			}
+		}
+
+		else if (Function == EFloatArrayToIntArrayFunctionEnum::Dataflow_FloatToInt_NonZeroToIndex)
+		{
+			int32 RetValIndex = 0;
+			for (int32 i = 0; i < FloatVal.Num(); i++)
+			{
+				if (FloatVal[i] != 0.0)
+				{
+					RetVal[RetValIndex] = i;
+					RetValIndex++;
+				}
+			}
+			RetVal.SetNum(RetValIndex);
+		}
+		else if (Function == EFloatArrayToIntArrayFunctionEnum::Dataflow_FloatToInt_ZeroToIndex)
+		{
+			int32 RetValIndex = 0;
+			for (int32 i = 0; i < FloatVal.Num(); i++)
+			{
+				if (FloatVal[i] == 0.0)
+				{
+					RetVal[RetValIndex] = i;
+					RetValIndex++;
+				}
+			}
+			RetVal.SetNum(RetValIndex);
+		}
+
+		SetValue<TArray<int32>>(Context, RetVal, &IntArray);
+	}
+}
+
 void FGetArrayElementDataflowNode::Evaluate(Dataflow::FContext& Context, const FDataflowOutput* Out) const
 {
 	if (Out->IsA<FVector>(&Point))
diff --git a/Engine/Plugins/Experimental/GeometryCollectionPlugin/Source/GeometryCollectionNodes/Public/Dataflow/GeometryCollectionNodes.h b/Engine/Plugins/Experimental/GeometryCollectionPlugin/Source/GeometryCollectionNodes/Public/Dataflow/GeometryCollectionNodes.h
--- a/Engine/Plugins/Experimental/GeometryCollectionPlugin/Source/GeometryCollectionNodes/Public/Dataflow/GeometryCollectionNodes.h
+++ b/Engine/Plugins/Experimental/GeometryCollectionPlugin/Source/GeometryCollectionNodes/Public/Dataflow/GeometryCollectionNodes.h
@@ -367,6 +367,61 @@ public:
 
 };
 
+UENUM(BlueprintType)
+enum class EFloatArrayToIntArrayFunctionEnum : uint8
+{
+	Dataflow_FloatToInt_Function_Floor UMETA(DisplayName = "Floor()"),
+	Dataflow_FloatToInt_Function_Ceil UMETA(DisplayName = "Ceil()"),
+	Dataflow_FloatToInt_Function_Round UMETA(DisplayName = "Round()"),
+	Dataflow_FloatToInt_Function_Truncate UMETA(DisplayName = "Truncate()"),
+	Dataflow_FloatToInt_NonZeroToIndex UMETA(DisplayName = "Non-zero to Index"),
+	Dataflow_FloatToInt_ZeroToIndex UMETA(DisplayName = "Zero to Index"),
+	//~~~
+	//256th entry
+	Dataflow_Max                UMETA(Hidden)
+};
+
+/**
+ *
+ * Converts a Float array to Int array using the specified method.
+ *
+ */
+USTRUCT()
+struct FFloatArrayToIntArrayDataflowNode : public FDataflowNode
+{
+	GENERATED_USTRUCT_BODY()
+		DATAFLOW_NODE_DEFINE_INTERNAL(FFloatArrayToIntArrayDataflowNode, "FloatArrayToIntArray", "Math|Conversions", "")
+
+public:
+	/** Conversion method:
+	* Floor takes the floor of each input float value - 1.1 turns into 1.
+	* Ceil takes the ceil - 1.1 turns into 2.
+	* Round rounds to the nearest integer - 1.1 turns into 1.
+	* Tuncate trucates like a type cast - 1.1 turns into 1.
+	* Non-zero to Index appends the index of all non-zero values to the output array.
+	* Zero to Index appends the index of all zero values to the output array.
+	*/
+	UPROPERTY(EditAnywhere, Category = "Float");
+	EFloatArrayToIntArrayFunctionEnum Function = EFloatArrayToIntArrayFunctionEnum::Dataflow_FloatToInt_NonZeroToIndex;
+
+	/** Float array value to convert */
+	UPROPERTY(EditAnywhere, Category = "Float", meta = (DataflowInput))
+	TArray<float> FloatArray;
+
+	/** Int array output */
+	UPROPERTY(meta = (DataflowOutput))
+	TArray<int32> IntArray;
+
+	FFloatArrayToIntArrayDataflowNode(const Dataflow::FNodeParameters& InParam, FGuid InGuid = FGuid::NewGuid())
+		: FDataflowNode(InParam, InGuid)
+	{
+		RegisterInputConnection(&FloatArray);
+		RegisterOutputConnection(&IntArray);
+	}
+
+	virtual void Evaluate(Dataflow::FContext& Context, const FDataflowOutput* Out) const override;
+
+};
 
 /**
  *
