Dataflow : Fix static analysis warning that happened to be a bug

diff --git a/Engine/Plugins/Experimental/GeometryCollectionPlugin/Source/GeometryCollectionNodes/Private/Dataflow/GeometryCollectionMathNodes.cpp b/Engine/Plugins/Experimental/GeometryCollectionPlugin/Source/GeometryCollectionNodes/Private/Dataflow/GeometryCollectionMathNodes.cpp
--- a/Engine/Plugins/Experimental/GeometryCollectionPlugin/Source/GeometryCollectionNodes/Private/Dataflow/GeometryCollectionMathNodes.cpp
+++ b/Engine/Plugins/Experimental/GeometryCollectionPlugin/Source/GeometryCollectionNodes/Private/Dataflow/GeometryCollectionMathNodes.cpp
@@ -596,7 +596,7 @@ void FNormalizeToRangeDataflowNode::Evaluate(Dataflow::FContext& Context, const
 
 		if (IRangeMin > InRangeMax)
 		{
-			Swap(InRangeMax, InRangeMax);
+			Swap(IRangeMin, InRangeMax);
 		}
 		Result = (InFloat - IRangeMin) / (InRangeMax - IRangeMin);
 
