ChaosFlesh : Node and UProperty updates

diff --git a/Engine/Plugins/Experimental/ChaosFlesh/Source/ChaosFleshEngine/Private/ChaosFlesh/ChaosDeformableCollisionsComponent.cpp b/Engine/Plugins/Experimental/ChaosFlesh/Source/ChaosFleshEngine/Private/ChaosFlesh/ChaosDeformableCollisionsComponent.cpp
--- a/Engine/Plugins/Experimental/ChaosFlesh/Source/ChaosFleshEngine/Private/ChaosFlesh/ChaosDeformableCollisionsComponent.cpp
+++ b/Engine/Plugins/Experimental/ChaosFlesh/Source/ChaosFleshEngine/Private/ChaosFlesh/ChaosDeformableCollisionsComponent.cpp
@@ -161,10 +161,7 @@ UDeformableCollisionsComponent::NewDeformableData()
 		if (RemovedBody)
 		{
 			RemovedBodiesData.Add({  {RemovedBody, Chaos::Softs::ERigidCollisionShapeType::Unknown, INDEX_NONE} });
-			if (CollisionBodiesPrevXf.Contains(RemovedBody))
-			{
-				CollisionBodiesPrevXf.Remove(RemovedBody);
-			}
+			CollisionBodiesPrevXf.Remove(RemovedBody);
 		}
 	}
 	AddedBodies.Empty();
diff --git a/Engine/Plugins/Experimental/ChaosFlesh/Source/ChaosFleshEngine/Public/ChaosFlesh/FleshComponent.h b/Engine/Plugins/Experimental/ChaosFlesh/Source/ChaosFleshEngine/Public/ChaosFlesh/FleshComponent.h
--- a/Engine/Plugins/Experimental/ChaosFlesh/Source/ChaosFleshEngine/Public/ChaosFlesh/FleshComponent.h
+++ b/Engine/Plugins/Experimental/ChaosFlesh/Source/ChaosFleshEngine/Public/ChaosFlesh/FleshComponent.h
@@ -75,7 +75,7 @@ public:
 	USimulationAsset* GetSimulationCollection() { return SimulationCollection; }
 	const USimulationAsset* GetSimulationCollection() const { return SimulationCollection; }
 
-	UPROPERTY()
+	UPROPERTY(EditAnywhere, Category = "ChaosPhysics")
 	TObjectPtr<UProceduralMeshComponent> Mesh;
 
 	/** @deprecated Use GetSkeletalMeshEmbeddedPositions() instead. */
diff --git a/Engine/Plugins/Experimental/GeometryCollectionPlugin/Source/GeometryCollectionNodes/Private/Dataflow/GeometryCollectionNodes.cpp b/Engine/Plugins/Experimental/GeometryCollectionPlugin/Source/GeometryCollectionNodes/Private/Dataflow/GeometryCollectionNodes.cpp
--- a/Engine/Plugins/Experimental/GeometryCollectionPlugin/Source/GeometryCollectionNodes/Private/Dataflow/GeometryCollectionNodes.cpp
+++ b/Engine/Plugins/Experimental/GeometryCollectionPlugin/Source/GeometryCollectionNodes/Private/Dataflow/GeometryCollectionNodes.cpp
@@ -83,6 +83,8 @@ namespace Dataflow
 		DATAFLOW_NODE_REGISTER_CREATION_FACTORY(FVectorArrayNormalizeDataflowNode);
 		DATAFLOW_NODE_REGISTER_CREATION_FACTORY(FMultiplyTransformDataflowNode);
 		DATAFLOW_NODE_REGISTER_CREATION_FACTORY(FInvertTransformDataflowNode);
+		DATAFLOW_NODE_REGISTER_CREATION_FACTORY(FSelectionToVertexListDataflowNode);
+		
 
 		// GeometryCollection
 		DATAFLOW_NODE_REGISTER_CREATION_FACTORY_NODE_COLORS_BY_CATEGORY("GeometryCollection", FLinearColor(0.55f, 0.45f, 1.0f), CDefaultNodeBodyTintColor);
@@ -1170,6 +1172,11 @@ void FSetVertexColorInCollectionFromVertexSelectionDataflowNode::Evaluate(Datafl
 	}
 }
 
+void FSelectionToVertexListDataflowNode::Evaluate(Dataflow::FContext& Context, const FDataflowOutput* Out) const
+{
+	const FDataflowVertexSelection& InVertexSelection = GetValue<FDataflowVertexSelection>(Context, &VertexSelection);
+	SetValue< TArray<int32> >(Context, InVertexSelection.AsArray(), &VertexList);
+}
 
 void FSetVertexColorInCollectionFromFloatArrayDataflowNode::Evaluate(Dataflow::FContext& Context, const FDataflowOutput* Out) const
 {
diff --git a/Engine/Plugins/Experimental/GeometryCollectionPlugin/Source/GeometryCollectionNodes/Public/Dataflow/GeometryCollectionNodes.h b/Engine/Plugins/Experimental/GeometryCollectionPlugin/Source/GeometryCollectionNodes/Public/Dataflow/GeometryCollectionNodes.h
--- a/Engine/Plugins/Experimental/GeometryCollectionPlugin/Source/GeometryCollectionNodes/Public/Dataflow/GeometryCollectionNodes.h
+++ b/Engine/Plugins/Experimental/GeometryCollectionPlugin/Source/GeometryCollectionNodes/Public/Dataflow/GeometryCollectionNodes.h
@@ -1456,6 +1456,38 @@ public:
 
 };
 
+/**
+ *
+ *
+ *
+ */
+USTRUCT()
+struct FSelectionToVertexListDataflowNode : public FDataflowNode
+{
+	GENERATED_USTRUCT_BODY()
+		DATAFLOW_NODE_DEFINE_INTERNAL(FSelectionToVertexListDataflowNode, "SelectionToVertexList", "Selection|Utility", "")
+
+public:
+
+	/**  */
+	UPROPERTY(meta = (DataflowInput, DisplayName = "VertexSelection", DataflowIntrinsic))
+	FDataflowVertexSelection VertexSelection;
+
+	/**  */
+	UPROPERTY(meta = (DataflowOutput, DisplayName = "VertexList"))
+	TArray<int32> VertexList;
+
+
+	FSelectionToVertexListDataflowNode(const Dataflow::FNodeParameters& InParam, FGuid InGuid = FGuid::NewGuid())
+		: FDataflowNode(InParam, InGuid)
+	{
+		RegisterInputConnection(&VertexSelection);
+		RegisterOutputConnection(&VertexList);
+	}
+
+	virtual void Evaluate(Dataflow::FContext& Context, const FDataflowOutput* Out) const override;
+
+};
 
 /**
  *
