Dataflow node imporvements
- Added a Divide node ( same as SafeDivision )
- Added a ScaleVector node ( Vector * float )

diff --git a/Engine/Plugins/Experimental/GeometryCollectionPlugin/Source/GeometryCollectionNodes/Private/Dataflow/GeometryCollectionMathNodes.cpp b/Engine/Plugins/Experimental/GeometryCollectionPlugin/Source/GeometryCollectionNodes/Private/Dataflow/GeometryCollectionMathNodes.cpp
--- a/Engine/Plugins/Experimental/GeometryCollectionPlugin/Source/GeometryCollectionNodes/Private/Dataflow/GeometryCollectionMathNodes.cpp
+++ b/Engine/Plugins/Experimental/GeometryCollectionPlugin/Source/GeometryCollectionNodes/Private/Dataflow/GeometryCollectionMathNodes.cpp
@@ -17,6 +17,7 @@ namespace Dataflow
 		DATAFLOW_NODE_REGISTER_CREATION_FACTORY(FSubtractDataflowNode);
 		DATAFLOW_NODE_REGISTER_CREATION_FACTORY(FMultiplyDataflowNode);
 		DATAFLOW_NODE_REGISTER_CREATION_FACTORY(FSafeDivideDataflowNode);
+		DATAFLOW_NODE_REGISTER_CREATION_FACTORY(FDivideDataflowNode);
 		DATAFLOW_NODE_REGISTER_CREATION_FACTORY(FDivisionDataflowNode);
 		DATAFLOW_NODE_REGISTER_CREATION_FACTORY(FSafeReciprocalDataflowNode);
 		DATAFLOW_NODE_REGISTER_CREATION_FACTORY(FSquareDataflowNode);
@@ -52,6 +53,7 @@ namespace Dataflow
 		DATAFLOW_NODE_REGISTER_CREATION_FACTORY(FArcTanDataflowNode);
 		DATAFLOW_NODE_REGISTER_CREATION_FACTORY(FArcTan2DataflowNode);
 		DATAFLOW_NODE_REGISTER_CREATION_FACTORY(FNormalizeToRangeDataflowNode);
+		DATAFLOW_NODE_REGISTER_CREATION_FACTORY(FScaleVectorDataflowNode);
 		DATAFLOW_NODE_REGISTER_CREATION_FACTORY(FDotProductDataflowNode);
 		DATAFLOW_NODE_REGISTER_CREATION_FACTORY(FCrossProductDataflowNode);
 		DATAFLOW_NODE_REGISTER_CREATION_FACTORY(FNormalizeDataflowNode);
@@ -604,6 +606,19 @@ void FNormalizeToRangeDataflowNode::Evaluate(Dataflow::FContext& Context, const
 	}
 }
 
+
+void FScaleVectorDataflowNode::Evaluate(Dataflow::FContext& Context, const FDataflowOutput* Out) const
+{
+	if (Out->IsA(&ScaledVector))
+	{
+		const FVector InVector = GetValue(Context, &Vector, Vector);
+		const float InScale = GetValue(Context, &Scale, Scale);
+
+		const FVector Result = InVector * InScale;
+		SetValue(Context, Result, &ScaledVector);
+	}
+}
+
 void FDotProductDataflowNode::Evaluate(Dataflow::FContext& Context, const FDataflowOutput* Out) const
 {
 	if (Out->IsA<float>(&ReturnValue))
diff --git a/Engine/Plugins/Experimental/GeometryCollectionPlugin/Source/GeometryCollectionNodes/Public/Dataflow/GeometryCollectionMathNodes.h b/Engine/Plugins/Experimental/GeometryCollectionPlugin/Source/GeometryCollectionNodes/Public/Dataflow/GeometryCollectionMathNodes.h
--- a/Engine/Plugins/Experimental/GeometryCollectionPlugin/Source/GeometryCollectionNodes/Public/Dataflow/GeometryCollectionMathNodes.h
+++ b/Engine/Plugins/Experimental/GeometryCollectionPlugin/Source/GeometryCollectionNodes/Public/Dataflow/GeometryCollectionMathNodes.h
@@ -135,6 +135,18 @@ public:
 	virtual void Evaluate(Dataflow::FContext& Context, const FDataflowOutput* Out) const override;
 };
 
+USTRUCT(meta = (DataflowGeometryCollection))
+struct FDivideDataflowNode : public FSafeDivideDataflowNode
+{
+	GENERATED_USTRUCT_BODY()
+	DATAFLOW_NODE_DEFINE_INTERNAL(FDivideDataflowNode, "Divide", "Math|Float", "")
+
+public:
+	FDivideDataflowNode(const Dataflow::FNodeParameters& InParam, FGuid InGuid = FGuid::NewGuid())
+		: FSafeDivideDataflowNode(InParam, InGuid)
+	{}
+};
+
 /**
 *
 *
@@ -1225,6 +1237,38 @@ public:
 	virtual void Evaluate(Dataflow::FContext& Context, const FDataflowOutput* Out) const override;
 };
 
+/**
+*
+*
+*
+*/
+USTRUCT(meta = (DataflowGeometryCollection))
+struct FScaleVectorDataflowNode : public FDataflowNode
+{
+	GENERATED_USTRUCT_BODY()
+	DATAFLOW_NODE_DEFINE_INTERNAL(FScaleVectorDataflowNode, "ScaleVector", "Math|Vector", "")
+
+public:
+	UPROPERTY(EditAnywhere, Category = "ScaleVector", meta = (DataflowInput));
+	FVector Vector = FVector(0.f);
+
+	UPROPERTY(EditAnywhere, Category = "ScaleVector", meta = (DataflowInput));
+	float Scale = 1.f;
+
+	UPROPERTY(meta = (DataflowOutput))
+	FVector ScaledVector = FVector(0.f);
+
+	FScaleVectorDataflowNode(const Dataflow::FNodeParameters& InParam, FGuid InGuid = FGuid::NewGuid())
+		: FDataflowNode(InParam, InGuid)
+	{
+		RegisterInputConnection(&Vector);
+		RegisterInputConnection(&Scale);
+		RegisterOutputConnection(&ScaledVector);
+	}
+
+	virtual void Evaluate(Dataflow::FContext& Context, const FDataflowOutput* Out) const override;
+};
+
 /**
 *
 *
