Expose new Grid, K-Means Iterations and Minimum Size options to the Dataflow auto-cluster node

diff --git a/Engine/Plugins/Experimental/GeometryCollectionPlugin/Source/GeometryCollectionNodes/Private/Dataflow/GeometryCollectionClusteringNodes.cpp b/Engine/Plugins/Experimental/GeometryCollectionPlugin/Source/GeometryCollectionNodes/Private/Dataflow/GeometryCollectionClusteringNodes.cpp
--- a/Engine/Plugins/Experimental/GeometryCollectionPlugin/Source/GeometryCollectionNodes/Private/Dataflow/GeometryCollectionClusteringNodes.cpp
+++ b/Engine/Plugins/Experimental/GeometryCollectionPlugin/Source/GeometryCollectionNodes/Private/Dataflow/GeometryCollectionClusteringNodes.cpp
@@ -62,6 +62,11 @@ void FAutoClusterDataflowNode::Evaluate(Dataflow::FContext& Context, const FData
 			bool InAutoCluster = AutoCluster;
 			bool InEnforceSiteParameters = EnforceSiteParameters;
 			bool InAvoidIsolated = AvoidIsolated;
+			int32 InGridX = GetValue<int32>(Context, &ClusterGridWidth);
+			int32 InGridY = GetValue<int32>(Context, &ClusterGridDepth);
+			int32 InGridZ = GetValue<int32>(Context, &ClusterGridHeight);
+			float InMinimumClusterSize = GetValue<float>(Context, &MinimumSize);
+			int32 InKMeansIterations = ClusterSizeMethod == EClusterSizeMethodEnum::Dataflow_ClusterSizeMethod_ByGrid ? DriftIterations : 500;
 
 			TArray<int32> SelectedBones;
 			InTransformSelection.AsArray(SelectedBones);
@@ -74,7 +79,8 @@ void FAutoClusterDataflowNode::Evaluate(Dataflow::FContext& Context, const FData
 				InSiteSize,
 				InAutoCluster,
 				InAvoidIsolated, 
-				InEnforceSiteParameters);
+				InEnforceSiteParameters,
+				InGridX, InGridY, InGridZ, InMinimumClusterSize, InKMeansIterations);
 
 			SetValue<FManagedArrayCollection>(Context, (const FManagedArrayCollection&)(*GeomCollection), &Collection);
 		}
diff --git a/Engine/Plugins/Experimental/GeometryCollectionPlugin/Source/GeometryCollectionNodes/Public/Dataflow/GeometryCollectionClusteringNodes.h b/Engine/Plugins/Experimental/GeometryCollectionPlugin/Source/GeometryCollectionNodes/Public/Dataflow/GeometryCollectionClusteringNodes.h
--- a/Engine/Plugins/Experimental/GeometryCollectionPlugin/Source/GeometryCollectionNodes/Public/Dataflow/GeometryCollectionClusteringNodes.h
+++ b/Engine/Plugins/Experimental/GeometryCollectionPlugin/Source/GeometryCollectionNodes/Public/Dataflow/GeometryCollectionClusteringNodes.h
@@ -18,6 +18,7 @@ enum class EClusterSizeMethodEnum : uint8
 	Dataflow_ClusterSizeMethod_ByNumber UMETA(DisplayName = "By Number"),
 	Dataflow_ClusterSizeMethod_ByFractionOfInput UMETA(DisplayName = "By Fraction Of Input"),
 	Dataflow_ClusterSizeMethod_BySize UMETA(DisplayName = "By Size"),
+	Dataflow_ClusterSizeMethod_ByGrid UMETA(DisplayName = "By Grid"),
 	//~~~
 	//256th entry
 	Dataflow_Max                UMETA(Hidden)
@@ -36,23 +37,43 @@ struct FAutoClusterDataflowNode : public FDataflowNode
 
 public:
 	/** How to choose the size of the clusters to create */
-	UPROPERTY(EditAnywhere, Category = "Cluster Size");
+	UPROPERTY(EditAnywhere, Category = ClusterSize);
 	EClusterSizeMethodEnum ClusterSizeMethod = EClusterSizeMethodEnum::Dataflow_ClusterSizeMethod_ByNumber;
 
 	/** Use a Voronoi diagram with this many Voronoi sites as a guide for deciding cluster boundaries */
-	UPROPERTY(EditAnywhere, Category = "Cluster Size", meta = (DataflowInput, UIMin = 2, UIMax = 5000, EditCondition = "ClusterSizeMethod == EClusterSizeMethodEnum::Dataflow_ClusterSizeMethod_ByNumber"))
+	UPROPERTY(EditAnywhere, Category = ClusterSize, meta = (DataflowInput, UIMin = 2, UIMax = 5000, EditCondition = "ClusterSizeMethod == EClusterSizeMethodEnum::Dataflow_ClusterSizeMethod_ByNumber"))
 	int32 ClusterSites = 10;
 
 	/** Choose the number of Voronoi sites used for clustering as a fraction of the number of child bones to process */
-	UPROPERTY(EditAnywhere, Category = "Cluster Size", meta = (DataflowInput, UIMin = 0.f, UIMax = 0.5f, EditCondition = "ClusterSizeMethod == EClusterSizeMethodEnum::Dataflow_ClusterSizeMethod_ByFractionOfInput"))
+	UPROPERTY(EditAnywhere, Category = ClusterSize, meta = (DataflowInput, UIMin = 0.f, UIMax = 0.5f, EditCondition = "ClusterSizeMethod == EClusterSizeMethodEnum::Dataflow_ClusterSizeMethod_ByFractionOfInput"))
 	float ClusterFraction = 0.25;
 
 	/** Choose the Edge-Size of the cube used to groups bones under a cluster (in cm). */
-	UPROPERTY(EditAnywhere, Category = "ClusterSize", meta = (DataflowInput, DisplayName = "Cluster Size", UIMin = ".01", UIMax = "100", ClampMin = ".0001", ClampMax = "10000", EditCondition = "ClusterSizeMethod == EClusterSizeMethodEnum::Dataflow_ClusterSizeMethod_BySize"))
+	UPROPERTY(EditAnywhere, Category = ClusterSize, meta = (DataflowInput, DisplayName = "Cluster Size", UIMin = ".01", UIMax = "100", ClampMin = ".0001", ClampMax = "10000", EditCondition = "ClusterSizeMethod == EClusterSizeMethodEnum::Dataflow_ClusterSizeMethod_BySize"))
 	float SiteSize = 1;
 
+	/** Choose the number of cluster sites to distribute along the X axis */
+	UPROPERTY(EditAnywhere, Category = ClusterSize, meta = (DataflowInput, ClampMin = "1", UIMax = "20", EditCondition = "ClusterSizeMethod == EClusterSizeMethodEnum::Dataflow_ClusterSizeMethod_ByGrid"))
+	int ClusterGridWidth = 2;
+
+	/** Choose the number of cluster sites to distribute along the Y axis */
+	UPROPERTY(EditAnywhere, Category = ClusterSize, meta = (DataflowInput, ClampMin = "1", UIMax = "20", EditCondition = "ClusterSizeMethod == EClusterSizeMethodEnum::Dataflow_ClusterSizeMethod_ByGrid"))
+	int ClusterGridDepth = 2;
+	
+	/** Choose the number of cluster sites to distribute along the Z axis */
+	UPROPERTY(EditAnywhere, Category = ClusterSize, meta = (DataflowInput, ClampMin = "1", UIMax = "20", EditCondition = "ClusterSizeMethod == EClusterSizeMethodEnum::Dataflow_ClusterSizeMethod_ByGrid"))
+	int ClusterGridHeight = 2;
+
+	/** If a cluster has volume less than this value (in cm) cubed, then the auto-cluster process will attempt to merge it into a neighboring cluster. */
+	UPROPERTY(EditAnywhere, Category = ClusterSize, meta = (DataflowInput, ClampMin = "0"))
+	float MinimumSize = 0;
+
+	/** For a grid distribution, optionally iteratively recenter the grid points to the center of the cluster geometry (technically: applying K-Means iterations) to balance the shape and distribution of the clusters */
+	UPROPERTY(EditAnywhere, Category = ClusterSize, meta = (DisplayName = "Grid Drift Iterations", ClampMin = "1", UIMax = "5", EditCondition = "ClusterSizeMethod == EClusterSizeMethodEnum::Dataflow_ClusterSizeMethod_ByGrid"))
+	int DriftIterations = 1;
+
 	/** If true, bones will only be added to the same cluster if they are physically connected (either directly, or via other bones in the same cluster) */
-	UPROPERTY(EditAnywhere, Category = "Auto Cluster")
+	UPROPERTY(EditAnywhere, Category = AutoCluster, meta = (DisplayName = "Enforce Cluster Connectivity"))
 	bool AutoCluster = true;
 
 	/** If true, make sure the site parameters are matched as close as possible ( bEnforceConnectivity can make the number of site larger than the requested input may produce without it ) */
@@ -60,7 +81,7 @@ public:
 	bool EnforceSiteParameters = true;
 
 	/** If true, prevent the creation of clusters with only a single child. Either by merging into a neighboring cluster, or not creating the cluster. */
-	UPROPERTY(EditAnywhere, Category = "Auto Cluster")
+	UPROPERTY(EditAnywhere, Category = AutoCluster)
 	bool AvoidIsolated = true;
 
 	/** Fractured GeometryCollection to cluster */
@@ -79,6 +100,11 @@ public:
 		RegisterInputConnection(&ClusterSites);
 		RegisterInputConnection(&ClusterFraction);
 		RegisterInputConnection(&SiteSize);
+		RegisterInputConnection(&ClusterGridWidth);
+		RegisterInputConnection(&ClusterGridDepth);
+		RegisterInputConnection(&ClusterGridHeight);
+		RegisterInputConnection(&MinimumSize);
+
 		RegisterOutputConnection(&Collection, &Collection);
 	}
 
