change autocluster dataflow node drift iterations minimum and default value to zero instead of one, to match tool

diff --git a/Engine/Plugins/Experimental/GeometryCollectionPlugin/Source/GeometryCollectionNodes/Public/Dataflow/GeometryCollectionClusteringNodes.h b/Engine/Plugins/Experimental/GeometryCollectionPlugin/Source/GeometryCollectionNodes/Public/Dataflow/GeometryCollectionClusteringNodes.h
--- a/Engine/Plugins/Experimental/GeometryCollectionPlugin/Source/GeometryCollectionNodes/Public/Dataflow/GeometryCollectionClusteringNodes.h
+++ b/Engine/Plugins/Experimental/GeometryCollectionPlugin/Source/GeometryCollectionNodes/Public/Dataflow/GeometryCollectionClusteringNodes.h
@@ -64,14 +64,14 @@ public:
 	UPROPERTY(EditAnywhere, Category = ClusterSize, meta = (DataflowInput, ClampMin = "1", UIMax = "20", EditCondition = "ClusterSizeMethod == EClusterSizeMethodEnum::Dataflow_ClusterSizeMethod_ByGrid"))
 	int ClusterGridHeight = 2;
 
+	/** For a grid distribution, optionally iteratively recenter the grid points to the center of the cluster geometry (technically: applying K-Means iterations) to balance the shape and distribution of the clusters */
+	UPROPERTY(EditAnywhere, Category = ClusterSize, meta = (DisplayName = "Grid Drift Iterations", ClampMin = "0", UIMax = "5", EditCondition = "ClusterSizeMethod == EClusterSizeMethodEnum::Dataflow_ClusterSizeMethod_ByGrid"))
+	int DriftIterations = 0;
+
 	/** If a cluster has volume less than this value (in cm) cubed, then the auto-cluster process will attempt to merge it into a neighboring cluster. */
 	UPROPERTY(EditAnywhere, Category = ClusterSize, meta = (DataflowInput, ClampMin = "0"))
 	float MinimumSize = 0;
 
-	/** For a grid distribution, optionally iteratively recenter the grid points to the center of the cluster geometry (technically: applying K-Means iterations) to balance the shape and distribution of the clusters */
-	UPROPERTY(EditAnywhere, Category = ClusterSize, meta = (DisplayName = "Grid Drift Iterations", ClampMin = "1", UIMax = "5", EditCondition = "ClusterSizeMethod == EClusterSizeMethodEnum::Dataflow_ClusterSizeMethod_ByGrid"))
-	int DriftIterations = 1;
-
 	/** If true, bones will only be added to the same cluster if they are physically connected (either directly, or via other bones in the same cluster) */
 	UPROPERTY(EditAnywhere, Category = AutoCluster, meta = (DisplayName = "Enforce Cluster Connectivity"))
 	bool AutoCluster = true;
