Add support for custom primitive data for geomtetry collection

diff --git a/Engine/Plugins/Experimental/GeometryCollectionPlugin/Source/GeometryCollectionNodes/Private/Dataflow/Nodes/GeometryCollectionAssetNodes.cpp b/Engine/Plugins/Experimental/GeometryCollectionPlugin/Source/GeometryCollectionNodes/Private/Dataflow/Nodes/GeometryCollectionAssetNodes.cpp
--- a/Engine/Plugins/Experimental/GeometryCollectionPlugin/Source/GeometryCollectionNodes/Private/Dataflow/Nodes/GeometryCollectionAssetNodes.cpp
+++ b/Engine/Plugins/Experimental/GeometryCollectionPlugin/Source/GeometryCollectionNodes/Private/Dataflow/Nodes/GeometryCollectionAssetNodes.cpp
@@ -30,6 +30,19 @@ namespace Dataflow
 
 // ===========================================================================================================================
 
+FGeometryCollectionTerminalDataflowNode::FGeometryCollectionTerminalDataflowNode(const Dataflow::FNodeParameters& InParam, FGuid InGuid)
+	: FDataflowTerminalNode(InParam, InGuid)
+{
+	RegisterInputConnection(&Collection);
+	RegisterOutputConnection(&Collection, &Collection);
+	RegisterInputConnection(&Materials);
+	RegisterOutputConnection(&Materials, &Materials);
+	RegisterInputConnection(&InstancedMeshes);
+	RegisterOutputConnection(&InstancedMeshes, &InstancedMeshes);
+}
+
+
+
 void FGeometryCollectionTerminalDataflowNode::SetAssetValue(TObjectPtr<UObject> Asset, Dataflow::FContext& Context) const
 {
 	using FGeometryCollectionPtr = TSharedPtr<FGeometryCollection, ESPMode::ThreadSafe>;
@@ -40,9 +53,9 @@ void FGeometryCollectionTerminalDataflowNode::SetAssetValue(TObjectPtr<UObject>
 	{
 		if (FGeometryCollectionPtr GeometryCollection = CollectionAsset->GetGeometryCollection())
 		{
-			const FManagedArrayCollection& InCollection = GetValue<FManagedArrayCollection>(Context, &Collection);
-			const FMaterialArray& InMaterials = GetValue<FMaterialArray>(Context, &Materials);
-			const FInstancedMeshesArray& InInstancedMeshes = GetValue<FInstancedMeshesArray>(Context, &InstancedMeshes);
+			const FManagedArrayCollection& InCollection = GetValue(Context, &Collection);
+			const FMaterialArray& InMaterials = GetValue(Context, &Materials);
+			const FInstancedMeshesArray& InInstancedMeshes = GetValue(Context, &InstancedMeshes);
 
 			if (InCollection.NumElements(FGeometryCollection::TransformGroup) == 0)
 			{
@@ -69,13 +82,13 @@ void FGeometryCollectionTerminalDataflowNode::Evaluate(Dataflow::FContext& Conte
 	using FMaterialArray = TArray<TObjectPtr<UMaterial>>;
 	using FInstancedMeshesArray = TArray<FGeometryCollectionAutoInstanceMesh>;
 
-	const FManagedArrayCollection& InCollection = GetValue<FManagedArrayCollection>(Context, &Collection);
-	const FMaterialArray& InMaterials = GetValue<FMaterialArray>(Context, &Materials);
-	const FInstancedMeshesArray& InInstancedMeshes = GetValue<FInstancedMeshesArray>(Context, &InstancedMeshes);
+	const FManagedArrayCollection& InCollection = GetValue(Context, &Collection);
+	const FMaterialArray& InMaterials = GetValue(Context, &Materials);
+	const FInstancedMeshesArray& InInstancedMeshes = GetValue(Context, &InstancedMeshes);
 
-	SetValue<FManagedArrayCollection>(Context, InCollection, &Collection);
-	SetValue<FMaterialArray>(Context, InMaterials, &Materials);
-	SetValue<FInstancedMeshesArray>(Context, InInstancedMeshes, &InstancedMeshes);
+	SetValue(Context, InCollection, &Collection);
+	SetValue(Context, InMaterials, &Materials);
+	SetValue(Context, InInstancedMeshes, &InstancedMeshes);
 }
 
 // ===========================================================================================================================
@@ -141,8 +154,7 @@ void FCreateGeometryCollectionFromSourcesDataflowNode::Evaluate(Dataflow::FConte
 {
 	ensure(Out->IsA(&Collection) || Out->IsA(&Materials) || Out->IsA(&InstancedMeshes));
 	
-	using FGeometryCollectionSourceArray = TArray<FGeometryCollectionSource>;
-	const FGeometryCollectionSourceArray& InSources = GetValue<FGeometryCollectionSourceArray>(Context, &Sources);
+	const TArray<FGeometryCollectionSource>& InSources = GetValue(Context, &Sources);
 
 	FGeometryCollection OutCollection;
 	TArray<TObjectPtr<UMaterial>> OutMaterials;
@@ -153,8 +165,9 @@ void FCreateGeometryCollectionFromSourcesDataflowNode::Evaluate(Dataflow::FConte
 	InstancedMeshFacade.DefineSchema();
 
 	constexpr bool bReindexMaterialsInLoop = false;
-	for (const FGeometryCollectionSource& Source: InSources)
+	for (int32 SourceIndex = 0; SourceIndex < InSources.Num(); SourceIndex++)
 	{
+		const FGeometryCollectionSource& Source = InSources[SourceIndex];
 		const int32 NumTransformsBeforeAppending = OutCollection.NumElements(FGeometryCollection::TransformGroup);
 
 		// todo: change AppendGeometryCollectionSource to take a FManagedArrayCollection so we could move the collection when assigning it to the output
@@ -164,10 +177,20 @@ void FCreateGeometryCollectionFromSourcesDataflowNode::Evaluate(Dataflow::FConte
 		FGeometryCollectionAutoInstanceMesh InstancedMesh;
 		InstancedMesh.Mesh = Cast<UStaticMesh>(Source.SourceGeometryObject.TryLoad());
 		InstancedMesh.Materials = Source.SourceMaterial;
-		const int32 InstancedMeshIndex = OutInstancedMeshes.AddUnique(InstancedMesh);
+		InstancedMesh.NumInstances = 0;
+		InstancedMesh.CustomData.Reset();
+
+		int32 InstancedMeshIndex = OutInstancedMeshes.Find(InstancedMesh);
+		if (InstancedMeshIndex == INDEX_NONE)
+		{
+			InstancedMeshIndex = OutInstancedMeshes.Add(InstancedMesh);
+		}
+		OutInstancedMeshes[InstancedMeshIndex].NumInstances++;
+		OutInstancedMeshes[InstancedMeshIndex].CustomData.Append(Source.InstanceCustomData);
 
 		// add the instanced mesh  for all the newly added transforms 
 		const int32 NumTransformsAfterAppending = OutCollection.NumElements(FGeometryCollection::TransformGroup);
+		//ensure((NumTransformsAfterAppending - NumTransformsBeforeAppending) == 1);
 		for (int32 TransformIndex = NumTransformsBeforeAppending; TransformIndex < NumTransformsAfterAppending; TransformIndex++)
 		{
 			InstancedMeshFacade.SetIndex(TransformIndex, InstancedMeshIndex);
@@ -199,5 +222,4 @@ void FCreateGeometryCollectionFromSourcesDataflowNode::Evaluate(Dataflow::FConte
 	SetValue(Context, static_cast<const FManagedArrayCollection&>(OutCollection), &Collection);
 	SetValue(Context, MoveTemp(OutMaterials), &Materials);
 	SetValue(Context, MoveTemp(OutInstancedMeshes), &InstancedMeshes);
-	
 }
diff --git a/Engine/Plugins/Experimental/GeometryCollectionPlugin/Source/GeometryCollectionNodes/Public/Dataflow/Nodes/GeometryCollectionAssetNodes.h b/Engine/Plugins/Experimental/GeometryCollectionPlugin/Source/GeometryCollectionNodes/Public/Dataflow/Nodes/GeometryCollectionAssetNodes.h
--- a/Engine/Plugins/Experimental/GeometryCollectionPlugin/Source/GeometryCollectionNodes/Public/Dataflow/Nodes/GeometryCollectionAssetNodes.h
+++ b/Engine/Plugins/Experimental/GeometryCollectionPlugin/Source/GeometryCollectionNodes/Public/Dataflow/Nodes/GeometryCollectionAssetNodes.h
@@ -31,16 +31,7 @@ public:
 	UPROPERTY(meta = (DataflowInput, DataflowOutput, DataflowPassthrough = "InstancedMeshes", DisplayName = "InstancedMeshes"))
 	TArray<FGeometryCollectionAutoInstanceMesh> InstancedMeshes;
 
-	FGeometryCollectionTerminalDataflowNode(const Dataflow::FNodeParameters& InParam, FGuid InGuid = FGuid::NewGuid())
-		: FDataflowTerminalNode(InParam, InGuid)
-	{
-		RegisterInputConnection(&Collection);
-		RegisterOutputConnection(&Collection, &Collection);
-		RegisterInputConnection(&Materials);
-		RegisterOutputConnection(&Materials, &Materials);
-		RegisterInputConnection(&InstancedMeshes);
-		RegisterOutputConnection(&InstancedMeshes, &InstancedMeshes);
-	}
+	FGeometryCollectionTerminalDataflowNode(const Dataflow::FNodeParameters& InParam, FGuid InGuid = FGuid::NewGuid());
 
 	virtual void Evaluate(Dataflow::FContext& Context) const override;
 	virtual void SetAssetValue(TObjectPtr<UObject> Asset, Dataflow::FContext& Context) const override;
diff --git a/Engine/Source/Runtime/Experimental/GeometryCollectionEngine/Private/GeometryCollection/GeometryCollectionComponent.cpp b/Engine/Source/Runtime/Experimental/GeometryCollectionEngine/Private/GeometryCollection/GeometryCollectionComponent.cpp
--- a/Engine/Source/Runtime/Experimental/GeometryCollectionEngine/Private/GeometryCollection/GeometryCollectionComponent.cpp
+++ b/Engine/Source/Runtime/Experimental/GeometryCollectionEngine/Private/GeometryCollection/GeometryCollectionComponent.cpp
@@ -4380,38 +4380,15 @@ void UGeometryCollectionComponent::AddAutoInstancesToISMPool()
 			{
 				// first count the instance per mesh 
 				const int32 NumMeshes = RestCollection->AutoInstanceMeshes.Num();
-				TArray<int32> InstanceCounts;
-				InstanceCounts.AddZeroed(NumMeshes);
-				int32 TotalInstanceCount = 0;
-
-				const TManagedArray<TSet<int32>>& Children = RestCollection->GetGeometryCollection()->Children;
-
-				const GeometryCollection::Facades::FCollectionInstancedMeshFacade InstancedMeshFacade(*RestCollection->GetGeometryCollection());
-				if (InstancedMeshFacade.IsValid())
-				{
-					for (int32 TransformIndex = 0; TransformIndex < InstancedMeshFacade.GetNumIndices(); TransformIndex++)
-					{
-						const int32 AutoInstanceMeshIndex = InstancedMeshFacade.GetIndex(TransformIndex);
-						if (Children[TransformIndex].Num() == 0)
-						{
-							InstanceCounts[AutoInstanceMeshIndex]++;
-							TotalInstanceCount++;
-						}
-					}
-				}
+				const int32 NumCustomDataFloatsFromComponent = ISMPoolMaterialCustomData.Num();
 
 				// Apply custom instance data from geometry collcection object if custom data array is set up correctly so that it's size matches a multiple of the number of instances.
-				const bool bHasValidAssetCustomInstanceData = TotalInstanceCount && RestCollection->AutoInstanceMaterialCustomData.Num() % TotalInstanceCount == 0;
-				const int32 NumCustomDataFloatsFromAsset = bHasValidAssetCustomInstanceData ? RestCollection->AutoInstanceMaterialCustomData.Num() / TotalInstanceCount : 0;
-				const int32 NumCustomDataFloatsFromComponent = ISMPoolMaterialCustomData.Num();
-				const int32 NumCustomDataFloats = NumCustomDataFloatsFromAsset + NumCustomDataFloatsFromComponent;
 				TArray<float, TInlineAllocator<64>> CustomFloatData;
-				int32 CustomDataFromAssetReadIndex = 0;
 
 				// now register each mesh 
-				for (int32 MeshIndex = 0; MeshIndex < NumMeshes; MeshIndex++)
+				int32 AutoInstanceMeshIndex = 0;
+				for (const FGeometryCollectionAutoInstanceMesh& AutoInstanceMesh : RestCollection->AutoInstanceMeshes)
 				{
-					const FGeometryCollectionAutoInstanceMesh& AutoInstanceMesh = RestCollection->AutoInstanceMeshes[MeshIndex];
 					if (const UStaticMesh* StaticMesh = AutoInstanceMesh.Mesh)
 					{
 						bool bMaterialOverride = false;
@@ -4430,25 +4407,29 @@ void UGeometryCollectionComponent::AddAutoInstancesToISMPool()
 						{
 							StaticMeshInstance.MaterialsOverrides = AutoInstanceMesh.Materials;
 						}
+
+						const bool bHasValidAssetCustomInstanceData = AutoInstanceMesh.NumInstances && (AutoInstanceMesh.CustomData.Num() % AutoInstanceMesh.NumInstances == 0);
+						const int32 NumCustomDataFloatsFromAsset = bHasValidAssetCustomInstanceData ? AutoInstanceMesh.GetNumDataPerInstance() : 0;
+						const int32 NumCustomDataFloats = NumCustomDataFloatsFromAsset + NumCustomDataFloatsFromComponent;
+
 						StaticMeshInstance.NumCustomDataFloats = NumCustomDataFloats;
 
 						if (NumCustomDataFloats > 0)
 						{
-							const int32 InstanceCount = InstanceCounts[MeshIndex];
 							CustomFloatData.Reset();
-							CustomFloatData.Reserve(InstanceCount * NumCustomDataFloats);
+							CustomFloatData.Reserve(AutoInstanceMesh.NumInstances * NumCustomDataFloats);
 							
-							for (int32 InstanceIndex = 0; InstanceIndex < InstanceCount; ++InstanceIndex)
+							for (int32 InstanceIndex = 0; InstanceIndex < AutoInstanceMesh.NumInstances; ++InstanceIndex)
 							{
 								CustomFloatData.Append(ISMPoolMaterialCustomData);
 								// Append each set of values from AutoInstanceMaterialCustomData after the common values from ISMPoolMaterialCustomData
-								CustomFloatData.Append(&RestCollection->AutoInstanceMaterialCustomData[CustomDataFromAssetReadIndex], NumCustomDataFloatsFromAsset);
-								CustomDataFromAssetReadIndex += NumCustomDataFloatsFromAsset;
+								CustomFloatData.Append(&AutoInstanceMesh.CustomData[InstanceIndex], NumCustomDataFloatsFromAsset);
 							}
 						}
 
-						ISMPoolAutoInstancesMeshIds.Add(ISMPoolComp->AddMeshToGroup(ISMPoolMeshGroupIndex, StaticMeshInstance, InstanceCounts[MeshIndex], CustomFloatData, bChaos_GC_UseHierarchicalISMForLeafMeshes));
+						ISMPoolAutoInstancesMeshIds.Add(ISMPoolComp->AddMeshToGroup(ISMPoolMeshGroupIndex, StaticMeshInstance, AutoInstanceMesh.NumInstances, CustomFloatData, bChaos_GC_UseHierarchicalISMForLeafMeshes));
 					}
+					AutoInstanceMeshIndex++;
 				}
 			}
 		}
diff --git a/Engine/Source/Runtime/Experimental/GeometryCollectionEngine/Private/GeometryCollection/GeometryCollectionObject.cpp b/Engine/Source/Runtime/Experimental/GeometryCollectionEngine/Private/GeometryCollection/GeometryCollectionObject.cpp
--- a/Engine/Source/Runtime/Experimental/GeometryCollectionEngine/Private/GeometryCollection/GeometryCollectionObject.cpp
+++ b/Engine/Source/Runtime/Experimental/GeometryCollectionEngine/Private/GeometryCollection/GeometryCollectionObject.cpp
@@ -40,6 +40,7 @@
 #include "Chaos/ChaosArchive.h"
 #include "GeometryCollectionProxyData.h"
 #include "GeometryCollection/Facades/CollectionHierarchyFacade.h"
+#include "GeometryCollection/Facades/CollectionInstancedMeshFacade.h"
 
 #include UE_INLINE_GENERATED_CPP_BY_NAME(GeometryCollectionObject)
 
@@ -412,12 +413,11 @@ void UGeometryCollection::Reset()
 {
 	if (GeometryCollection.IsValid())
 	{
-		Modify();
+		Modify(); 
 		GeometryCollection->Empty();
 		Materials.Empty();
 		EmbeddedGeometryExemplar.Empty();
 		AutoInstanceMeshes.Empty();
-		AutoInstanceMaterialCustomData.Empty();
 		InvalidateCollection();
 	}
 }
@@ -999,6 +999,11 @@ void UGeometryCollection::Serialize(FArchive& Ar)
 		UpdateRootIndex();
 	}
 
+	if (Ar.IsLoading())
+	{
+		FillAutoInstanceMeshesInstancesIfNeeded();
+	}
+
 #if WITH_EDITORONLY_DATA
 	if (bCreateSimulationData)
 	{
@@ -1030,6 +1035,35 @@ void UGeometryCollection::SetEnableNanite(bool bValue)
 	}
 }
 
+void UGeometryCollection::FillAutoInstanceMeshesInstancesIfNeeded()
+{
+	// make sure the instanced meshes have their instance count properly set 
+	if (GeometryCollection && AutoInstanceMeshes.Num() > 0 && AutoInstanceMeshes[0].NumInstances == 0)
+	{
+		//  make sure to rest all of it first 
+		for (FGeometryCollectionAutoInstanceMesh& AutoInstanceMesh : AutoInstanceMeshes)
+		{
+			AutoInstanceMesh.NumInstances = 0;
+		}
+
+		const GeometryCollection::Facades::FCollectionInstancedMeshFacade InstancedMeshFacade(*GeometryCollection);
+
+		const int32 NumTransforms = GeometryCollection->Children.Num();
+		for (int32 TransformIndex = 0; TransformIndex < NumTransforms; TransformIndex++)
+		{
+			// only applies to leaves nodes
+			if (GeometryCollection->Children[TransformIndex].Num() == 0)
+			{
+				const int32 AutoInstanceMeshIndex = InstancedMeshFacade.GetIndex(TransformIndex);
+				if (AutoInstanceMeshes.IsValidIndex(AutoInstanceMeshIndex))
+				{
+					AutoInstanceMeshes[AutoInstanceMeshIndex].NumInstances++;
+				}
+			}
+		}
+	}
+}
+
 #if WITH_EDITOR
 
 void UGeometryCollection::CreateSimulationDataImp(bool bCopyFromDDC)
@@ -1353,6 +1387,11 @@ void UGeometryCollection::RemoveExemplars(const TArray<int32>& SortedRemovalIndi
 	}
 }
 
+int32 FGeometryCollectionAutoInstanceMesh::GetNumDataPerInstance() const
+{
+	return NumInstances? (CustomData.Num() / NumInstances): 0;
+}
+
 bool FGeometryCollectionAutoInstanceMesh::operator ==(const FGeometryCollectionAutoInstanceMesh& Other) const
 {
 	return (Mesh == Other.Mesh) && (Materials == Other.Materials);
@@ -1365,9 +1404,12 @@ const FGeometryCollectionAutoInstanceMesh& UGeometryCollection::GetAutoInstanceM
 }
 
 /**  find or add a auto instance mesh from another one and return its index */
-int32 UGeometryCollection::FindOrAddAutoInstanceMesh(const FGeometryCollectionAutoInstanceMesh& AutoInstanecMesh)
+int32 UGeometryCollection::FindOrAddAutoInstanceMesh(const FGeometryCollectionAutoInstanceMesh& AutoInstanceMesh)
 {
-	return AutoInstanceMeshes.AddUnique(AutoInstanecMesh);
+	int32 AutoInstanceMeshIndex = AutoInstanceMeshes.AddUnique(AutoInstanceMesh);
+	FGeometryCollectionAutoInstanceMesh& Instance = AutoInstanceMeshes[AutoInstanceMeshIndex];
+	Instance.NumInstances++;
+	return AutoInstanceMeshIndex;
 }
 
 int32 UGeometryCollection::FindOrAddAutoInstanceMesh(const UStaticMesh* StaticMesh, const TArray<UMaterialInterface*>& MeshMaterials)
@@ -1375,7 +1417,8 @@ int32 UGeometryCollection::FindOrAddAutoInstanceMesh(const UStaticMesh* StaticMe
 	FGeometryCollectionAutoInstanceMesh NewMesh;
 	NewMesh.Mesh = StaticMesh;
 	NewMesh.Materials = MeshMaterials;
-	return AutoInstanceMeshes.AddUnique(NewMesh);
+
+	return FindOrAddAutoInstanceMesh(NewMesh);
 }
 
 FGuid UGeometryCollection::GetIdGuid() const
diff --git a/Engine/Source/Runtime/Experimental/GeometryCollectionEngine/Public/GeometryCollection/GeometryCollectionObject.h b/Engine/Source/Runtime/Experimental/GeometryCollectionEngine/Public/GeometryCollection/GeometryCollectionObject.h
--- a/Engine/Source/Runtime/Experimental/GeometryCollectionEngine/Public/GeometryCollection/GeometryCollectionObject.h
+++ b/Engine/Source/Runtime/Experimental/GeometryCollectionEngine/Public/GeometryCollection/GeometryCollectionObject.h
@@ -44,6 +44,9 @@ struct GEOMETRYCOLLECTIONENGINE_API FGeometryCollectionSource
 	UPROPERTY(BlueprintReadWrite, EditAnywhere, Category = "GeometrySource")
 	TArray<TObjectPtr<UMaterialInterface>> SourceMaterial;
 
+	UPROPERTY(BlueprintReadWrite, EditAnywhere, Category = "GeometrySource|Instance")
+	TArray<float> InstanceCustomData;
+
 	//~ Note: bAddInternalMaterials defaults to true so a 'Reset' of a geometry collection that was created before this member was added will have consistent behavior. New geometry collections should always set bAddInternalMaterials to false.
 	/** (Legacy) Whether source materials will be duplicated to create new slots for internal materials, or existing odd materials will be considered internal. (For non-Geometry Collection inputs only.) */
 	UPROPERTY(BlueprintReadWrite, EditAnywhere, Category = "GeometrySource", DisplayName = "(Legacy) Add Internal Materials")
@@ -76,6 +79,14 @@ struct GEOMETRYCOLLECTIONENGINE_API FGeometryCollectionAutoInstanceMesh
 	UPROPERTY(BlueprintReadWrite, EditAnywhere, Category = "AutoInstance")
 	TArray<TObjectPtr<UMaterialInterface>> Materials;
 
+	UPROPERTY(VisibleAnywhere, Category = "AutoInstance")
+	int32 NumInstances = 0;
+
+	UPROPERTY(VisibleAnywhere, Category = "AutoInstance")
+	TArray<float> CustomData;
+
+	int32 GetNumDataPerInstance() const;
+
 	bool operator ==(const FGeometryCollectionAutoInstanceMesh& Other) const;
 };
 
@@ -425,7 +436,7 @@ public:
 	const FGeometryCollectionAutoInstanceMesh& GetAutoInstanceMesh(int32 AutoInstanceMeshIndex) const;
 
 	/**  find or add a auto instance mesh from another one and return its index */
-	int32 FindOrAddAutoInstanceMesh(const FGeometryCollectionAutoInstanceMesh& AutoInstanecMesh);
+	int32 FindOrAddAutoInstanceMesh(const FGeometryCollectionAutoInstanceMesh& AutoInstanceMesh);
 
 	/** find or add a auto instance mesh from a mesh and alist of material and return its index */
 	int32 FindOrAddAutoInstanceMesh(const UStaticMesh* StaticMesh, const TArray<UMaterialInterface*>& Materials);
@@ -552,10 +563,6 @@ public:
 	UPROPERTY(EditAnywhere, Category = "Rendering")
 	TArray<FGeometryCollectionAutoInstanceMesh> AutoInstanceMeshes;
 
-	/** Array of material custom data for all static meshes in AutoInstanceMeshes. */
-	UPROPERTY(EditAnywhere, Category = "Rendering")
-	TArray<float> AutoInstanceMaterialCustomData;
-
 	UFUNCTION(BlueprintCallable, Category = "Nanite")
 	void SetEnableNanite(bool bValue);
 
@@ -761,9 +768,12 @@ private:
 	*/
 	void ValidateSizeSpecificDataDefaults();
 
-	// update cachedroot index using the current hierarchy setup
+	// update cached root index using the current hierarchy setup
 	void UpdateRootIndex();
 
+	// fill instanced mesh instance count from geometry collection data if not done yet 
+	void FillAutoInstanceMeshesInstancesIfNeeded();
+
 private:
 	/** Guid created on construction of this collection. It should be used to uniquely identify this collection */
 	UPROPERTY()
