Dataflow : add few new type conversion nodes

diff --git a/Engine/Plugins/Experimental/GeometryCollectionPlugin/Source/GeometryCollectionNodes/Private/Dataflow/GeometryCollectionConversionNodes.cpp b/Engine/Plugins/Experimental/GeometryCollectionPlugin/Source/GeometryCollectionNodes/Private/Dataflow/GeometryCollectionConversionNodes.cpp
--- a/Engine/Plugins/Experimental/GeometryCollectionPlugin/Source/GeometryCollectionNodes/Private/Dataflow/GeometryCollectionConversionNodes.cpp
+++ b/Engine/Plugins/Experimental/GeometryCollectionPlugin/Source/GeometryCollectionNodes/Private/Dataflow/GeometryCollectionConversionNodes.cpp
@@ -16,6 +16,8 @@ namespace Dataflow
 		DATAFLOW_NODE_REGISTER_CREATION_FACTORY(FIntToStringDataflowNode);
 		DATAFLOW_NODE_REGISTER_CREATION_FACTORY(FBoolToStringDataflowNode);
 		DATAFLOW_NODE_REGISTER_CREATION_FACTORY(FIntToFloatDataflowNode);
+		DATAFLOW_NODE_REGISTER_CREATION_FACTORY(FFloatToDoubleDataflowNode);
+		DATAFLOW_NODE_REGISTER_CREATION_FACTORY(FIntToDoubleDataflowNode)
 		DATAFLOW_NODE_REGISTER_CREATION_FACTORY(FFloatToIntDataflowNode);
 		DATAFLOW_NODE_REGISTER_CREATION_FACTORY(FIntToBoolDataflowNode);
 		DATAFLOW_NODE_REGISTER_CREATION_FACTORY(FBoolToIntDataflowNode);
@@ -64,10 +66,28 @@ void FBoolToStringDataflowNode::Evaluate(Dataflow::FContext& Context, const FDat
 
 void FIntToFloatDataflowNode::Evaluate(Dataflow::FContext& Context, const FDataflowOutput* Out) const
 {
-	if (Out->IsA<float>(&Float))
+	if (Out->IsA(&Float))
 	{
 		float Value = float(GetValue<int32>(Context, &Int));
-		SetValue<float>(Context, Value, &Float);
+		SetValue(Context, Value, &Float);
+	}
+}
+
+void FIntToDoubleDataflowNode::Evaluate(Dataflow::FContext& Context, const FDataflowOutput* Out) const
+{
+	if (Out->IsA(&Double))
+	{
+		double Value = double(GetValue<int32>(Context, &Int));
+		SetValue(Context, Value, &Double);
+	}
+}
+
+void FFloatToDoubleDataflowNode::Evaluate(Dataflow::FContext& Context, const FDataflowOutput* Out) const
+{
+	if (Out->IsA(&Double))
+	{
+		double Value = double(GetValue<float>(Context, &Float));
+		SetValue(Context, Value, &Double);
 	}
 }
 
diff --git a/Engine/Plugins/Experimental/GeometryCollectionPlugin/Source/GeometryCollectionNodes/Public/Dataflow/GeometryCollectionConversionNodes.h b/Engine/Plugins/Experimental/GeometryCollectionPlugin/Source/GeometryCollectionNodes/Public/Dataflow/GeometryCollectionConversionNodes.h
--- a/Engine/Plugins/Experimental/GeometryCollectionPlugin/Source/GeometryCollectionNodes/Public/Dataflow/GeometryCollectionConversionNodes.h
+++ b/Engine/Plugins/Experimental/GeometryCollectionPlugin/Source/GeometryCollectionNodes/Public/Dataflow/GeometryCollectionConversionNodes.h
@@ -149,9 +149,65 @@ public:
 	}
 
 	virtual void Evaluate(Dataflow::FContext& Context, const FDataflowOutput* Out) const override;
+};
+
+/**
+ *
+ * Converts an Int to a Double
+ *
+ */
+USTRUCT(meta = (DataflowGeometryCollection))
+struct FIntToDoubleDataflowNode : public FDataflowNode
+{
+	GENERATED_USTRUCT_BODY()
+	DATAFLOW_NODE_DEFINE_INTERNAL(FIntToDoubleDataflowNode, "IntToDouble", "Math|Conversions", "")
+
+public:
+	UPROPERTY(EditAnywhere, Category = "Int", meta = (DataflowInput))
+	int32 Int = 0;
+
+	UPROPERTY(meta = (DataflowOutput))
+	double Double = 0.0;
+
+	FIntToDoubleDataflowNode(const Dataflow::FNodeParameters& InParam, FGuid InGuid = FGuid::NewGuid())
+		: FDataflowNode(InParam, InGuid)
+	{
+		RegisterInputConnection(&Int);
+		RegisterOutputConnection(&Double);
+	}
+
+	virtual void Evaluate(Dataflow::FContext& Context, const FDataflowOutput* Out) const override;
+};
+
+/**
+ *
+ * Converts an Float to a Double
+ *
+ */
+USTRUCT(meta = (DataflowGeometryCollection))
+struct FFloatToDoubleDataflowNode : public FDataflowNode
+{
+	GENERATED_USTRUCT_BODY()
+	DATAFLOW_NODE_DEFINE_INTERNAL(FFloatToDoubleDataflowNode, "FloatToDouble", "Math|Conversions", "")
+
+public:
+	UPROPERTY(EditAnywhere, Category = "Int", meta = (DataflowInput))
+	float Float = 0.f;
+
+	UPROPERTY(meta = (DataflowOutput))
+	double Double = 0.0;
 
+	FFloatToDoubleDataflowNode(const Dataflow::FNodeParameters& InParam, FGuid InGuid = FGuid::NewGuid())
+		: FDataflowNode(InParam, InGuid)
+	{
+		RegisterInputConnection(&Float);
+		RegisterOutputConnection(&Double);
+	}
+
+	virtual void Evaluate(Dataflow::FContext& Context, const FDataflowOutput* Out) const override;
 };
 
+
 UENUM(BlueprintType)
 enum class EFloatToIntFunctionEnum : uint8
 {
