minor fixes to contact area threshold UI in dataflow proximity node

diff --git a/Engine/Plugins/Experimental/GeometryCollectionPlugin/Source/GeometryCollectionNodes/Public/Dataflow/GeometryCollectionNodes.h b/Engine/Plugins/Experimental/GeometryCollectionPlugin/Source/GeometryCollectionNodes/Public/Dataflow/GeometryCollectionNodes.h
--- a/Engine/Plugins/Experimental/GeometryCollectionPlugin/Source/GeometryCollectionNodes/Public/Dataflow/GeometryCollectionNodes.h
+++ b/Engine/Plugins/Experimental/GeometryCollectionPlugin/Source/GeometryCollectionNodes/Public/Dataflow/GeometryCollectionNodes.h
@@ -994,11 +994,11 @@ UENUM(BlueprintType)
 enum class EProximityContactFilteringMethodEnum : uint8
 {
 	/** Rejects proximity if the bounding boxes do not overlap by more than Contact Threshold centimeters in any major axis direction (or at least half the max possible). This can filter out corner connections of box-like shapes. */
-	Dataflow_ProximityContactFitlteringMethod_ProjectedBoundsOverlap UMETA(DisplayName = "Projected Bounds Overlap"),
+	Dataflow_ProximityContactFilteringMethod_ProjectedBoundsOverlap UMETA(DisplayName = "Projected Bounds Overlap"),
 	/** Rejects proximity if the intersection of convex hulls (allowing for optional offset) follows a sharp, thin region which is not wider than Contact Threshold centimeters (or at least half the max possible). */
-	Dataflow_ProximityContactFitlteringMethod_ConvexHullSharp UMETA(DisplayName = "Convex Hull Sharp Contact"),
+	Dataflow_ProximityContactFilteringMethod_ConvexHullSharp UMETA(DisplayName = "Convex Hull Sharp Contact"),
 	/** Rejects proximity if the surface area of the intersection of convex hulls (allowing for optional offset) is smaller than Contact Threshold squared (or at least half the max possible). */
-	Dataflow_ProximityContactFitlteringMethod_ConvexHullArea UMETA(DisplayName = "Convex Hull Area Contact"),
+	Dataflow_ProximityContactFilteringMethod_ConvexHullArea UMETA(DisplayName = "Convex Hull Area Contact"),
 	//~~~
 	//256th entry
 	Dataflow_Max                UMETA(Hidden)
@@ -1021,7 +1021,8 @@ public:
 	EProximityMethodEnum ProximityMethod = EProximityMethodEnum::Dataflow_ProximityMethod_Precise;
 
 	/** If hull-based proximity detection is enabled, amount to expand hulls when searching for overlapping neighbors */
-	UPROPERTY(EditAnywhere, Category = "Proximity", meta = (ClampMin = "0", EditCondition = "ProximityMethod == EProximityMethodEnum::Dataflow_ProximityMethod_ConvexHull"))
+	UPROPERTY(EditAnywhere, Category = "Proximity", meta = (ClampMin = "0", 
+		EditCondition = "ProximityMethod == EProximityMethodEnum::Dataflow_ProximityMethod_ConvexHull || FilterContactMethod == EProximityContactFilteringMethodEnum::Dataflow_ProximityContactFilteringMethod_ConvexHullSharp || FilterContactMethod == EProximityContactFilteringMethodEnum::Dataflow_ProximityContactFilteringMethod_ConvexHullArea"))
 	float DistanceThreshold = 1;
 
 	// If greater than zero, proximity will be additionally filtered by a 'contact' threshold, in cm, to exclude grazing / corner proximity
@@ -1029,7 +1030,8 @@ public:
 	float ContactThreshold = 0;
 
 	/** How to use the Contact Threshold (if > 0) to filter out unwanted small or corner contacts from the proximity graph. If contact threshold is zero, no filtering is applied. */
-	EProximityContactFilteringMethodEnum FilterContactMethod = EProximityContactFilteringMethodEnum::Dataflow_ProximityContactFitlteringMethod_ProjectedBoundsOverlap;
+	UPROPERTY(EditAnywhere, Category = "Proximity");
+	EProximityContactFilteringMethodEnum FilterContactMethod = EProximityContactFilteringMethodEnum::Dataflow_ProximityContactFilteringMethod_ProjectedBoundsOverlap;
 
 	/** Whether to automatically transform the proximity graph into a connection graph to be used for simulation */
 	UPROPERTY(EditAnywhere, Category = "Proximity")
