add 'uncluster' dataflow node (corresponding to uncluster tool action in fracture mode)

diff --git a/Engine/Plugins/Experimental/GeometryCollectionPlugin/Source/GeometryCollectionNodes/Private/Dataflow/GeometryCollectionClusteringNodes.cpp b/Engine/Plugins/Experimental/GeometryCollectionPlugin/Source/GeometryCollectionNodes/Private/Dataflow/GeometryCollectionClusteringNodes.cpp
--- a/Engine/Plugins/Experimental/GeometryCollectionPlugin/Source/GeometryCollectionNodes/Private/Dataflow/GeometryCollectionClusteringNodes.cpp
+++ b/Engine/Plugins/Experimental/GeometryCollectionPlugin/Source/GeometryCollectionNodes/Private/Dataflow/GeometryCollectionClusteringNodes.cpp
@@ -22,6 +22,8 @@
 #include "GeometryCollection/GeometryCollectionAlgo.h"
 #include "GeometryCollection/GeometryCollectionClusteringUtility.h"
 #include "GeometryCollection/GeometryCollectionConvexUtility.h"
+#include "GeometryCollection/Facades/CollectionTransformSelectionFacade.h"
+#include "GeometryCollection/Facades/CollectionHierarchyFacade.h"
 #include "Voronoi/Voronoi.h"
 #include "PlanarCut.h"
 #include "GeometryCollection/GeometryCollectionProximityUtility.h"
@@ -39,6 +41,7 @@ namespace Dataflow
 
 		DATAFLOW_NODE_REGISTER_CREATION_FACTORY(FAutoClusterDataflowNode);
 		DATAFLOW_NODE_REGISTER_CREATION_FACTORY(FClusterFlattenDataflowNode);
+		DATAFLOW_NODE_REGISTER_CREATION_FACTORY(FClusterUnclusterDataflowNode);
 
 		// GeometryCollection|Cluster
 		DATAFLOW_NODE_REGISTER_CREATION_FACTORY_NODE_COLORS_BY_CATEGORY("GeometryCollection|Cluster", FLinearColor(.25f, 0.45f, 0.8f), CDefaultNodeBodyTintColor);
@@ -114,3 +117,32 @@ void FClusterFlattenDataflowNode::Evaluate(Dataflow::FContext& Context, const FD
 		}
 	}
 }
+
+
+void FClusterUnclusterDataflowNode::Evaluate(Dataflow::FContext& Context, const FDataflowOutput* Out) const
+{
+	if (Out->IsA<FManagedArrayCollection>(&Collection))
+	{
+		const FManagedArrayCollection& InCollection = GetValue<FManagedArrayCollection>(Context, &Collection);
+		const FDataflowTransformSelection& InTransformSelection = GetValue<FDataflowTransformSelection>(Context, &TransformSelection);
+		if (TUniquePtr<FGeometryCollection> GeomCollection = TUniquePtr<FGeometryCollection>(InCollection.NewCopy<FGeometryCollection>()))
+		{
+			Chaos::Facades::FCollectionHierarchyFacade HierarchyFacade(*GeomCollection);
+			HierarchyFacade.GenerateLevelAttribute();
+
+			TArray<int32> Selection = InTransformSelection.AsArray();
+			GeometryCollection::Facades::FCollectionTransformSelectionFacade SelectionFacade(*GeomCollection);
+			SelectionFacade.ConvertSelectionToClusterNodes(Selection, false);
+			SelectionFacade.RemoveRootNodes(Selection);
+			if (Selection.IsEmpty())
+			{
+				return;
+			}
+			FGeometryCollectionClusteringUtility::CollapseHierarchyOneLevel(GeomCollection.Get(), Selection);
+			FGeometryCollectionClusteringUtility::RemoveDanglingClusters(GeomCollection.Get());
+
+			HierarchyFacade.GenerateLevelAttribute();
+			SetValue<FManagedArrayCollection>(Context, (const FManagedArrayCollection&)(*GeomCollection), &Collection);
+		}
+	}
+}
diff --git a/Engine/Plugins/Experimental/GeometryCollectionPlugin/Source/GeometryCollectionNodes/Public/Dataflow/GeometryCollectionClusteringNodes.h b/Engine/Plugins/Experimental/GeometryCollectionPlugin/Source/GeometryCollectionNodes/Public/Dataflow/GeometryCollectionClusteringNodes.h
--- a/Engine/Plugins/Experimental/GeometryCollectionPlugin/Source/GeometryCollectionNodes/Public/Dataflow/GeometryCollectionClusteringNodes.h
+++ b/Engine/Plugins/Experimental/GeometryCollectionPlugin/Source/GeometryCollectionNodes/Public/Dataflow/GeometryCollectionClusteringNodes.h
@@ -142,6 +142,37 @@ public:
 
 };
 
+/**
+ * Uncluster selected nodes
+ */
+USTRUCT(meta = (DataflowGeometryCollection))
+struct FClusterUnclusterDataflowNode : public FDataflowNode
+{
+	GENERATED_USTRUCT_BODY()
+	DATAFLOW_NODE_DEFINE_INTERNAL(FClusterUnclusterDataflowNode, "Uncluster", "GeometryCollection|Cluster", "")
+
+public:
+
+	/** Fractured GeometryCollection to uncluster */
+	UPROPERTY(meta = (DataflowInput, DataflowOutput, DataflowPassthrough = "Collection", DataflowIntrinsic))
+	FManagedArrayCollection Collection;
+
+	/** Bone selection */
+	UPROPERTY(meta = (DataflowInput, DisplayName = "TransformSelection"))
+	FDataflowTransformSelection TransformSelection;
+
+	FClusterUnclusterDataflowNode(const Dataflow::FNodeParameters& InParam, FGuid InGuid = FGuid::NewGuid())
+		: FDataflowNode(InParam, InGuid)
+	{
+		RegisterInputConnection(&Collection);
+		RegisterInputConnection(&TransformSelection);
+		RegisterOutputConnection(&Collection, &Collection);
+	}
+
+	virtual void Evaluate(Dataflow::FContext& Context, const FDataflowOutput* Out) const override;
+
+};
+
 
 namespace Dataflow
 {
diff --git a/Engine/Source/Runtime/Experimental/Chaos/Private/GeometryCollection/Facades/CollectionTransformSelectionFacade.cpp b/Engine/Source/Runtime/Experimental/Chaos/Private/GeometryCollection/Facades/CollectionTransformSelectionFacade.cpp
--- a/Engine/Source/Runtime/Experimental/Chaos/Private/GeometryCollection/Facades/CollectionTransformSelectionFacade.cpp
+++ b/Engine/Source/Runtime/Experimental/Chaos/Private/GeometryCollection/Facades/CollectionTransformSelectionFacade.cpp
@@ -188,6 +188,48 @@ namespace GeometryCollection::Facades
 		InOutSelection = RigidSelection;
 	}
 
+	void FCollectionTransformSelectionFacade::ConvertSelectionToClusterNodes(TArray<int32>& InOutSelection, bool bLeaveRigidRoots) const
+	{
+		const TManagedArray<int32>& Parents = ParentAttribute.Get();
+		const TManagedArray<int32>& SimulationType = SimulationTypeAttribute.Get();
+
+		Sanitize(InOutSelection);
+
+		TSet<int32> ClusterNodes;
+		for (int32 Index : InOutSelection)
+		{
+			int32 SimType = SimulationType[Index];
+			if (SimType == FGeometryCollection::ESimulationTypes::FST_Clustered)
+			{
+				ClusterNodes.Add(Index);
+			}
+			else // if Index is not a cluster, select the cluster containing Index
+			{
+				int32 Parent = Parents[Index];
+				int32 CouldBeRoot = Index;
+				if (SimType == FGeometryCollection::ESimulationTypes::FST_None && Parent != -1)
+				{
+					CouldBeRoot = Parent;
+					Parent = Parents[Parent];
+				}
+				// Special case handling for the rigid root case; only discard from selection if !bLeaveRigidRoots
+				if (Parent == -1 && SimulationType[CouldBeRoot] == FGeometryCollection::ESimulationTypes::FST_Rigid)
+				{
+					if (bLeaveRigidRoots)
+					{
+						ClusterNodes.Add(Index);
+					}
+					continue;
+				}
+				else if (Parent != FGeometryCollection::Invalid && SimulationType[Parent] == FGeometryCollection::ESimulationTypes::FST_Clustered)
+				{
+					ClusterNodes.Add(Parent);
+				}
+			}
+		}
+		InOutSelection = ClusterNodes.Array();
+	}
+
 	TArray<int32> FCollectionTransformSelectionFacade::SelectRootBones() const
 	{
 		TArray<int32> OutSelection;
diff --git a/Engine/Source/Runtime/Experimental/Chaos/Private/GeometryCollection/GeometryCollectionClusteringUtility.cpp b/Engine/Source/Runtime/Experimental/Chaos/Private/GeometryCollection/GeometryCollectionClusteringUtility.cpp
--- a/Engine/Source/Runtime/Experimental/Chaos/Private/GeometryCollection/GeometryCollectionClusteringUtility.cpp
+++ b/Engine/Source/Runtime/Experimental/Chaos/Private/GeometryCollection/GeometryCollectionClusteringUtility.cpp
@@ -339,7 +339,10 @@ void FGeometryCollectionClusteringUtility::CollapseHierarchyOneLevel(FGeometryCo
 
 	TArray<int32> Roots;
 	GetRootBones(GeometryCollection, Roots);
-	RecursivelyUpdateChildBoneNames(Roots[0], Children, BoneNames);
+	if (!Roots.IsEmpty())
+	{
+		RecursivelyUpdateChildBoneNames(Roots[0], Children, BoneNames);
+	}
 
 	ValidateResults(GeometryCollection);
 }
diff --git a/Engine/Source/Runtime/Experimental/Chaos/Public/GeometryCollection/Facades/CollectionTransformSelectionFacade.h b/Engine/Source/Runtime/Experimental/Chaos/Public/GeometryCollection/Facades/CollectionTransformSelectionFacade.h
--- a/Engine/Source/Runtime/Experimental/Chaos/Public/GeometryCollection/Facades/CollectionTransformSelectionFacade.h
+++ b/Engine/Source/Runtime/Experimental/Chaos/Public/GeometryCollection/Facades/CollectionTransformSelectionFacade.h
@@ -61,6 +61,9 @@ namespace GeometryCollection::Facades
 		void ConvertSelectionToRigidNodes(const int32 Index, TArray<int32>& InSelection) const;
 		void ConvertSelectionToRigidNodes(TArray<int32>& InSelection) const;
 
+		/** Replace any Embedded or Rigid transforms in InOutSelection with their parent cluster. If there is just a rigid root node selected, optionally keep it if bLeaveRigidRoot is true. */
+		void ConvertSelectionToClusterNodes(TArray<int32>& InOutSelection, bool bLeaveRigidRoots = true) const;
+
 		/**  */
 		bool CanSelectRootBones() const { return ParentAttribute.IsValid(); }
 		TArray<int32> SelectRootBones() const;
