fix dataflow Uncluster node not setting the output in the case where it has no unclustering to do

diff --git a/Engine/Plugins/Experimental/GeometryCollectionPlugin/Source/GeometryCollectionNodes/Private/Dataflow/GeometryCollectionClusteringNodes.cpp b/Engine/Plugins/Experimental/GeometryCollectionPlugin/Source/GeometryCollectionNodes/Private/Dataflow/GeometryCollectionClusteringNodes.cpp
--- a/Engine/Plugins/Experimental/GeometryCollectionPlugin/Source/GeometryCollectionNodes/Private/Dataflow/GeometryCollectionClusteringNodes.cpp
+++ b/Engine/Plugins/Experimental/GeometryCollectionPlugin/Source/GeometryCollectionNodes/Private/Dataflow/GeometryCollectionClusteringNodes.cpp
@@ -136,6 +136,7 @@ void FClusterUnclusterDataflowNode::Evaluate(Dataflow::FContext& Context, const
 			SelectionFacade.RemoveRootNodes(Selection);
 			if (Selection.IsEmpty())
 			{
+				SetValue<FManagedArrayCollection>(Context, (const FManagedArrayCollection&)(*GeomCollection), &Collection);
 				return;
 			}
 			FGeometryCollectionClusteringUtility::CollapseHierarchyOneLevel(GeomCollection.Get(), Selection);
