Dataflow : add get root index node for managed array collections

diff --git a/Engine/Plugins/Experimental/GeometryCollectionPlugin/Source/GeometryCollectionNodes/Private/Dataflow/GeometryCollectionNodes.cpp b/Engine/Plugins/Experimental/GeometryCollectionPlugin/Source/GeometryCollectionNodes/Private/Dataflow/GeometryCollectionNodes.cpp
--- a/Engine/Plugins/Experimental/GeometryCollectionPlugin/Source/GeometryCollectionNodes/Private/Dataflow/GeometryCollectionNodes.cpp
+++ b/Engine/Plugins/Experimental/GeometryCollectionPlugin/Source/GeometryCollectionNodes/Private/Dataflow/GeometryCollectionNodes.cpp
@@ -35,6 +35,7 @@
 #include "GeometryCollection/Facades/CollectionAnchoringFacade.h"
 #include "GeometryCollection/Facades/CollectionRemoveOnBreakFacade.h"
 #include "GeometryCollection/Facades/CollectionTransformFacade.h"
+#include "GeometryCollection/Facades/CollectionHierarchyFacade.h"
 #include "DynamicMesh/MeshTransforms.h"
 #include "DynamicMesh/DynamicMesh3.h"
 
@@ -60,6 +61,7 @@ namespace Dataflow
 		DATAFLOW_NODE_REGISTER_CREATION_FACTORY(FGetArrayElementDataflowNode);
 		DATAFLOW_NODE_REGISTER_CREATION_FACTORY(FGetNumArrayElementsDataflowNode);
 		DATAFLOW_NODE_REGISTER_CREATION_FACTORY(FGetBoundingBoxesFromCollectionDataflowNode);
+		DATAFLOW_NODE_REGISTER_CREATION_FACTORY(FGetRootIndexFromCollectionDataflowNode);
 		DATAFLOW_NODE_REGISTER_CREATION_FACTORY(FGetCentroidsFromCollectionDataflowNode);
 		DATAFLOW_NODE_REGISTER_CREATION_FACTORY(FTransformCollectionDataflowNode);
 		DATAFLOW_NODE_REGISTER_CREATION_FACTORY(FBakeTransformsInCollectionDataflowNode);
@@ -396,6 +398,17 @@ void FGetBoundingBoxesFromCollectionDataflowNode::Evaluate(Dataflow::FContext& C
 	}
 }
 
+void FGetRootIndexFromCollectionDataflowNode::Evaluate(Dataflow::FContext& Context, const FDataflowOutput* Out) const
+{
+	if (Out->IsA(&RootIndex))
+	{
+		const FManagedArrayCollection& InCollection = GetValue<FManagedArrayCollection>(Context, &Collection);
+
+		Chaos::Facades::FCollectionHierarchyFacade HierarchyFacade(InCollection);
+		SetValue(Context, HierarchyFacade.GetRootIndex(), &RootIndex);
+	}
+}
+
 void FGetCentroidsFromCollectionDataflowNode::Evaluate(Dataflow::FContext& Context, const FDataflowOutput* Out) const
 {
 	if (Out->IsA<TArray<FVector>>(&Centroids))
diff --git a/Engine/Plugins/Experimental/GeometryCollectionPlugin/Source/GeometryCollectionNodes/Public/Dataflow/GeometryCollectionNodes.h b/Engine/Plugins/Experimental/GeometryCollectionPlugin/Source/GeometryCollectionNodes/Public/Dataflow/GeometryCollectionNodes.h
--- a/Engine/Plugins/Experimental/GeometryCollectionPlugin/Source/GeometryCollectionNodes/Public/Dataflow/GeometryCollectionNodes.h
+++ b/Engine/Plugins/Experimental/GeometryCollectionPlugin/Source/GeometryCollectionNodes/Public/Dataflow/GeometryCollectionNodes.h
@@ -541,6 +541,37 @@ public:
 };
 
 
+/**
+ *
+ * Get the root node index
+ *
+ */
+USTRUCT(meta = (DataflowGeometryCollection))
+struct FGetRootIndexFromCollectionDataflowNode : public FDataflowNode
+{
+	GENERATED_USTRUCT_BODY()
+	DATAFLOW_NODE_DEFINE_INTERNAL(FGetRootIndexFromCollectionDataflowNode, "GetRootIndexFromCollection", "GeometryCollection|Utilities", "")
+
+public:
+	UPROPERTY(meta = (DataflowInput))
+	FManagedArrayCollection Collection;
+
+	UPROPERTY(meta = (DataflowOutput))
+	int32 RootIndex;
+
+	FGetRootIndexFromCollectionDataflowNode(const Dataflow::FNodeParameters& InParam, FGuid InGuid = FGuid::NewGuid())
+		: FDataflowNode(InParam, InGuid)
+	{
+		RegisterInputConnection(&Collection);
+		RegisterOutputConnection(&RootIndex);
+	}
+
+	virtual void Evaluate(Dataflow::FContext& Context, const FDataflowOutput* Out) const override;
+
+};
+
+
+
 /**
  *
  * Gets centroids of pieces from a Collection
