add a BranchCollection Dataflow node

diff --git a/Engine/Plugins/Experimental/GeometryCollectionPlugin/Source/GeometryCollectionNodes/Private/Dataflow/GeometryCollectionNodes.cpp b/Engine/Plugins/Experimental/GeometryCollectionPlugin/Source/GeometryCollectionNodes/Private/Dataflow/GeometryCollectionNodes.cpp
--- a/Engine/Plugins/Experimental/GeometryCollectionPlugin/Source/GeometryCollectionNodes/Private/Dataflow/GeometryCollectionNodes.cpp
+++ b/Engine/Plugins/Experimental/GeometryCollectionPlugin/Source/GeometryCollectionNodes/Private/Dataflow/GeometryCollectionNodes.cpp
@@ -68,6 +68,7 @@ namespace Dataflow
 		DATAFLOW_NODE_REGISTER_CREATION_FACTORY(FTransformMeshDataflowNode);
 		DATAFLOW_NODE_REGISTER_CREATION_FACTORY(FCompareIntDataflowNode);
 		DATAFLOW_NODE_REGISTER_CREATION_FACTORY(FBranchDataflowNode);
+		DATAFLOW_NODE_REGISTER_CREATION_FACTORY(FBranchCollectionDataflowNode);
 		DATAFLOW_NODE_REGISTER_CREATION_FACTORY(FGetSchemaDataflowNode);
 		DATAFLOW_NODE_REGISTER_CREATION_FACTORY(FRemoveOnBreakDataflowNode);
 		DATAFLOW_NODE_REGISTER_CREATION_FACTORY(FSetAnchorStateDataflowNode);
@@ -591,6 +592,34 @@ void FBranchDataflowNode::Evaluate(Dataflow::FContext& Context, const FDataflowO
 }
 
 
+void FBranchCollectionDataflowNode::Evaluate(Dataflow::FContext& Context, const FDataflowOutput* Out) const
+{
+	if (Out->IsA(&ChosenCollection))
+	{
+		bool InCondition = GetValue<bool>(Context, &bCondition);
+
+		if (InCondition)
+		{
+			if (IsConnected(&TrueCollection))
+			{
+				const FManagedArrayCollection& InTrueCollection = GetValue(Context, &TrueCollection);
+				SetValue(Context, InTrueCollection, &ChosenCollection);
+				return;
+			}
+		}
+		else
+		{
+			if (IsConnected(&FalseCollection))
+			{
+				const FManagedArrayCollection& InFalseCollection = GetValue(Context, &FalseCollection);
+				SetValue(Context, InFalseCollection, &ChosenCollection);
+				return;
+			}
+		}
+	}
+}
+
+
 namespace {
 	inline FName GetArrayTypeString(FManagedArrayCollection::EArrayType ArrayType)
 	{
diff --git a/Engine/Plugins/Experimental/GeometryCollectionPlugin/Source/GeometryCollectionNodes/Public/Dataflow/GeometryCollectionNodes.h b/Engine/Plugins/Experimental/GeometryCollectionPlugin/Source/GeometryCollectionNodes/Public/Dataflow/GeometryCollectionNodes.h
--- a/Engine/Plugins/Experimental/GeometryCollectionPlugin/Source/GeometryCollectionNodes/Public/Dataflow/GeometryCollectionNodes.h
+++ b/Engine/Plugins/Experimental/GeometryCollectionPlugin/Source/GeometryCollectionNodes/Public/Dataflow/GeometryCollectionNodes.h
@@ -880,7 +880,50 @@ public:
 
 /**
  *
- * Collects grooup and attribute information from the Collection and outputs it into a formatted string
+ * Branch between two Managed Array Collections based on Boolean condition
+ *
+ */
+USTRUCT()
+struct FBranchCollectionDataflowNode : public FDataflowNode
+{
+	GENERATED_USTRUCT_BODY()
+	DATAFLOW_NODE_DEFINE_INTERNAL(FBranchCollectionDataflowNode, "BranchCollection", "Utilities|FlowControl", "")
+
+public:
+	/** Collection input for the 'true' case */
+	UPROPERTY(meta = (DataflowInput))
+	FManagedArrayCollection TrueCollection;
+
+	/** Collection input for the 'false' case */
+	UPROPERTY(meta = (DataflowInput))
+	FManagedArrayCollection FalseCollection;
+
+	/** Condition to select which Collection is chosen as ChosenCollection */
+	UPROPERTY(EditAnywhere, Category = "Branch");
+	bool bCondition = false;
+
+	/** Output Collection */
+	UPROPERTY(meta = (DataflowOutput))
+	FManagedArrayCollection ChosenCollection;
+
+	FBranchCollectionDataflowNode(const Dataflow::FNodeParameters& InParam, FGuid InGuid = FGuid::NewGuid())
+		: FDataflowNode(InParam, InGuid)
+	{
+		RegisterInputConnection(&TrueCollection);
+		RegisterInputConnection(&FalseCollection);
+		RegisterInputConnection(&bCondition);
+		RegisterOutputConnection(&ChosenCollection);
+	}
+
+	virtual void Evaluate(Dataflow::FContext& Context, const FDataflowOutput* Out) const override;
+
+};
+
+
+
+/**
+ *
+ * Collects group and attribute information from the Collection and outputs it into a formatted string
  *
  */
 USTRUCT(meta = (DataflowGeometryCollection))
