Dataflow : add a node to read specific float values from a float array

diff --git a/Engine/Plugins/Experimental/GeometryCollectionPlugin/Source/GeometryCollectionNodes/Private/Dataflow/GeometryCollectionNodes.cpp b/Engine/Plugins/Experimental/GeometryCollectionPlugin/Source/GeometryCollectionNodes/Private/Dataflow/GeometryCollectionNodes.cpp
--- a/Engine/Plugins/Experimental/GeometryCollectionPlugin/Source/GeometryCollectionNodes/Private/Dataflow/GeometryCollectionNodes.cpp
+++ b/Engine/Plugins/Experimental/GeometryCollectionPlugin/Source/GeometryCollectionNodes/Private/Dataflow/GeometryCollectionNodes.cpp
@@ -57,6 +57,7 @@ namespace Dataflow
 		DATAFLOW_NODE_REGISTER_CREATION_FACTORY(FStringAppendDataflowNode);
 		DATAFLOW_NODE_REGISTER_CREATION_FACTORY(FHashStringDataflowNode);
 		DATAFLOW_NODE_REGISTER_CREATION_FACTORY(FHashVectorDataflowNode);
+		DATAFLOW_NODE_REGISTER_CREATION_FACTORY(FGetFloatArrayElementDataflowNode);
 		DATAFLOW_NODE_REGISTER_CREATION_FACTORY(FFloatArrayToIntArrayDataflowNode);
 		DATAFLOW_NODE_REGISTER_CREATION_FACTORY(FGetArrayElementDataflowNode);
 		DATAFLOW_NODE_REGISTER_CREATION_FACTORY(FGetNumArrayElementsDataflowNode);
@@ -260,6 +261,18 @@ void FHashVectorDataflowNode::Evaluate(Dataflow::FContext& Context, const FDataf
 	}
 }
 
+void FGetFloatArrayElementDataflowNode::Evaluate(Dataflow::FContext& Context, const FDataflowOutput* Out) const
+{
+	if (Out->IsA(&FloatValue))
+	{
+		const TArray<float>& InArray = GetValue(Context, &FloatArray);
+		const int32& InIndex = GetValue(Context, &Index);
+		const float OutputValue = InArray.IsValidIndex(InIndex) ? InArray[InIndex] : 0.0f;
+
+		SetValue(Context, OutputValue, &FloatValue);
+	}
+}
+
 void FFloatArrayToIntArrayDataflowNode::Evaluate(Dataflow::FContext& Context, const FDataflowOutput* Out) const
 {
 	if (Out->IsA<TArray<int32>>(&IntArray))
diff --git a/Engine/Plugins/Experimental/GeometryCollectionPlugin/Source/GeometryCollectionNodes/Public/Dataflow/GeometryCollectionNodes.h b/Engine/Plugins/Experimental/GeometryCollectionPlugin/Source/GeometryCollectionNodes/Public/Dataflow/GeometryCollectionNodes.h
--- a/Engine/Plugins/Experimental/GeometryCollectionPlugin/Source/GeometryCollectionNodes/Public/Dataflow/GeometryCollectionNodes.h
+++ b/Engine/Plugins/Experimental/GeometryCollectionPlugin/Source/GeometryCollectionNodes/Public/Dataflow/GeometryCollectionNodes.h
@@ -430,6 +430,42 @@ public:
 
 };
 
+/**
+ *
+ * Returns the specified element from an array
+ *
+ */
+USTRUCT()
+struct FGetFloatArrayElementDataflowNode : public FDataflowNode
+{
+	GENERATED_USTRUCT_BODY()
+	DATAFLOW_NODE_DEFINE_INTERNAL(FGetFloatArrayElementDataflowNode, "GetFloatArrayElement", "Utilities|Array", "")
+
+public:
+	/** Element index */
+	UPROPERTY(EditAnywhere, Category = "Index");
+	int32 Index = 0;
+
+	/** Array to get the element from */
+	UPROPERTY(meta = (DataflowInput))
+	TArray<float> FloatArray;
+
+	/** Specified element */
+	UPROPERTY(meta = (DataflowOutput))
+	float FloatValue = 0;
+
+	FGetFloatArrayElementDataflowNode(const Dataflow::FNodeParameters& InParam, FGuid InGuid = FGuid::NewGuid())
+		: FDataflowNode(InParam, InGuid)
+	{
+		RegisterInputConnection(&FloatArray);
+		RegisterInputConnection(&Index);
+		RegisterOutputConnection(&FloatValue);
+	}
+
+	virtual void Evaluate(Dataflow::FContext& Context, const FDataflowOutput* Out) const override;
+
+};
+
 /**
  *
  * Returns the specified element from an array
