Change name of the 'Merge' node in geometry collection clustering nodes to 'Cluster Merge' so it is more distinguished from the many other types of merges we may have in other contexts
+ minor comment cleanup

diff --git a/Engine/Plugins/Experimental/GeometryCollectionPlugin/Source/GeometryCollectionNodes/Public/Dataflow/GeometryCollectionClusteringNodes.h b/Engine/Plugins/Experimental/GeometryCollectionPlugin/Source/GeometryCollectionNodes/Public/Dataflow/GeometryCollectionClusteringNodes.h
--- a/Engine/Plugins/Experimental/GeometryCollectionPlugin/Source/GeometryCollectionNodes/Public/Dataflow/GeometryCollectionClusteringNodes.h
+++ b/Engine/Plugins/Experimental/GeometryCollectionPlugin/Source/GeometryCollectionNodes/Public/Dataflow/GeometryCollectionClusteringNodes.h
@@ -206,17 +206,17 @@ public:
 
 
 /**
- * Merge selected clusters under a new parent
+ * Merge selected bones under a new parent cluster
  */
 USTRUCT(meta = (DataflowGeometryCollection))
 struct FClusterMergeDataflowNode : public FDataflowNode
 {
 	GENERATED_USTRUCT_BODY()
-	DATAFLOW_NODE_DEFINE_INTERNAL(FClusterMergeDataflowNode, "Merge", "GeometryCollection|Cluster", "")
+	DATAFLOW_NODE_DEFINE_INTERNAL(FClusterMergeDataflowNode, "ClusterMerge", "GeometryCollection|Cluster", "")
 
 public:
 
-	/** Collection on which to merge clusters */
+	/** Collection on which to merge bones into a cluster */
 	UPROPERTY(meta = (DataflowInput, DataflowOutput, DataflowPassthrough = "Collection", DataflowIntrinsic))
 	FManagedArrayCollection Collection;
 
