add a float comparison dataflow node to match the int comparison node, and move them both to a Math>Compare category. make them share the operator implementation.

diff --git a/Engine/Plugins/Experimental/GeometryCollectionPlugin/Source/GeometryCollectionNodes/Private/Dataflow/GeometryCollectionNodes.cpp b/Engine/Plugins/Experimental/GeometryCollectionPlugin/Source/GeometryCollectionNodes/Private/Dataflow/GeometryCollectionNodes.cpp
--- a/Engine/Plugins/Experimental/GeometryCollectionPlugin/Source/GeometryCollectionNodes/Private/Dataflow/GeometryCollectionNodes.cpp
+++ b/Engine/Plugins/Experimental/GeometryCollectionPlugin/Source/GeometryCollectionNodes/Private/Dataflow/GeometryCollectionNodes.cpp
@@ -68,6 +68,7 @@ namespace Dataflow
 		DATAFLOW_NODE_REGISTER_CREATION_FACTORY(FBakeTransformsInCollectionDataflowNode);
 		DATAFLOW_NODE_REGISTER_CREATION_FACTORY(FTransformMeshDataflowNode);
 		DATAFLOW_NODE_REGISTER_CREATION_FACTORY(FCompareIntDataflowNode);
+		DATAFLOW_NODE_REGISTER_CREATION_FACTORY(FCompareFloatDataflowNode);
 		DATAFLOW_NODE_REGISTER_CREATION_FACTORY(FBranchDataflowNode);
 		DATAFLOW_NODE_REGISTER_CREATION_FACTORY(FBranchCollectionDataflowNode);
 		DATAFLOW_NODE_REGISTER_CREATION_FACTORY(FGetSchemaDataflowNode);
@@ -541,6 +542,32 @@ void FTransformMeshDataflowNode::Evaluate(Dataflow::FContext& Context, const FDa
 	}
 }
 
+namespace
+{
+	// helper to apply an ECompareOperationEnum operation to various numeric types
+	template <typename T>
+	static bool ApplyDataflowOperationComparison(T A, T B, ECompareOperationEnum Operation)
+	{
+		switch (Operation)
+		{
+		case ECompareOperationEnum::Dataflow_Compare_Equal:
+			return A == B;
+		case ECompareOperationEnum::Dataflow_Compare_Smaller:
+			return A < B;
+		case ECompareOperationEnum::Dataflow_Compare_SmallerOrEqual:
+			return A <= B;
+		case ECompareOperationEnum::Dataflow_Compare_Greater:
+			return A > B;
+		case ECompareOperationEnum::Dataflow_Compare_GreaterOrEqual:
+			return A >= B;
+		default:
+			ensureMsgf(false, TEXT("Invalid ECompareOperationEnum value: %u"), (uint8)Operation);
+		}
+
+		return false;
+	}
+}
+
 
 void FCompareIntDataflowNode::Evaluate(Dataflow::FContext& Context, const FDataflowOutput* Out) const
 {
@@ -548,33 +575,26 @@ void FCompareIntDataflowNode::Evaluate(Dataflow::FContext& Context, const FDataf
 	{
 		int32 IntAValue = GetValue<int32>(Context, &IntA);
 		int32 IntBValue = GetValue<int32>(Context, &IntB);
-		bool ResultValue;
+		bool ResultValue = ApplyDataflowOperationComparison(IntAValue, IntBValue, Operation);
 
-		if (Operation == ECompareOperationEnum::Dataflow_Compare_Equal)
-		{
-			ResultValue = IntAValue == IntBValue ? true : false;
-		}
-		else if (Operation == ECompareOperationEnum::Dataflow_Compare_Smaller)
-		{
-			ResultValue = IntAValue < IntBValue ? true : false;
-		}
-		else if (Operation == ECompareOperationEnum::Dataflow_Compare_SmallerOrEqual)
-		{
-			ResultValue = IntAValue <= IntBValue ? true : false;
-		}
-		else if (Operation == ECompareOperationEnum::Dataflow_Compare_Greater)
-		{
-			ResultValue = IntAValue > IntBValue ? true : false;
-		}
-		else if (Operation == ECompareOperationEnum::Dataflow_Compare_GreaterOrEqual)
-		{
-			ResultValue = IntAValue >= IntBValue ? true : false;
-		}
+		SetValue<bool>(Context, ResultValue, &Result);
+	}
+}
+
+
+void FCompareFloatDataflowNode::Evaluate(Dataflow::FContext& Context, const FDataflowOutput* Out) const
+{
+	if (Out->IsA(&Result))
+	{
+		float AValue = GetValue(Context, &FloatA);
+		float BValue = GetValue(Context, &FloatB);
+		bool ResultValue = ApplyDataflowOperationComparison(AValue, BValue, Operation);
 
 		SetValue<bool>(Context, ResultValue, &Result);
 	}
 }
 
+
 void FBranchDataflowNode::Evaluate(Dataflow::FContext& Context, const FDataflowOutput* Out) const
 {
 	if (Out->IsA<TObjectPtr<UDynamicMesh>>(&Mesh))
diff --git a/Engine/Plugins/Experimental/GeometryCollectionPlugin/Source/GeometryCollectionNodes/Public/Dataflow/GeometryCollectionNodes.h b/Engine/Plugins/Experimental/GeometryCollectionPlugin/Source/GeometryCollectionNodes/Public/Dataflow/GeometryCollectionNodes.h
--- a/Engine/Plugins/Experimental/GeometryCollectionPlugin/Source/GeometryCollectionNodes/Public/Dataflow/GeometryCollectionNodes.h
+++ b/Engine/Plugins/Experimental/GeometryCollectionPlugin/Source/GeometryCollectionNodes/Public/Dataflow/GeometryCollectionNodes.h
@@ -840,7 +840,7 @@ USTRUCT()
 struct FCompareIntDataflowNode : public FDataflowNode
 {
 	GENERATED_USTRUCT_BODY()
-	DATAFLOW_NODE_DEFINE_INTERNAL(FCompareIntDataflowNode, "CompareInt", "Math|Int", "")
+	DATAFLOW_NODE_DEFINE_INTERNAL(FCompareIntDataflowNode, "CompareInt", "Math|Compare", "")
 
 public:
 	/** Comparison operation */
@@ -872,6 +872,49 @@ public:
 };
 
 
+/**
+ *
+ * Comparison between floats
+ *
+ */
+USTRUCT()
+struct FCompareFloatDataflowNode : public FDataflowNode
+{
+	GENERATED_USTRUCT_BODY()
+		DATAFLOW_NODE_DEFINE_INTERNAL(FCompareFloatDataflowNode, "CompareFloat", "Math|Compare", "")
+
+public:
+	/** Comparison operation */
+	UPROPERTY(EditAnywhere, Category = "Compare");
+	ECompareOperationEnum Operation = ECompareOperationEnum::Dataflow_Compare_Equal;
+
+	/** Float input */
+	UPROPERTY(EditAnywhere, Category = "Compare");
+	float FloatA = 0;
+
+	/** Float input */
+	UPROPERTY(EditAnywhere, Category = "Compare");
+	float FloatB = 0;
+
+	/** Boolean result of the comparison */
+	UPROPERTY(meta = (DataflowOutput));
+	bool Result = false;
+
+	FCompareFloatDataflowNode(const Dataflow::FNodeParameters& InParam, FGuid InGuid = FGuid::NewGuid())
+		: FDataflowNode(InParam, InGuid)
+	{
+		RegisterInputConnection(&FloatA);
+		RegisterInputConnection(&FloatB);
+		RegisterOutputConnection(&Result);
+	}
+
+	virtual void Evaluate(Dataflow::FContext& Context, const FDataflowOutput* Out) const override;
+
+};
+
+
+
+
 /**
  *
  * Branch between two inputs based on boolean condition
