Make parameter display names and ordering for FCreateNonOverlappingConvexHullsDataflowNode better match the fracture mode convex tool

diff --git a/Engine/Plugins/Experimental/GeometryCollectionPlugin/Source/GeometryCollectionNodes/Public/Dataflow/GeometryCollectionUtilityNodes.h b/Engine/Plugins/Experimental/GeometryCollectionPlugin/Source/GeometryCollectionNodes/Public/Dataflow/GeometryCollectionUtilityNodes.h
--- a/Engine/Plugins/Experimental/GeometryCollectionPlugin/Source/GeometryCollectionNodes/Public/Dataflow/GeometryCollectionUtilityNodes.h
+++ b/Engine/Plugins/Experimental/GeometryCollectionPlugin/Source/GeometryCollectionNodes/Public/Dataflow/GeometryCollectionUtilityNodes.h
@@ -74,26 +74,26 @@ public:
 	UPROPERTY(meta = (DataflowInput, DataflowOutput))
 	FManagedArrayCollection Collection;
 
-	/** Fraction of the convex hulls for a cluster that we can remove before using the hulls of the children */
-	UPROPERTY(EditAnywhere, Category = "Convex", meta = (DataflowInput, UIMin = 0.01f, UIMax = 1.f))
-	float CanRemoveFraction = 0.3f;
+	/** Fraction (of geometry volume) by which a cluster's convex hull volume can exceed the actual geometry volume before instead using the hulls of the children.  0 means the convex volume cannot exceed the geometry volume; 1 means the convex volume is allowed to be 100% larger (2x) the geometry volume. */
+	UPROPERTY(EditAnywhere, Category = Convex, meta = (DataflowInput, DisplayName = "Allow Larger Hull Fraction", ClampMin = 0.f))
+	float CanExceedFraction = .5f;
 
 	/** Computed convex hulls are simplified to keep points spaced at least this far apart (except where needed to keep the hull from collapsing to zero volume) */
-	UPROPERTY(EditAnywhere, Category = "Convex", meta = (DataflowInput, UIMin = 0.f))
+	UPROPERTY(EditAnywhere, Category = Convex, meta = (DataflowInput, ClampMin = 0.f))
 	float SimplificationDistanceThreshold = 10.f;
 
-	/** Fraction (of geometry volume) by which a cluster's convex hull volume can exceed the actual geometry volume before instead using the hulls of the children.  0 means the convex volume cannot exceed the geometry volume; 1 means the convex volume is allowed to be 100% larger (2x) the geometry volume. */
-	UPROPERTY(EditAnywhere, Category = "Convex", meta = (DataflowInput, UIMin = 0.f))
-	float CanExceedFraction = 0.5f;
-
 	/** Whether and in what cases to automatically cut away overlapping parts of the convex hulls, to avoid the simulation 'popping' to fix the overlaps */
-	UPROPERTY(EditAnywhere, Category = "Convex")
+	UPROPERTY(EditAnywhere, Category = AutomaticOverlapRemoval, meta = (DisplayName = "Remove Overlaps"))
 	EConvexOverlapRemovalMethodEnum OverlapRemovalMethod = EConvexOverlapRemovalMethodEnum::Dataflow_EConvexOverlapRemovalMethod_All;
 
 	/** Overlap removal will be computed as if convex hulls were this percentage smaller (in range 0-100) */
-	UPROPERTY(EditAnywhere, Category = "Convex", meta = (DataflowInput, UIMin = 0.f, UIMax = 100.f))
+	UPROPERTY(EditAnywhere, Category = AutomaticOverlapRemoval, meta = (DataflowInput, ClampMin = 0.f, ClampMax = 99.9f))
 	float OverlapRemovalShrinkPercent = 0.f;
 
+	/** Fraction of the convex hulls for a cluster that we can remove before using the hulls of the children */
+	UPROPERTY(EditAnywhere, Category = AutomaticOverlapRemoval, meta = (DataflowInput, DisplayName = "Max Removal Fraction", ClampMin = 0.01f, ClampMax = 1.f))
+	float CanRemoveFraction = 0.3f;
+
 	FCreateNonOverlappingConvexHullsDataflowNode(const Dataflow::FNodeParameters& InParam, FGuid InGuid = FGuid::NewGuid());
 
 	virtual void Evaluate(Dataflow::FContext& Context, const FDataflowOutput* Out) const override;
