dataflow : make convex generation node to be not likely crash if the collection is not plugged while a selection is

diff --git a/Engine/Plugins/Experimental/GeometryCollectionPlugin/Source/GeometryCollectionNodes/Private/Dataflow/GeometryCollectionUtilityNodes.cpp b/Engine/Plugins/Experimental/GeometryCollectionPlugin/Source/GeometryCollectionNodes/Private/Dataflow/GeometryCollectionUtilityNodes.cpp
--- a/Engine/Plugins/Experimental/GeometryCollectionPlugin/Source/GeometryCollectionNodes/Private/Dataflow/GeometryCollectionUtilityNodes.cpp
+++ b/Engine/Plugins/Experimental/GeometryCollectionPlugin/Source/GeometryCollectionNodes/Private/Dataflow/GeometryCollectionUtilityNodes.cpp
@@ -35,7 +35,7 @@ FCreateLeafConvexHullsDataflowNode::FCreateLeafConvexHullsDataflowNode(const Dat
 
 void FCreateLeafConvexHullsDataflowNode::Evaluate(Dataflow::FContext& Context, const FDataflowOutput* Out) const
 {
-	if (Out->IsA(&Collection))
+	if (Out->IsA(&Collection) && IsConnected(&Collection))
 	{
 		const FManagedArrayCollection& InCollection = GetValue(Context, &Collection);
 		if (InCollection.NumElements(FGeometryCollection::TransformGroup) == 0)
@@ -78,7 +78,7 @@ FCreateNonOverlappingConvexHullsDataflowNode::FCreateNonOverlappingConvexHullsDa
 
 void FCreateNonOverlappingConvexHullsDataflowNode::Evaluate(Dataflow::FContext& Context, const FDataflowOutput* Out) const
 {
-	if (Out->IsA<FManagedArrayCollection>(&Collection))
+	if (Out->IsA<FManagedArrayCollection>(&Collection) && IsConnected(&Collection))
 	{
 		const FManagedArrayCollection& InCollection = GetValue<FManagedArrayCollection>(Context, &Collection);
 
@@ -114,7 +114,7 @@ FGenerateClusterConvexHullsFromLeafHullsDataflowNode::FGenerateClusterConvexHull
 
 void FGenerateClusterConvexHullsFromLeafHullsDataflowNode::Evaluate(Dataflow::FContext& Context, const FDataflowOutput* Out) const
 {
-	if (Out->IsA<FManagedArrayCollection>(&Collection))
+	if (Out->IsA<FManagedArrayCollection>(&Collection) && IsConnected(&Collection))
 	{
 		const FManagedArrayCollection& InCollection = GetValue<FManagedArrayCollection>(Context, &Collection);
 
@@ -161,7 +161,7 @@ FGenerateClusterConvexHullsFromChildrenHullsDataflowNode::FGenerateClusterConvex
 
 void FGenerateClusterConvexHullsFromChildrenHullsDataflowNode::Evaluate(Dataflow::FContext& Context, const FDataflowOutput* Out) const
 {
-	if (Out->IsA<FManagedArrayCollection>(&Collection))
+	if (Out->IsA<FManagedArrayCollection>(&Collection) && IsConnected(&Collection))
 	{
 		const FManagedArrayCollection& InCollection = GetValue<FManagedArrayCollection>(Context, &Collection);
 
