Dataflow : fix flatten geometry collection node crash
- also fix the logic to properly use the root index instead of 0

diff --git a/Engine/Plugins/Experimental/GeometryCollectionPlugin/Source/GeometryCollectionNodes/Private/Dataflow/GeometryCollectionClusteringNodes.cpp b/Engine/Plugins/Experimental/GeometryCollectionPlugin/Source/GeometryCollectionNodes/Private/Dataflow/GeometryCollectionClusteringNodes.cpp
--- a/Engine/Plugins/Experimental/GeometryCollectionPlugin/Source/GeometryCollectionNodes/Private/Dataflow/GeometryCollectionClusteringNodes.cpp
+++ b/Engine/Plugins/Experimental/GeometryCollectionPlugin/Source/GeometryCollectionNodes/Private/Dataflow/GeometryCollectionClusteringNodes.cpp
@@ -95,27 +95,29 @@ void FAutoClusterDataflowNode::Evaluate(Dataflow::FContext& Context, const FData
 
 void FClusterFlattenDataflowNode::Evaluate(Dataflow::FContext& Context, const FDataflowOutput* Out) const
 {
-	if (Out->IsA<FManagedArrayCollection>(&Collection))
+	if (Out->IsA(&Collection) && IsConnected(&Collection))
 	{
-		const FManagedArrayCollection& InCollection = GetValue<FManagedArrayCollection>(Context, &Collection);
-		if (TUniquePtr<FGeometryCollection> GeomCollection = TUniquePtr<FGeometryCollection>(InCollection.NewCopy<FGeometryCollection>()))
+		const FManagedArrayCollection& InCollection = GetValue(Context, &Collection);
+		if (InCollection.NumElements(FGeometryCollection::TransformAttribute) > 0)
 		{
-			FGeometryCollectionClusteringUtility::UpdateHierarchyLevelOfChildren(GeomCollection.Get(), -1);
-
-			const TManagedArray<int32>& Levels = GeomCollection->GetAttribute<int32>("Level", FGeometryCollection::TransformGroup);
-
-			// Populate Selected Bones in an Array
-			// @todo(harsha) Implement with Selection
-			// For every bone in selected array: [ClusterIndex]
-			int32 ClusterIndex = 0;
-			TArray<int32> LeafBones;
-			FGeometryCollectionClusteringUtility::GetLeafBones(GeomCollection.Get(), ClusterIndex, true, LeafBones);
-			FGeometryCollectionClusteringUtility::ClusterBonesUnderExistingNode(GeomCollection.Get(), ClusterIndex, LeafBones);
-			FGeometryCollectionClusteringUtility::RemoveDanglingClusters(GeomCollection.Get());
-			// End for
-
-			FGeometryCollectionClusteringUtility::UpdateHierarchyLevelOfChildren(GeomCollection.Get(), -1);
-			SetValue<FManagedArrayCollection>(Context, (const FManagedArrayCollection&)(*GeomCollection), &Collection);
+			if (TUniquePtr<FGeometryCollection> GeomCollection = TUniquePtr<FGeometryCollection>(InCollection.NewCopy<FGeometryCollection>()))
+			{
+				Chaos::Facades::FCollectionHierarchyFacade HierarchyFacade(*GeomCollection);
+				HierarchyFacade.GenerateLevelAttribute();
+
+				// Populate Selected Bones in an Array
+				// @todo(harsha) Implement with Selection
+				// For every bone in selected array: [RootClusterIndex]
+				const int32 RootClusterIndex = HierarchyFacade.GetRootIndex();
+				TArray<int32> LeafBones;
+				FGeometryCollectionClusteringUtility::GetLeafBones(GeomCollection.Get(), RootClusterIndex, true, LeafBones);
+				FGeometryCollectionClusteringUtility::ClusterBonesUnderExistingNode(GeomCollection.Get(), RootClusterIndex, LeafBones);
+				FGeometryCollectionClusteringUtility::RemoveDanglingClusters(GeomCollection.Get());
+				// End for
+
+				HierarchyFacade.GenerateLevelAttribute();
+				SetValue(Context, (const FManagedArrayCollection&)(*GeomCollection), &Collection);
+			}
 		}
 	}
 }
