Dataflow (Submitted on behalf of gustav.melich)
- Added GetFloatOverrideFromAsset, GetIntOverrideFromAsset, GetBoolOverrideFromAsset, GetStringOverrideFromAsset nodes
- Added 'Overrides' property to GC asset
- Added override node invalidation when 'Overrides' property gets changed on the asset
- bug fixes

diff --git a/Engine/Plugins/Experimental/Dataflow/Source/DataflowEditor/Private/Dataflow/DataflowEditorCommands.cpp b/Engine/Plugins/Experimental/Dataflow/Source/DataflowEditor/Private/Dataflow/DataflowEditorCommands.cpp
--- a/Engine/Plugins/Experimental/Dataflow/Source/DataflowEditor/Private/Dataflow/DataflowEditorCommands.cpp
+++ b/Engine/Plugins/Experimental/Dataflow/Source/DataflowEditor/Private/Dataflow/DataflowEditorCommands.cpp
@@ -5,6 +5,7 @@
 #include "Dataflow/DataflowEdNode.h"
 #include "Dataflow/DataflowNodeFactory.h"
 #include "Dataflow/DataflowObject.h"
+#include "Dataflow/DataflowOverrideNode.h"
 #include "EdGraph/EdGraphNode.h"
 #include "IStructureDetailsView.h"
 #include "EdGraphNode_Comment.h"
@@ -240,6 +241,31 @@ void FDataflowEditorCommands::OnNodeTitleCommitted(const FText& InNewText, EText
 	}
 }
 
+void FDataflowEditorCommands::OnAssetPropertyValueChanged(UDataflow* Graph, TSharedPtr<Dataflow::FEngineContext>& Context, Dataflow::FTimestamp& OutLastNodeTimestamp, const FPropertyChangedEvent& InPropertyChangedEvent)
+{
+	if (ensureMsgf(Graph != nullptr, TEXT("Warning : Failed to find valid graph.")))
+	{
+		if (InPropertyChangedEvent.ChangeType == EPropertyChangeType::ValueSet ||
+			InPropertyChangedEvent.ChangeType == EPropertyChangeType::ArrayRemove ||
+			InPropertyChangedEvent.ChangeType == EPropertyChangeType::ArrayClear)
+		{
+			if (InPropertyChangedEvent.GetPropertyName() == FName("Overrides_Key") ||
+				InPropertyChangedEvent.GetPropertyName() == FName("Overrides"))
+			{
+				for (const TSharedPtr<FDataflowNode>& DataflowNode : Graph->Dataflow->GetNodes())
+				{
+					if (DataflowNode->IsA(FDataflowOverrideNode::StaticType()))
+					{
+						// TODO: For now we invalidate all the FDataflowOverrideNode nodes
+						// Once the Variable system will be in place only the neccessary nodes
+						// will be invalidated
+						DataflowNode->Invalidate();
+					}
+				}
+			}
+		}
+	}
+}
 
 void FDataflowEditorCommands::OnPropertyValueChanged(UDataflow* OutDataflow, TSharedPtr<Dataflow::FEngineContext>& Context, Dataflow::FTimestamp& OutLastNodeTimestamp, const FPropertyChangedEvent& InPropertyChangedEvent, const TSet<UObject*>& SelectedNodes)
 {
diff --git a/Engine/Plugins/Experimental/Dataflow/Source/DataflowEditor/Private/Dataflow/DataflowEditorToolkit.cpp b/Engine/Plugins/Experimental/Dataflow/Source/DataflowEditor/Private/Dataflow/DataflowEditorToolkit.cpp
--- a/Engine/Plugins/Experimental/Dataflow/Source/DataflowEditor/Private/Dataflow/DataflowEditorToolkit.cpp
+++ b/Engine/Plugins/Experimental/Dataflow/Source/DataflowEditor/Private/Dataflow/DataflowEditorToolkit.cpp
@@ -102,6 +102,16 @@ FDataflowEditorToolkit::~FDataflowEditorToolkit()
 		GraphEditor->OnSelectionChangedMulticast.Remove(OnSelectionChangedMulticastDelegateHandle);
 		GraphEditor->OnNodeDeletedMulticast.Remove(OnNodeDeletedMulticastDelegateHandle);
 	}
+
+	if (NodeDetailsEditor)
+	{
+		NodeDetailsEditor->GetOnFinishedChangingPropertiesDelegate().Remove(OnFinishedChangingPropertiesDelegateHandle);
+	}
+
+	if (AssetDetailsEditor)
+	{
+		 AssetDetailsEditor->OnFinishedChangingProperties().Remove(OnFinishedChangingAssetPropertiesDelegateHandle);
+	}
 }
 
 void FDataflowEditorToolkit::InitializeEditor(const EToolkitMode::Type Mode, const TSharedPtr<IToolkitHost>& InitToolkitHost, UObject* ObjectToEdit)
@@ -180,6 +190,11 @@ void FDataflowEditorToolkit::OnPropertyValueChanged(const FPropertyChangedEvent&
 	FDataflowEditorCommands::OnPropertyValueChanged(this->GetDataflow(), Context, LastNodeTimestamp, PropertyChangedEvent, PrevNodeSelection);
 }
 
+void FDataflowEditorToolkit::OnAssetPropertyValueChanged(const FPropertyChangedEvent& PropertyChangedEvent)
+{
+	FDataflowEditorCommands::OnAssetPropertyValueChanged(this->GetDataflow(), Context, LastNodeTimestamp, PropertyChangedEvent);
+}
+
 bool FDataflowEditorToolkit::OnNodeVerifyTitleCommit(const FText& NewText, UEdGraphNode* GraphNode, FText& OutErrorMessage)
 {
 	return FDataflowEditorCommands::OnNodeVerifyTitleCommit(NewText, GraphNode, OutErrorMessage);
@@ -337,7 +352,7 @@ TSharedPtr<IStructureDetailsView> FDataflowEditorToolkit::CreateNodeDetailsEdito
 	}
 	TSharedPtr<IStructureDetailsView> DetailsView = PropertyEditorModule.CreateStructureDetailView(DetailsViewArgs, StructureViewArgs, nullptr);
 	DetailsView->GetDetailsView()->SetObject(ObjectToEdit);
-	DetailsView->GetOnFinishedChangingPropertiesDelegate().AddSP(this, &FDataflowEditorToolkit::OnPropertyValueChanged);
+	OnFinishedChangingPropertiesDelegateHandle = DetailsView->GetOnFinishedChangingPropertiesDelegate().AddSP(this, &FDataflowEditorToolkit::OnPropertyValueChanged);
 
 	return DetailsView;
 }
@@ -358,6 +373,9 @@ TSharedPtr<IDetailsView> FDataflowEditorToolkit::CreateAssetDetailsEditorWidget(
 
 	TSharedPtr<IDetailsView> DetailsView = PropertyEditorModule.CreateDetailView(DetailsViewArgs);
 	DetailsView->SetObject(ObjectToEdit);
+
+	OnFinishedChangingAssetPropertiesDelegateHandle = DetailsView->OnFinishedChangingProperties().AddSP(this, &FDataflowEditorToolkit::OnAssetPropertyValueChanged);
+
 	return DetailsView;
 
 }
diff --git a/Engine/Plugins/Experimental/Dataflow/Source/DataflowEditor/Public/Dataflow/DataflowEditorCommands.h b/Engine/Plugins/Experimental/Dataflow/Source/DataflowEditor/Public/Dataflow/DataflowEditorCommands.h
--- a/Engine/Plugins/Experimental/Dataflow/Source/DataflowEditor/Public/Dataflow/DataflowEditorCommands.h
+++ b/Engine/Plugins/Experimental/Dataflow/Source/DataflowEditor/Public/Dataflow/DataflowEditorCommands.h
@@ -95,6 +95,7 @@ public:
 	*  OnPropertyValueChanged
 	*/
 	static void OnPropertyValueChanged(UDataflow* Graph, TSharedPtr<Dataflow::FEngineContext>& Context, Dataflow::FTimestamp& OutLastNodeTimestamp, const FPropertyChangedEvent& PropertyChangedEvent, const TSet<UObject*>& SelectedNodes = TSet<UObject*>());
+	static void OnAssetPropertyValueChanged(UDataflow* Graph, TSharedPtr<Dataflow::FEngineContext>& Context, Dataflow::FTimestamp& OutLastNodeTimestamp, const FPropertyChangedEvent& PropertyChangedEvent);
 
 	/*
 	*  OnSelectedNodesChanged
diff --git a/Engine/Plugins/Experimental/Dataflow/Source/DataflowEditor/Public/Dataflow/DataflowEditorToolkit.h b/Engine/Plugins/Experimental/Dataflow/Source/DataflowEditor/Public/Dataflow/DataflowEditorToolkit.h
--- a/Engine/Plugins/Experimental/Dataflow/Source/DataflowEditor/Public/Dataflow/DataflowEditorToolkit.h
+++ b/Engine/Plugins/Experimental/Dataflow/Source/DataflowEditor/Public/Dataflow/DataflowEditorToolkit.h
@@ -103,6 +103,7 @@ protected:
 	void OnNodeTitleCommitted(const FText& InNewText, ETextCommit::Type InCommitType, UEdGraphNode* GraphNode);
 	void OnNodeSelectionChanged(const TSet<UObject*>& NewSelection);
 	void OnNodeDeleted(const TSet<UObject*>& NewSelection);
+	void OnAssetPropertyValueChanged(const FPropertyChangedEvent& PropertyChangedEvent);
 	//~ End DataflowEditorActions
 
 private:
@@ -166,4 +167,6 @@ private:
 
 	FDelegateHandle OnSelectionChangedMulticastDelegateHandle;
 	FDelegateHandle OnNodeDeletedMulticastDelegateHandle;
+	FDelegateHandle OnFinishedChangingPropertiesDelegateHandle;
+	FDelegateHandle OnFinishedChangingAssetPropertiesDelegateHandle;
 };
diff --git a/Engine/Plugins/Experimental/GeometryCollectionPlugin/Source/GeometryCollectionNodes/Private/Dataflow/GeometryCollectionNodesPlugin.cpp b/Engine/Plugins/Experimental/GeometryCollectionPlugin/Source/GeometryCollectionNodes/Private/Dataflow/GeometryCollectionNodesPlugin.cpp
--- a/Engine/Plugins/Experimental/GeometryCollectionPlugin/Source/GeometryCollectionNodes/Private/Dataflow/GeometryCollectionNodesPlugin.cpp
+++ b/Engine/Plugins/Experimental/GeometryCollectionPlugin/Source/GeometryCollectionNodes/Private/Dataflow/GeometryCollectionNodesPlugin.cpp
@@ -17,7 +17,7 @@
 #include "Dataflow/GeometryCollectionUtilityNodes.h"
 #include "Dataflow/GeometryCollectionMaterialNodes.h"
 #include "Dataflow/GeometryCollectionFieldNodes.h"
-//#include "Dataflow/GeometryCollectionOverrideNodes.h"
+#include "Dataflow/GeometryCollectionOverrideNodes.h"
 #include "Dataflow/GeometryCollectionMakeNodes.h"
 #include "Dataflow/GeometryCollectionMathNodes.h"
 #include "Dataflow/GeometryCollectionConversionNodes.h"
@@ -40,7 +40,7 @@ void IGeometryCollectionNodesPlugin::StartupModule()
 	Dataflow::GeometryCollectionUtilityNodes();
 	Dataflow::GeometryCollectionMaterialNodes();
 	Dataflow::GeometryCollectionFieldNodes();
-	//Dataflow::GeometryCollectionOverrideNodes();
+	Dataflow::GeometryCollectionOverrideNodes();
 	Dataflow::GeometryCollectionMakeNodes();
 	Dataflow::GeometryCollectionMathNodes();
 	Dataflow::GeometryCollectionConversionNodes();
diff --git a/Engine/Plugins/Experimental/GeometryCollectionPlugin/Source/GeometryCollectionNodes/Private/Dataflow/GeometryCollectionOverrideNodes.cpp b/Engine/Plugins/Experimental/GeometryCollectionPlugin/Source/GeometryCollectionNodes/Private/Dataflow/GeometryCollectionOverrideNodes.cpp
new file mode 100644
--- /dev/null
+++ b/Engine/Plugins/Experimental/GeometryCollectionPlugin/Source/GeometryCollectionNodes/Private/Dataflow/GeometryCollectionOverrideNodes.cpp
@@ -0,0 +1,124 @@
+// Copyright Epic Games, Inc. All Rights Reserved.
+
+#include "Dataflow/GeometryCollectionOverrideNodes.h"
+#include "Dataflow/DataflowCore.h"
+
+#include "Engine/Engine.h"
+
+#include UE_INLINE_GENERATED_CPP_BY_NAME(GeometryCollectionOverrideNodes)
+
+namespace Dataflow
+{
+	void GeometryCollectionOverrideNodes()
+	{
+		static const FLinearColor CDefaultNodeBodyTintColor = FLinearColor(0.f, 0.f, 0.f, 0.5f);
+
+		DATAFLOW_NODE_REGISTER_CREATION_FACTORY(FGetFloatOverrideFromAssetDataflowNode);
+		DATAFLOW_NODE_REGISTER_CREATION_FACTORY(FGetIntOverrideFromAssetDataflowNode);
+		DATAFLOW_NODE_REGISTER_CREATION_FACTORY(FGetBoolOverrideFromAssetDataflowNode);
+		DATAFLOW_NODE_REGISTER_CREATION_FACTORY(FGetStringOverrideFromAssetDataflowNode);
+
+		// Override
+		DATAFLOW_NODE_REGISTER_CREATION_FACTORY_NODE_COLORS_BY_CATEGORY("Override", FLinearColor(1.f, 0.4f, 0.4f), CDefaultNodeBodyTintColor);
+	}
+}
+
+
+void FGetFloatOverrideFromAssetDataflowNode::Evaluate(Dataflow::FContext& Context, const FDataflowOutput* Out) const
+{
+	if (Out->IsA(&Float))
+	{
+		float NewValue = FCString::Atof(*GetDefaultValue(Context));
+
+		if (const Dataflow::FEngineContext* EngineContext = Context.AsType<Dataflow::FEngineContext>())
+		{
+			const FString ValueFromAsset = GetValueFromAsset(Context, EngineContext->Owner);
+
+			if (!ValueFromAsset.IsEmpty() && ValueFromAsset.IsNumeric())
+			{
+				NewValue = FCString::Atof(*ValueFromAsset);
+			}
+		}
+
+		SetValue(Context, NewValue, &Float);
+	}
+}
+
+void FGetIntOverrideFromAssetDataflowNode::Evaluate(Dataflow::FContext& Context, const FDataflowOutput* Out) const
+{
+	if (Out->IsA(&Int))
+	{
+		int32 NewValue = FCString::Atoi(*GetDefaultValue(Context));
+
+		if (const Dataflow::FEngineContext* EngineContext = Context.AsType<Dataflow::FEngineContext>())
+		{
+			const FString ValueFromAsset = GetValueFromAsset(Context, EngineContext->Owner);
+
+			if (!ValueFromAsset.IsEmpty() && ValueFromAsset.IsNumeric())
+			{
+				NewValue = FCString::Atoi(*ValueFromAsset);
+			}
+		}
+
+		SetValue(Context, NewValue, &Int);
+	}
+}
+
+static bool StringToBool(const FString& InString, bool InDefault)
+{
+	bool Result = InDefault;
+
+	if (!InString.IsEmpty())
+	{
+		if (InString.IsNumeric())
+		{
+			Result = FCString::Atoi(*InString) == 0 ? false : true;
+		}
+		else
+		{
+			Result = !InString.Compare("false") ? false : !InString.Compare("true") ? true : InDefault;
+		}
+	}
+
+	return Result;
+}
+
+void FGetBoolOverrideFromAssetDataflowNode::Evaluate(Dataflow::FContext& Context, const FDataflowOutput* Out) const
+{
+	if (Out->IsA(&Bool))
+	{
+		bool NewValue = StringToBool(GetDefaultValue(Context), false);
+
+		if (const Dataflow::FEngineContext* EngineContext = Context.AsType<Dataflow::FEngineContext>())
+		{
+			const FString ValueFromAsset = GetValueFromAsset(Context, EngineContext->Owner);
+
+			if (!ValueFromAsset.IsEmpty())
+			{
+				NewValue = StringToBool(ValueFromAsset, false);
+			}
+		}
+
+		SetValue(Context, NewValue, &Bool);
+	}
+}
+
+void FGetStringOverrideFromAssetDataflowNode::Evaluate(Dataflow::FContext& Context, const FDataflowOutput* Out) const
+{
+	if (Out->IsA(&String))
+	{
+		FString NewValue = GetDefaultValue(Context);
+
+		if (const Dataflow::FEngineContext* EngineContext = Context.AsType<Dataflow::FEngineContext>())
+		{
+			const FString ValueFromAsset = GetValueFromAsset(Context, EngineContext->Owner);
+
+			if (!ValueFromAsset.IsEmpty())
+			{
+				NewValue = ValueFromAsset;
+			}
+		}
+
+		SetValue(Context, NewValue, &String);
+	}
+}
diff --git a/Engine/Plugins/Experimental/GeometryCollectionPlugin/Source/GeometryCollectionNodes/Public/Dataflow/GeometryCollectionOverrideNodes.h b/Engine/Plugins/Experimental/GeometryCollectionPlugin/Source/GeometryCollectionNodes/Public/Dataflow/GeometryCollectionOverrideNodes.h
new file mode 100644
--- /dev/null
+++ b/Engine/Plugins/Experimental/GeometryCollectionPlugin/Source/GeometryCollectionNodes/Public/Dataflow/GeometryCollectionOverrideNodes.h
@@ -0,0 +1,110 @@
+// Copyright Epic Games, Inc. All Rights Reserved.
+
+#pragma once 
+
+#include "CoreMinimal.h"
+#include "Dataflow/DataflowEngine.h"
+#include "Dataflow/DataflowOverrideNode.h"
+
+#include "GeometryCollectionOverrideNodes.generated.h"
+
+/**
+ *
+ * 
+ *
+ */
+USTRUCT(meta = (DataflowGeometryCollection, DataflowOverride))
+struct FGetFloatOverrideFromAssetDataflowNode : public FDataflowOverrideNode
+{
+	GENERATED_USTRUCT_BODY()
+	DATAFLOW_NODE_DEFINE_INTERNAL(FGetFloatOverrideFromAssetDataflowNode, "GetFloatOverrideFromAsset", "GeometryCollection|Overrides", "")
+
+public:
+	UPROPERTY(meta = (DataflowOutput))
+	float Float = 0.f;
+
+	FGetFloatOverrideFromAssetDataflowNode(const Dataflow::FNodeParameters& InParam, FGuid InGuid = FGuid::NewGuid())
+		: FDataflowOverrideNode(InParam, InGuid)
+	{
+		RegisterOutputConnection(&Float);
+	}
+
+	virtual void Evaluate(Dataflow::FContext& Context, const FDataflowOutput* Out) const override;
+};
+
+/**
+ *
+ *
+ *
+ */
+USTRUCT(meta = (DataflowGeometryCollection, DataflowOverride))
+struct FGetIntOverrideFromAssetDataflowNode : public FDataflowOverrideNode
+{
+	GENERATED_USTRUCT_BODY()
+	DATAFLOW_NODE_DEFINE_INTERNAL(FGetIntOverrideFromAssetDataflowNode, "GetIntOverrideFromAsset", "GeometryCollection|Overrides", "")
+
+public:
+	UPROPERTY(meta = (DataflowOutput))
+	int32 Int = 0;
+
+	FGetIntOverrideFromAssetDataflowNode(const Dataflow::FNodeParameters& InParam, FGuid InGuid = FGuid::NewGuid())
+		: FDataflowOverrideNode(InParam, InGuid)
+	{
+		RegisterOutputConnection(&Int);
+	}
+
+	virtual void Evaluate(Dataflow::FContext& Context, const FDataflowOutput* Out) const override;
+};
+
+/**
+ *
+ *
+ *
+ */
+USTRUCT(meta = (DataflowGeometryCollection, DataflowOverride))
+struct FGetBoolOverrideFromAssetDataflowNode : public FDataflowOverrideNode
+{
+	GENERATED_USTRUCT_BODY()
+	DATAFLOW_NODE_DEFINE_INTERNAL(FGetBoolOverrideFromAssetDataflowNode, "GetBoolOverrideFromAsset", "GeometryCollection|Overrides", "")
+
+public:
+	UPROPERTY(meta = (DataflowOutput))
+	bool Bool = false;
+
+	FGetBoolOverrideFromAssetDataflowNode(const Dataflow::FNodeParameters& InParam, FGuid InGuid = FGuid::NewGuid())
+		: FDataflowOverrideNode(InParam, InGuid)
+	{
+	RegisterOutputConnection(&Bool);
+	}
+
+	virtual void Evaluate(Dataflow::FContext& Context, const FDataflowOutput* Out) const override;
+};
+
+/**
+ *
+ *
+ *
+ */
+USTRUCT(meta = (DataflowGeometryCollection, DataflowOverride))
+struct FGetStringOverrideFromAssetDataflowNode : public FDataflowOverrideNode
+{
+	GENERATED_USTRUCT_BODY()
+	DATAFLOW_NODE_DEFINE_INTERNAL(FGetStringOverrideFromAssetDataflowNode, "GetStringOverrideFromAsset", "GeometryCollection|Overrides", "")
+
+public:
+	UPROPERTY(meta = (DataflowOutput))
+	FString String;
+
+	FGetStringOverrideFromAssetDataflowNode(const Dataflow::FNodeParameters& InParam, FGuid InGuid = FGuid::NewGuid())
+		: FDataflowOverrideNode(InParam, InGuid)
+	{
+		RegisterOutputConnection(&String);
+	}
+
+	virtual void Evaluate(Dataflow::FContext& Context, const FDataflowOutput* Out) const override;
+};
+
+namespace Dataflow
+{
+	void GeometryCollectionOverrideNodes();
+}
diff --git a/Engine/Source/Runtime/Experimental/Dataflow/Core/Private/Dataflow/DataflowOverrideNode.cpp b/Engine/Source/Runtime/Experimental/Dataflow/Core/Private/Dataflow/DataflowOverrideNode.cpp
new file mode 100644
--- /dev/null
+++ b/Engine/Source/Runtime/Experimental/Dataflow/Core/Private/Dataflow/DataflowOverrideNode.cpp
@@ -0,0 +1,60 @@
+// Copyright Epic Games, Inc. All Rights Reserved.
+
+#include "Dataflow/DataflowOverrideNode.h"
+
+#include "ChaosLog.h"
+#include "Dataflow/DataflowInputOutput.h"
+#include "Dataflow/DataflowArchive.h"
+
+#include UE_INLINE_GENERATED_CPP_BY_NAME(DataflowOverrideNode)
+
+template<class T>
+static const T FindOverrideMapProperty(const UObject* Owner, FName PropertyName, FName ArrayKey, const T& Default = T())
+{
+	if (Owner && Owner->GetClass())
+	{
+		if (FMapProperty* MapProperty = FindFProperty<FMapProperty>(Owner->GetClass(), PropertyName))
+		{
+			if (const TMap<FString, FString>* ContainerData = MapProperty->ContainerPtrToValuePtr<TMap<FString, FString>>(Owner))
+			{
+				const FString Key = ArrayKey.ToString();
+				if (ContainerData->Contains(Key))
+				{
+					return (FString)(*(ContainerData->Find(Key)));
+				}
+			}
+		}
+	}
+	return Default;
+}
+
+bool FDataflowOverrideNode::ShouldInvalidate(FName InKey) const
+{
+	if (IsConnected<FName>(&Key))
+	{
+		return true;
+	}
+
+	return InKey == Key;
+}
+
+FString FDataflowOverrideNode::GetDefaultValue(Dataflow::FContext& Context) const
+{
+	return GetValue<FString>(Context, &Default, Default);
+}
+
+FString FDataflowOverrideNode::GetValueFromAsset(Dataflow::FContext& Context, const UObject* InOwner) const
+{
+	FName InKey = GetValue<FName>(Context, &Key, Key);
+	FString InDefault = GetValue<FString>(Context, &Default, Default);
+
+	if (InOwner)
+	{
+		if (!InKey.IsNone() && InKey.IsValid() && InKey.ToString().Len() > 0)
+		{
+			return FindOverrideMapProperty<FString>(InOwner, FName("Overrides"), InKey, InDefault);
+		}
+	}
+
+	return InDefault;
+}
diff --git a/Engine/Source/Runtime/Experimental/Dataflow/Core/Public/Dataflow/DataflowOverrideNode.h b/Engine/Source/Runtime/Experimental/Dataflow/Core/Public/Dataflow/DataflowOverrideNode.h
new file mode 100644
--- /dev/null
+++ b/Engine/Source/Runtime/Experimental/Dataflow/Core/Public/Dataflow/DataflowOverrideNode.h
@@ -0,0 +1,59 @@
+// Copyright Epic Games, Inc. All Rights Reserved.
+
+#pragma once
+
+#include "ChaosLog.h"
+#include "CoreMinimal.h"
+#include "Dataflow/DataflowConnection.h"
+#include "Dataflow/DataflowNodeParameters.h"
+#include "Dataflow/DataflowNode.h"
+
+#include "DataflowOverrideNode.generated.h"
+
+struct FDataflowInput;
+struct FDataflowOutput;
+
+/**
+* FDataflowOverrideNode
+*		Base class for override nodes within the Dataflow graph. 
+* 
+*		Override Nodes allow to access to Override property on
+*		the asset. They can read the values by the key.
+*/
+USTRUCT()
+struct DATAFLOWCORE_API FDataflowOverrideNode : public FDataflowNode
+{
+	GENERATED_USTRUCT_BODY()
+
+	DATAFLOW_NODE_DEFINE_INTERNAL(FDataflowOverrideNode, "DataflowOverrideNode", "BaseClass", "")
+
+public:
+	FDataflowOverrideNode(const Dataflow::FNodeParameters& Param, FGuid InGuid = FGuid::NewGuid())
+		: Super(Param,InGuid) 
+	{
+		RegisterInputConnection(&Key);
+		RegisterInputConnection(&Default);
+	}
+
+	virtual ~FDataflowOverrideNode() { }
+
+	bool ShouldInvalidate(FName InKey) const;
+
+	FString GetDefaultValue(Dataflow::FContext& Context) const;
+
+	FString GetValueFromAsset(Dataflow::FContext& Context, const UObject* InOwner) const;
+
+	//
+	// Evaluate
+	//
+	virtual void Evaluate(Dataflow::FContext& Context, const FDataflowOutput* Out) const { ensure(false); };
+
+public:
+	UPROPERTY(EditAnywhere, Category = "Overrides", meta = (DataflowInput))
+	FName Key = "Key";
+
+	UPROPERTY(EditAnywhere, Category = "Overrides", meta = (DataflowInput))
+	FString Default = FString("0");
+};
+
+
diff --git a/Engine/Source/Runtime/Experimental/GeometryCollectionEngine/Public/GeometryCollection/GeometryCollectionObject.h b/Engine/Source/Runtime/Experimental/GeometryCollectionEngine/Public/GeometryCollection/GeometryCollectionObject.h
--- a/Engine/Source/Runtime/Experimental/GeometryCollectionEngine/Public/GeometryCollection/GeometryCollectionObject.h
+++ b/Engine/Source/Runtime/Experimental/GeometryCollectionEngine/Public/GeometryCollection/GeometryCollectionObject.h
@@ -10,6 +10,7 @@
 #include "InstanceUniformShaderParameters.h"
 #include "Interfaces/Interface_AssetUserData.h"
 #include "Misc/Crc.h"
+#include "Containers/Map.h"
 
 #include "GeometryCollectionObject.generated.h"
 
@@ -781,6 +782,8 @@ public:
 	UPROPERTY(EditAnywhere, Category = "Dataflow")
 	FString DataflowTerminal = "GeometryCollectionTerminal";
 
+	UPROPERTY(EditAnywhere, Category = "Dataflow")
+	TMap<FString, FString> Overrides;
 
 private:
 #if WITH_EDITOR
