Dataflow : add extra outputs to override nodes
- Add IsOverriden and typed Default value
- Expose Geometry collection Overrides map to blueprint
- Change display name of a Overrides map to DataFlow Overrides

diff --git a/Engine/Plugins/Experimental/GeometryCollectionPlugin/Source/GeometryCollectionNodes/Private/Dataflow/GeometryCollectionOverrideNodes.cpp b/Engine/Plugins/Experimental/GeometryCollectionPlugin/Source/GeometryCollectionNodes/Private/Dataflow/GeometryCollectionOverrideNodes.cpp
--- a/Engine/Plugins/Experimental/GeometryCollectionPlugin/Source/GeometryCollectionNodes/Private/Dataflow/GeometryCollectionOverrideNodes.cpp
+++ b/Engine/Plugins/Experimental/GeometryCollectionPlugin/Source/GeometryCollectionNodes/Private/Dataflow/GeometryCollectionOverrideNodes.cpp
@@ -26,9 +26,11 @@ namespace Dataflow
 
 void FGetFloatOverrideFromAssetDataflowNode::Evaluate(Dataflow::FContext& Context, const FDataflowOutput* Out) const
 {
-	if (Out->IsA(&Float))
+	if (Out->IsA(&Float) || Out->IsA(&FloatDefault) || Out->IsA(&IsOverriden))
 	{
-		float NewValue = FCString::Atof(*GetDefaultValue(Context));
+		const float DefaultValue = FCString::Atof(*GetDefaultValue(Context));
+		float NewValue = DefaultValue;
+		bool bIsOverriden = false;
 
 		if (const Dataflow::FEngineContext* EngineContext = Context.AsType<Dataflow::FEngineContext>())
 		{
@@ -37,18 +39,23 @@ void FGetFloatOverrideFromAssetDataflowNode::Evaluate(Dataflow::FContext& Contex
 			if (!ValueFromAsset.IsEmpty() && ValueFromAsset.IsNumeric())
 			{
 				NewValue = FCString::Atof(*ValueFromAsset);
+				bIsOverriden = true;
 			}
 		}
 
 		SetValue(Context, NewValue, &Float);
+		SetValue(Context, DefaultValue, &FloatDefault);
+		SetValue(Context, bIsOverriden, &IsOverriden);
 	}
 }
 
 void FGetIntOverrideFromAssetDataflowNode::Evaluate(Dataflow::FContext& Context, const FDataflowOutput* Out) const
 {
-	if (Out->IsA(&Int))
+	if (Out->IsA(&Int) || Out->IsA(&IntDefault) || Out->IsA(&IsOverriden))
 	{
-		int32 NewValue = FCString::Atoi(*GetDefaultValue(Context));
+		const int32 DefaultValue = FCString::Atoi(*GetDefaultValue(Context));
+		int32 NewValue = DefaultValue;
+		bool bIsOverriden = false;
 
 		if (const Dataflow::FEngineContext* EngineContext = Context.AsType<Dataflow::FEngineContext>())
 		{
@@ -57,10 +64,13 @@ void FGetIntOverrideFromAssetDataflowNode::Evaluate(Dataflow::FContext& Context,
 			if (!ValueFromAsset.IsEmpty() && ValueFromAsset.IsNumeric())
 			{
 				NewValue = FCString::Atoi(*ValueFromAsset);
+				bIsOverriden = true;
 			}
 		}
 
 		SetValue(Context, NewValue, &Int);
+		SetValue(Context, DefaultValue, &IntDefault);
+		SetValue(Context, bIsOverriden, &IsOverriden);
 	}
 }
 
@@ -85,9 +95,11 @@ static bool StringToBool(const FString& InString, bool InDefault)
 
 void FGetBoolOverrideFromAssetDataflowNode::Evaluate(Dataflow::FContext& Context, const FDataflowOutput* Out) const
 {
-	if (Out->IsA(&Bool))
+	if (Out->IsA(&Bool) || Out->IsA(&BoolDefault) || Out->IsA(&IsOverriden))
 	{
-		bool NewValue = StringToBool(GetDefaultValue(Context), false);
+		const bool DefaultValue = StringToBool(GetDefaultValue(Context), false);
+		bool NewValue = DefaultValue;
+		bool bIsOverriden = false;
 
 		if (const Dataflow::FEngineContext* EngineContext = Context.AsType<Dataflow::FEngineContext>())
 		{
@@ -96,18 +108,23 @@ void FGetBoolOverrideFromAssetDataflowNode::Evaluate(Dataflow::FContext& Context
 			if (!ValueFromAsset.IsEmpty())
 			{
 				NewValue = StringToBool(ValueFromAsset, false);
+				bIsOverriden = true;
 			}
 		}
 
 		SetValue(Context, NewValue, &Bool);
+		SetValue(Context, DefaultValue, &BoolDefault);
+		SetValue(Context, bIsOverriden, &IsOverriden);
 	}
 }
 
 void FGetStringOverrideFromAssetDataflowNode::Evaluate(Dataflow::FContext& Context, const FDataflowOutput* Out) const
 {
-	if (Out->IsA(&String))
+	if (Out->IsA(&String) || Out->IsA(&StringDefault) || Out->IsA(&IsOverriden))
 	{
-		FString NewValue = GetDefaultValue(Context);
+		const FString DefaultValue = GetDefaultValue(Context);
+		FString NewValue = DefaultValue;
+		bool bIsOverriden = false;
 
 		if (const Dataflow::FEngineContext* EngineContext = Context.AsType<Dataflow::FEngineContext>())
 		{
@@ -116,9 +133,12 @@ void FGetStringOverrideFromAssetDataflowNode::Evaluate(Dataflow::FContext& Conte
 			if (!ValueFromAsset.IsEmpty())
 			{
 				NewValue = ValueFromAsset;
+				bIsOverriden = true;
 			}
 		}
 
 		SetValue(Context, NewValue, &String);
+		SetValue(Context, DefaultValue, &StringDefault);
+		SetValue(Context, bIsOverriden, &IsOverriden);
 	}
 }
diff --git a/Engine/Plugins/Experimental/GeometryCollectionPlugin/Source/GeometryCollectionNodes/Public/Dataflow/GeometryCollectionOverrideNodes.h b/Engine/Plugins/Experimental/GeometryCollectionPlugin/Source/GeometryCollectionNodes/Public/Dataflow/GeometryCollectionOverrideNodes.h
--- a/Engine/Plugins/Experimental/GeometryCollectionPlugin/Source/GeometryCollectionNodes/Public/Dataflow/GeometryCollectionOverrideNodes.h
+++ b/Engine/Plugins/Experimental/GeometryCollectionPlugin/Source/GeometryCollectionNodes/Public/Dataflow/GeometryCollectionOverrideNodes.h
@@ -23,10 +23,14 @@ public:
 	UPROPERTY(meta = (DataflowOutput))
 	float Float = 0.f;
 
+	UPROPERTY(meta = (DataflowOutput))
+	float FloatDefault = 0.f;
+
 	FGetFloatOverrideFromAssetDataflowNode(const Dataflow::FNodeParameters& InParam, FGuid InGuid = FGuid::NewGuid())
 		: FDataflowOverrideNode(InParam, InGuid)
 	{
 		RegisterOutputConnection(&Float);
+		RegisterOutputConnection(&FloatDefault);
 	}
 
 	virtual void Evaluate(Dataflow::FContext& Context, const FDataflowOutput* Out) const override;
@@ -47,10 +51,14 @@ public:
 	UPROPERTY(meta = (DataflowOutput))
 	int32 Int = 0;
 
+	UPROPERTY(meta = (DataflowOutput))
+	int32 IntDefault = 0;
+
 	FGetIntOverrideFromAssetDataflowNode(const Dataflow::FNodeParameters& InParam, FGuid InGuid = FGuid::NewGuid())
 		: FDataflowOverrideNode(InParam, InGuid)
 	{
 		RegisterOutputConnection(&Int);
+		RegisterOutputConnection(&IntDefault);
 	}
 
 	virtual void Evaluate(Dataflow::FContext& Context, const FDataflowOutput* Out) const override;
@@ -71,10 +79,14 @@ public:
 	UPROPERTY(meta = (DataflowOutput))
 	bool Bool = false;
 
+	UPROPERTY(meta = (DataflowOutput))
+	bool BoolDefault = false;
+
 	FGetBoolOverrideFromAssetDataflowNode(const Dataflow::FNodeParameters& InParam, FGuid InGuid = FGuid::NewGuid())
 		: FDataflowOverrideNode(InParam, InGuid)
 	{
-	RegisterOutputConnection(&Bool);
+		RegisterOutputConnection(&Bool);
+		RegisterOutputConnection(&BoolDefault);
 	}
 
 	virtual void Evaluate(Dataflow::FContext& Context, const FDataflowOutput* Out) const override;
@@ -95,10 +107,14 @@ public:
 	UPROPERTY(meta = (DataflowOutput))
 	FString String;
 
+	UPROPERTY(meta = (DataflowOutput))
+	FString StringDefault;
+
 	FGetStringOverrideFromAssetDataflowNode(const Dataflow::FNodeParameters& InParam, FGuid InGuid = FGuid::NewGuid())
 		: FDataflowOverrideNode(InParam, InGuid)
 	{
 		RegisterOutputConnection(&String);
+		RegisterOutputConnection(&StringDefault);
 	}
 
 	virtual void Evaluate(Dataflow::FContext& Context, const FDataflowOutput* Out) const override;
diff --git a/Engine/Source/Runtime/Experimental/Dataflow/Core/Private/Dataflow/DataflowOverrideNode.cpp b/Engine/Source/Runtime/Experimental/Dataflow/Core/Private/Dataflow/DataflowOverrideNode.cpp
--- a/Engine/Source/Runtime/Experimental/Dataflow/Core/Private/Dataflow/DataflowOverrideNode.cpp
+++ b/Engine/Source/Runtime/Experimental/Dataflow/Core/Private/Dataflow/DataflowOverrideNode.cpp
@@ -46,15 +46,15 @@ FString FDataflowOverrideNode::GetDefaultValue(Dataflow::FContext& Context) cons
 FString FDataflowOverrideNode::GetValueFromAsset(Dataflow::FContext& Context, const UObject* InOwner) const
 {
 	FName InKey = GetValue<FName>(Context, &Key, Key);
-	FString InDefault = GetValue<FString>(Context, &Default, Default);
+	FString EmptyString;
 
 	if (InOwner)
 	{
 		if (!InKey.IsNone() && InKey.IsValid() && InKey.ToString().Len() > 0)
 		{
-			return FindOverrideMapProperty<FString>(InOwner, FName("Overrides"), InKey, InDefault);
+			return FindOverrideMapProperty<FString>(InOwner, FName("Overrides"), InKey, EmptyString);
 		}
 	}
 
-	return InDefault;
+	return EmptyString;
 }
diff --git a/Engine/Source/Runtime/Experimental/Dataflow/Core/Public/Dataflow/DataflowOverrideNode.h b/Engine/Source/Runtime/Experimental/Dataflow/Core/Public/Dataflow/DataflowOverrideNode.h
--- a/Engine/Source/Runtime/Experimental/Dataflow/Core/Public/Dataflow/DataflowOverrideNode.h
+++ b/Engine/Source/Runtime/Experimental/Dataflow/Core/Public/Dataflow/DataflowOverrideNode.h
@@ -33,6 +33,7 @@ public:
 	{
 		RegisterInputConnection(&Key);
 		RegisterInputConnection(&Default);
+		RegisterOutputConnection(&IsOverriden);
 	}
 
 	virtual ~FDataflowOverrideNode() { }
@@ -54,6 +55,9 @@ public:
 
 	UPROPERTY(EditAnywhere, Category = "Overrides", meta = (DataflowInput))
 	FString Default = FString("0");
+
+	UPROPERTY(meta = (DataflowOutput))
+	bool IsOverriden = false;
 };
 
 
diff --git a/Engine/Source/Runtime/Experimental/GeometryCollectionEngine/Public/GeometryCollection/GeometryCollectionObject.h b/Engine/Source/Runtime/Experimental/GeometryCollectionEngine/Public/GeometryCollection/GeometryCollectionObject.h
--- a/Engine/Source/Runtime/Experimental/GeometryCollectionEngine/Public/GeometryCollection/GeometryCollectionObject.h
+++ b/Engine/Source/Runtime/Experimental/GeometryCollectionEngine/Public/GeometryCollection/GeometryCollectionObject.h
@@ -783,7 +783,7 @@ public:
 	UPROPERTY(EditAnywhere, Category = "Dataflow")
 	FString DataflowTerminal = "GeometryCollectionTerminal";
 
-	UPROPERTY(EditAnywhere, Category = "Dataflow")
+	UPROPERTY(EditAnywhere, BlueprintReadWrite, Category = "Dataflow", DisplayName = "DataFlow Overrides")
 	TMap<FString, FString> Overrides;
 
 private:
