add dataflow node to get the sum of convex hull volumes for a given selection

diff --git a/Engine/Plugins/Experimental/GeometryCollectionPlugin/Source/GeometryCollectionNodes/Private/Dataflow/GeometryCollectionUtilityNodes.cpp b/Engine/Plugins/Experimental/GeometryCollectionPlugin/Source/GeometryCollectionNodes/Private/Dataflow/GeometryCollectionUtilityNodes.cpp
--- a/Engine/Plugins/Experimental/GeometryCollectionPlugin/Source/GeometryCollectionNodes/Private/Dataflow/GeometryCollectionUtilityNodes.cpp
+++ b/Engine/Plugins/Experimental/GeometryCollectionPlugin/Source/GeometryCollectionNodes/Private/Dataflow/GeometryCollectionUtilityNodes.cpp
@@ -6,6 +6,8 @@
 #include "GeometryCollection/ManagedArrayCollection.h"
 #include "GeometryCollection/GeometryCollection.h"
 #include "GeometryCollection/GeometryCollectionConvexUtility.h"
+#include "GeometryCollection/Facades/CollectionHierarchyFacade.h"
+#include "GeometryCollection/Facades/CollectionTransformSelectionFacade.h"
 
 #include "FractureEngineConvex.h"
 
@@ -27,6 +29,7 @@ namespace Dataflow
 		DATAFLOW_NODE_REGISTER_CREATION_FACTORY(FGenerateClusterConvexHullsFromLeafHullsDataflowNode);
 		DATAFLOW_NODE_REGISTER_CREATION_FACTORY(FGenerateClusterConvexHullsFromChildrenHullsDataflowNode);
 		DATAFLOW_NODE_REGISTER_CREATION_FACTORY(FUpdateVolumeAttributesDataflowNode);
+		DATAFLOW_NODE_REGISTER_CREATION_FACTORY(FGetConvexHullVolumeDataflowNode);
 	}
 }
 
@@ -310,3 +313,66 @@ void FUpdateVolumeAttributesDataflowNode::Evaluate(Dataflow::FContext& Context,
 		SetValue(Context, MoveTemp(InCollection), &Collection);
 	}
 }
+
+
+FGetConvexHullVolumeDataflowNode::FGetConvexHullVolumeDataflowNode(const Dataflow::FNodeParameters& InParam, FGuid InGuid)
+	: FDataflowNode(InParam, InGuid)
+{
+	RegisterInputConnection(&Collection);
+	RegisterInputConnection(&TransformSelection);
+	RegisterOutputConnection(&Volume);
+}
+
+void FGetConvexHullVolumeDataflowNode::Evaluate(Dataflow::FContext& Context, const FDataflowOutput* Out) const
+{
+	if (Out->IsA(&Volume))
+	{
+		float VolumeSum = 0;
+
+		if (!IsConnected(&Collection) || !IsConnected(&TransformSelection))
+		{
+			SetValue(Context, VolumeSum, &Volume);
+			return;
+		}
+
+		const FManagedArrayCollection& InCollection = GetValue(Context, &Collection);
+		const FDataflowTransformSelection& InSelection = GetValue(Context, &TransformSelection);
+
+		if (!FGeometryCollectionConvexUtility::HasConvexHullData(&InCollection))
+		{
+			SetValue(Context, VolumeSum, &Volume);
+			return;
+		}
+
+		const int32 NumTransforms = InCollection.NumElements(FGeometryCollection::TransformGroup);
+		TArray<int32> SelectionToSum = InSelection.AsArray();
+		GeometryCollection::Facades::FCollectionTransformSelectionFacade SelectionFacade(InCollection);
+		SelectionFacade.Sanitize(SelectionToSum);
+		if (NumTransforms == 0 || SelectionToSum.Num() == 0)
+		{
+			SetValue(Context, VolumeSum, &Volume);
+			return;
+		}
+
+		const TManagedArray<TSet<int32>>& TransformToConvexIndices = InCollection.GetAttribute<TSet<int32>>("TransformToConvexIndices", FTransformCollection::TransformGroup);
+		const TManagedArray<TUniquePtr<Chaos::FConvex>>& ConvexHulls = InCollection.GetAttribute<TUniquePtr<Chaos::FConvex>>("ConvexHull", "Convex");
+
+		Chaos::Facades::FCollectionHierarchyFacade HierarchyFacade(InCollection);
+		while (!SelectionToSum.IsEmpty())
+		{
+			int32 TransformIdx = SelectionToSum.Pop(false);
+			if (!bSumChildrenForClustersWithoutHulls || !TransformToConvexIndices[TransformIdx].IsEmpty())
+			{
+				for (int32 ConvexIdx : TransformToConvexIndices[TransformIdx])
+				{
+					VolumeSum += ConvexHulls[ConvexIdx]->GetVolume();
+				}
+			}
+			else if (const TSet<int32>* Children = HierarchyFacade.FindChildren(TransformIdx))
+			{
+				SelectionToSum.Append(Children->Array());
+			}
+		}
+		SetValue(Context, VolumeSum, &Volume);
+	}
+}
diff --git a/Engine/Plugins/Experimental/GeometryCollectionPlugin/Source/GeometryCollectionNodes/Public/Dataflow/GeometryCollectionUtilityNodes.h b/Engine/Plugins/Experimental/GeometryCollectionPlugin/Source/GeometryCollectionNodes/Public/Dataflow/GeometryCollectionUtilityNodes.h
--- a/Engine/Plugins/Experimental/GeometryCollectionPlugin/Source/GeometryCollectionNodes/Public/Dataflow/GeometryCollectionUtilityNodes.h
+++ b/Engine/Plugins/Experimental/GeometryCollectionPlugin/Source/GeometryCollectionNodes/Public/Dataflow/GeometryCollectionUtilityNodes.h
@@ -277,6 +277,36 @@ public:
 	virtual void Evaluate(Dataflow::FContext& Context, const FDataflowOutput* Out) const override;
 };
 
+/**
+ * Get the sum of volumes of the convex hulls on the selected nodes
+ */
+USTRUCT(meta = (DataflowGeometryCollection))
+struct FGetConvexHullVolumeDataflowNode : public FDataflowNode
+{
+	GENERATED_USTRUCT_BODY()
+	DATAFLOW_NODE_DEFINE_INTERNAL(FGetConvexHullVolumeDataflowNode, "GetConvexHullVolume", "GeometryCollection|Utilities", "")
+
+public:
+	UPROPERTY(meta = (DataflowInput))
+	FManagedArrayCollection Collection;
+
+	/** The transforms to consider */
+	UPROPERTY(meta = (DataflowInput, DataflowIntrinsic))
+	FDataflowTransformSelection TransformSelection;
+
+	/** Sum of convex hull volumes */
+	UPROPERTY(meta = (DataflowOutput));
+	float Volume;
+
+	/** For any cluster transform that has no convex hulls, whether to fall back to the convex hulls of the cluster's children. Otherwise, the cluster will not add to the total volume sum. */
+	UPROPERTY(EditAnywhere, Category = Options)
+	bool bSumChildrenForClustersWithoutHulls = true;
+
+	FGetConvexHullVolumeDataflowNode(const Dataflow::FNodeParameters& InParam, FGuid InGuid = FGuid::NewGuid());
+
+	virtual void Evaluate(Dataflow::FContext& Context, const FDataflowOutput* Out) const override;
+};
+
 
 namespace Dataflow
 {
