Add a default value for FGetConvexHullVolumeDataflowNode::Volume

diff --git a/Engine/Plugins/Experimental/GeometryCollectionPlugin/Source/GeometryCollectionNodes/Public/Dataflow/GeometryCollectionUtilityNodes.h b/Engine/Plugins/Experimental/GeometryCollectionPlugin/Source/GeometryCollectionNodes/Public/Dataflow/GeometryCollectionUtilityNodes.h
--- a/Engine/Plugins/Experimental/GeometryCollectionPlugin/Source/GeometryCollectionNodes/Public/Dataflow/GeometryCollectionUtilityNodes.h
+++ b/Engine/Plugins/Experimental/GeometryCollectionPlugin/Source/GeometryCollectionNodes/Public/Dataflow/GeometryCollectionUtilityNodes.h
@@ -296,7 +296,7 @@ public:
 
 	/** Sum of convex hull volumes */
 	UPROPERTY(meta = (DataflowOutput));
-	float Volume;
+	float Volume = 0.f;
 
 	/** For any cluster transform that has no convex hulls, whether to fall back to the convex hulls of the cluster's children. Otherwise, the cluster will not add to the total volume sum. */
 	UPROPERTY(EditAnywhere, Category = Options)
