Dataflow:
- Updated GetDefaultValue() to handle int32, float special cases

diff --git a/Engine/Plugins/Experimental/GeometryCollectionPlugin/Source/GeometryCollectionNodes/Private/Dataflow/GeometryCollectionOverrideNodes.cpp b/Engine/Plugins/Experimental/GeometryCollectionPlugin/Source/GeometryCollectionNodes/Private/Dataflow/GeometryCollectionOverrideNodes.cpp
--- a/Engine/Plugins/Experimental/GeometryCollectionPlugin/Source/GeometryCollectionNodes/Private/Dataflow/GeometryCollectionOverrideNodes.cpp
+++ b/Engine/Plugins/Experimental/GeometryCollectionPlugin/Source/GeometryCollectionNodes/Private/Dataflow/GeometryCollectionOverrideNodes.cpp
@@ -28,7 +28,7 @@ void FGetFloatOverrideFromAssetDataflowNode::Evaluate(Dataflow::FContext& Contex
 {
 	if (Out->IsA(&Float) || Out->IsA(&FloatDefault) || Out->IsA(&IsOverriden))
 	{
-		const float DefaultValue = FCString::Atof(*GetDefaultValue(Context));
+		const float DefaultValue = GetDefaultValue<float>(Context);
 		float NewValue = DefaultValue;
 		bool bIsOverriden = false;
 
@@ -53,7 +53,7 @@ void FGetIntOverrideFromAssetDataflowNode::Evaluate(Dataflow::FContext& Context,
 {
 	if (Out->IsA(&Int) || Out->IsA(&IntDefault) || Out->IsA(&IsOverriden))
 	{
-		const int32 DefaultValue = FCString::Atoi(*GetDefaultValue(Context));
+		const int32 DefaultValue = GetDefaultValue<int32>(Context);
 		int32 NewValue = DefaultValue;
 		bool bIsOverriden = false;
 
@@ -97,7 +97,7 @@ void FGetBoolOverrideFromAssetDataflowNode::Evaluate(Dataflow::FContext& Context
 {
 	if (Out->IsA(&Bool) || Out->IsA(&BoolDefault) || Out->IsA(&IsOverriden))
 	{
-		const bool DefaultValue = StringToBool(GetDefaultValue(Context), false);
+		const bool DefaultValue = StringToBool(GetDefaultValue<FString>(Context), false);
 		bool NewValue = DefaultValue;
 		bool bIsOverriden = false;
 
@@ -122,7 +122,7 @@ void FGetStringOverrideFromAssetDataflowNode::Evaluate(Dataflow::FContext& Conte
 {
 	if (Out->IsA(&String) || Out->IsA(&StringDefault) || Out->IsA(&IsOverriden))
 	{
-		const FString DefaultValue = GetDefaultValue(Context);
+		const FString DefaultValue = GetDefaultValue<FString>(Context);
 		FString NewValue = DefaultValue;
 		bool bIsOverriden = false;
 
diff --git a/Engine/Source/Runtime/Experimental/Dataflow/Core/Private/Dataflow/DataflowOverrideNode.cpp b/Engine/Source/Runtime/Experimental/Dataflow/Core/Private/Dataflow/DataflowOverrideNode.cpp
--- a/Engine/Source/Runtime/Experimental/Dataflow/Core/Private/Dataflow/DataflowOverrideNode.cpp
+++ b/Engine/Source/Runtime/Experimental/Dataflow/Core/Private/Dataflow/DataflowOverrideNode.cpp
@@ -38,9 +38,32 @@ bool FDataflowOverrideNode::ShouldInvalidate(FName InKey) const
 	return InKey == Key;
 }
 
-FString FDataflowOverrideNode::GetDefaultValue(Dataflow::FContext& Context) const
+template <>
+int32 FDataflowOverrideNode::GetDefaultValue(Dataflow::FContext& Context) const
 {
-	return GetValue<FString>(Context, &Default, Default);
+	const FString DefaultValue = GetValue<FString>(Context, &Default, Default);
+	int32 DefaultIntValue = 0;
+
+	if (!DefaultValue.IsEmpty() && DefaultValue.IsNumeric())
+	{
+		DefaultIntValue = FCString::Atoi(*DefaultValue);
+	}
+
+	return DefaultIntValue;
+}
+
+template <>
+float FDataflowOverrideNode::GetDefaultValue(Dataflow::FContext& Context) const
+{
+	const FString DefaultValue = GetValue<FString>(Context, &Default, Default);
+	float DefaultFloatValue = 0.f;
+
+	if (!DefaultValue.IsEmpty() && DefaultValue.IsNumeric())
+	{
+		DefaultFloatValue = FCString::Atof(*DefaultValue);
+	}
+
+	return DefaultFloatValue;
 }
 
 FString FDataflowOverrideNode::GetValueFromAsset(Dataflow::FContext& Context, const UObject* InOwner) const
diff --git a/Engine/Source/Runtime/Experimental/Dataflow/Core/Public/Dataflow/DataflowOverrideNode.h b/Engine/Source/Runtime/Experimental/Dataflow/Core/Public/Dataflow/DataflowOverrideNode.h
--- a/Engine/Source/Runtime/Experimental/Dataflow/Core/Public/Dataflow/DataflowOverrideNode.h
+++ b/Engine/Source/Runtime/Experimental/Dataflow/Core/Public/Dataflow/DataflowOverrideNode.h
@@ -40,7 +40,17 @@ public:
 
 	bool ShouldInvalidate(FName InKey) const;
 
-	FString GetDefaultValue(Dataflow::FContext& Context) const;
+	template <class T>
+	T GetDefaultValue(Dataflow::FContext& Context) const
+	{
+		return GetValue<T>(Context, &Default, Default);
+	}
+
+	template <>
+	int32 GetDefaultValue(Dataflow::FContext& Context) const;
+
+	template <>
+	float GetDefaultValue(Dataflow::FContext& Context) const;
 
 	FString GetValueFromAsset(Dataflow::FContext& Context, const UObject* InOwner) const;
 
