Dataflow:
- Added FloatArrayComputeStatistics node to compute Min/Max/Mean/Median/Mode/Sum values of a FloatArray (leaf size)
- Added RemoveFloatArrayElementDataflowNode node
- Updated FGetNumArrayElementsDataflowNode to handle float and int arrays as well

diff --git a/Engine/Plugins/Experimental/GeometryCollectionPlugin/Source/GeometryCollectionNodes/Private/Dataflow/GeometryCollectionArrayNodes.cpp b/Engine/Plugins/Experimental/GeometryCollectionPlugin/Source/GeometryCollectionNodes/Private/Dataflow/GeometryCollectionArrayNodes.cpp
new file mode 100644
--- /dev/null
+++ b/Engine/Plugins/Experimental/GeometryCollectionPlugin/Source/GeometryCollectionNodes/Private/Dataflow/GeometryCollectionArrayNodes.cpp
@@ -0,0 +1,588 @@
+// Copyright Epic Games, Inc. All Rights Reserved.
+
+#include "Dataflow/GeometryCollectionArrayNodes.h"
+#include "Dataflow/DataflowCore.h"
+
+
+//#include UE_INLINE_GENERATED_CPP_BY_NAME(GeometryCollectionArrayNodes)
+
+namespace Dataflow
+{
+
+	void GeometryCollectionArrayNodes()
+	{
+		static const FLinearColor CDefaultNodeBodyTintColor = FLinearColor(0.f, 0.f, 0.f, 0.5f);
+
+		DATAFLOW_NODE_REGISTER_CREATION_FACTORY(FGetFloatArrayElementDataflowNode);
+		DATAFLOW_NODE_REGISTER_CREATION_FACTORY(FFloatArrayToIntArrayDataflowNode);
+		DATAFLOW_NODE_REGISTER_CREATION_FACTORY(FGetArrayElementDataflowNode);
+		DATAFLOW_NODE_REGISTER_CREATION_FACTORY(FGetNumArrayElementsDataflowNode);
+		DATAFLOW_NODE_REGISTER_CREATION_FACTORY(FBoolArrayToFaceSelectionDataflowNode);
+		DATAFLOW_NODE_REGISTER_CREATION_FACTORY(FFloatArrayToVertexSelectionDataflowNode);
+		DATAFLOW_NODE_REGISTER_CREATION_FACTORY(FFloatArrayNormalizeDataflowNode);
+		DATAFLOW_NODE_REGISTER_CREATION_FACTORY(FVectorArrayNormalizeDataflowNode);
+		DATAFLOW_NODE_REGISTER_CREATION_FACTORY(FUnionIntArraysDataflowNode);
+		DATAFLOW_NODE_REGISTER_CREATION_FACTORY(FRemoveFloatArrayElementDataflowNode);
+		DATAFLOW_NODE_REGISTER_CREATION_FACTORY(FFloatArrayComputeStatisticsDataflowNode);
+
+	}
+}
+
+void FGetFloatArrayElementDataflowNode::Evaluate(Dataflow::FContext& Context, const FDataflowOutput* Out) const
+{
+	if (Out->IsA(&FloatValue))
+	{
+		const TArray<float>& InArray = GetValue(Context, &FloatArray);
+		const int32& InIndex = GetValue(Context, &Index);
+		const float OutputValue = InArray.IsValidIndex(InIndex) ? InArray[InIndex] : 0.0f;
+
+		SetValue(Context, OutputValue, &FloatValue);
+	}
+}
+
+void FFloatArrayToIntArrayDataflowNode::Evaluate(Dataflow::FContext& Context, const FDataflowOutput* Out) const
+{
+	if (Out->IsA<TArray<int32>>(&IntArray))
+	{
+		TArray<float> FloatVal = GetValue<TArray<float>>(Context, &FloatArray);
+		TArray<int32> RetVal; RetVal.SetNumUninitialized(FloatVal.Num());
+		if (Function == EFloatArrayToIntArrayFunctionEnum::Dataflow_FloatToInt_Function_Floor)
+		{
+			for (int32 i = 0; i < FloatVal.Num(); i++)
+			{
+				RetVal[i] = FMath::FloorToInt32(FloatVal[i]);
+			}
+		}
+		else if (Function == EFloatArrayToIntArrayFunctionEnum::Dataflow_FloatToInt_Function_Ceil)
+		{
+			for (int32 i = 0; i < FloatVal.Num(); i++)
+			{
+				RetVal[i] = FMath::CeilToInt32(FloatVal[i]);
+			}
+		}
+		else if (Function == EFloatArrayToIntArrayFunctionEnum::Dataflow_FloatToInt_Function_Round)
+		{
+			for (int32 i = 0; i < FloatVal.Num(); i++)
+			{
+				RetVal[i] = FMath::RoundToInt32(FloatVal[i]);
+			}
+		}
+		else if (Function == EFloatArrayToIntArrayFunctionEnum::Dataflow_FloatToInt_Function_Truncate)
+		{
+			for (int32 i = 0; i < FloatVal.Num(); i++)
+			{
+				RetVal[i] = int32(FMath::TruncToFloat(FloatVal[i]));
+			}
+		}
+
+		else if (Function == EFloatArrayToIntArrayFunctionEnum::Dataflow_FloatToInt_NonZeroToIndex)
+		{
+			int32 RetValIndex = 0;
+			for (int32 i = 0; i < FloatVal.Num(); i++)
+			{
+				if (FloatVal[i] != 0.0)
+				{
+					RetVal[RetValIndex] = i;
+					RetValIndex++;
+				}
+			}
+			RetVal.SetNum(RetValIndex);
+		}
+		else if (Function == EFloatArrayToIntArrayFunctionEnum::Dataflow_FloatToInt_ZeroToIndex)
+		{
+			int32 RetValIndex = 0;
+			for (int32 i = 0; i < FloatVal.Num(); i++)
+			{
+				if (FloatVal[i] == 0.0)
+				{
+					RetVal[RetValIndex] = i;
+					RetValIndex++;
+				}
+			}
+			RetVal.SetNum(RetValIndex);
+		}
+
+		SetValue<TArray<int32>>(Context, RetVal, &IntArray);
+	}
+}
+
+void FGetArrayElementDataflowNode::Evaluate(Dataflow::FContext& Context, const FDataflowOutput* Out) const
+{
+	if (Out->IsA<FVector>(&Point))
+	{
+		const TArray<FVector>& Array = GetValue<TArray<FVector>>(Context, &Points);
+		if (Array.Num() > 0 && Index >= 0 && Index < Array.Num())
+		{
+			SetValue<FVector>(Context, Array[Index], &Point);
+			return;
+		}
+
+		SetValue<FVector>(Context, FVector(0.f), &Point);
+	}
+}
+
+void FGetNumArrayElementsDataflowNode::Evaluate(Dataflow::FContext& Context, const FDataflowOutput* Out) const
+{
+	if (Out->IsA<int32>(&NumElements))
+	{
+		if (IsConnected<TArray<float>>(&FloatArray))
+		{
+			SetValue<int32>(Context, GetValue<TArray<float>>(Context, &FloatArray).Num(), &NumElements);
+			return;
+		}
+		else if (IsConnected<TArray<int32>>(&IntArray))
+		{
+			SetValue<int32>(Context, GetValue<TArray<int32>>(Context, &IntArray).Num(), &NumElements);
+			return;
+		}
+		else if (IsConnected<TArray<FVector>>(&Points))
+		{
+			SetValue<int32>(Context, GetValue<TArray<FVector>>(Context, &Points).Num(), &NumElements);
+			return;
+		}
+		else if (IsConnected<TArray<FVector3f>>(&Vector3fArray))
+		{
+			SetValue<int32>(Context, GetValue<TArray<FVector3f>>(Context, &Vector3fArray).Num(), &NumElements);
+			return;
+		}
+
+		SetValue<int32>(Context, 0, &NumElements);
+	}
+}
+
+void FBoolArrayToFaceSelectionDataflowNode::Evaluate(Dataflow::FContext& Context, const FDataflowOutput* Out) const
+{
+	if (Out->IsA<FDataflowFaceSelection>(&FaceSelection))
+	{
+		const TArray<bool>& InBoolAttributeData = GetValue<TArray<bool>>(Context, &BoolAttributeData);
+
+		FDataflowFaceSelection NewFaceSelection;
+		NewFaceSelection.Initialize(InBoolAttributeData.Num(), false);
+		NewFaceSelection.SetFromArray(InBoolAttributeData);
+
+		SetValue<FDataflowFaceSelection>(Context, NewFaceSelection, &FaceSelection);
+	}
+}
+
+
+void FFloatArrayToVertexSelectionDataflowNode::Evaluate(Dataflow::FContext& Context, const FDataflowOutput* Out) const
+{
+	if (Out->IsA<FDataflowVertexSelection>(&VertexSelection))
+	{
+		const TArray<float>& InFloatArray = GetValue<TArray<float>>(Context, &FloatArray);
+
+		FDataflowVertexSelection NewVertexSelection;
+		NewVertexSelection.Initialize(InFloatArray.Num(), false);
+
+		for (int32 Idx = 0; Idx < InFloatArray.Num(); ++Idx)
+		{
+			if (Operation == ECompareOperation1Enum::Dataflow_Compare_Equal)
+			{
+				if (InFloatArray[Idx] == Threshold)
+				{
+					NewVertexSelection.SetSelected(Idx);
+				}
+			}
+			else if (Operation == ECompareOperation1Enum::Dataflow_Compare_Smaller)
+			{
+				if (InFloatArray[Idx] < Threshold)
+				{
+					NewVertexSelection.SetSelected(Idx);
+				}
+			}
+			else if (Operation == ECompareOperation1Enum::Dataflow_Compare_SmallerOrEqual)
+			{
+				if (InFloatArray[Idx] <= Threshold)
+				{
+					NewVertexSelection.SetSelected(Idx);
+				}
+			}
+			else if (Operation == ECompareOperation1Enum::Dataflow_Compare_Greater)
+			{
+				if (InFloatArray[Idx] > Threshold)
+				{
+					NewVertexSelection.SetSelected(Idx);
+				}
+			}
+			else if (Operation == ECompareOperation1Enum::Dataflow_Compare_GreaterOrEqual)
+			{
+				if (InFloatArray[Idx] >= Threshold)
+				{
+					NewVertexSelection.SetSelected(Idx);
+				}
+			}
+		}
+
+		SetValue<FDataflowVertexSelection>(Context, NewVertexSelection, &VertexSelection);
+	}
+}
+
+void FFloatArrayNormalizeDataflowNode::Evaluate(Dataflow::FContext& Context, const FDataflowOutput* Out) const
+{
+	if (Out->IsA<TArray<float>>(&OutFloatArray))
+	{
+		const TArray<float>& InInFloatArray = GetValue<TArray<float>>(Context, &InFloatArray);
+		const FDataflowVertexSelection& InSelection = GetValue<FDataflowVertexSelection>(Context, &Selection);
+		const float InMinRange = GetValue<float>(Context, &MinRange);
+		const float InMaxRange = GetValue<float>(Context, &MaxRange);
+
+		const int32 NumElems = InFloatArray.Num();
+
+		TArray<float> NewFloatArray;
+		NewFloatArray.Init(false, NumElems);
+
+		if (IsConnected<FDataflowVertexSelection>(&Selection))
+		{
+			if (InInFloatArray.Num() == InSelection.Num())
+			{
+				// Compute Min/Max
+				float MinValue = 1e9, MaxValue = 1e-9;
+
+				for (int32 Idx = 0; Idx < NumElems; ++Idx)
+				{
+					if (InSelection.IsSelected(Idx))
+					{
+						if (InInFloatArray[Idx] < MinValue)
+						{
+							MinValue = InInFloatArray[Idx];
+						}
+
+						if (InInFloatArray[Idx] > MaxValue)
+						{
+							MaxValue = InInFloatArray[Idx];
+						}
+					}
+				}
+
+				for (int32 Idx = 0; Idx < NumElems; ++Idx)
+				{
+					if (InSelection.IsSelected(Idx))
+					{
+						// Normalize it
+						NewFloatArray[Idx] = (NewFloatArray[Idx] - MinValue) / (MaxValue - MinValue);
+
+						// Transform it into (MinRange, MaxRange)
+						NewFloatArray[Idx] = InMinRange + NewFloatArray[Idx] * (InMaxRange - InMinRange);
+					}
+				}
+
+				SetValue<TArray<float>>(Context, NewFloatArray, &OutFloatArray);
+				return;
+			}
+		}
+		else
+		{
+			// Compute Min/Max
+			float MinValue = 1e9, MaxValue = 1e-9;
+
+			for (int32 Idx = 0; Idx < NumElems; ++Idx)
+			{
+				if (InInFloatArray[Idx] < MinValue)
+				{
+					MinValue = InInFloatArray[Idx];
+				}
+
+				if (InInFloatArray[Idx] > MaxValue)
+				{
+					MaxValue = InInFloatArray[Idx];
+				}
+			}
+
+			for (int32 Idx = 0; Idx < NumElems; ++Idx)
+			{
+				// Normalize it
+				NewFloatArray[Idx] = (NewFloatArray[Idx] - MinValue) / (MaxValue - MinValue);
+
+				// Transform it into (MinRange, MaxRange)
+				NewFloatArray[Idx] = InMinRange + NewFloatArray[Idx] * (InMaxRange - InMinRange);
+			}
+
+			SetValue<TArray<float>>(Context, NewFloatArray, &OutFloatArray);
+			return;
+		}
+
+		SetValue<TArray<float>>(Context, TArray<float>(), &OutFloatArray);
+	}
+}
+
+
+void FVectorArrayNormalizeDataflowNode::Evaluate(Dataflow::FContext& Context, const FDataflowOutput* Out) const
+{
+	if (Out->IsA<TArray<FVector>>(&OutVectorArray))
+	{
+		const TArray<FVector>& InInVectorArray = GetValue<TArray<FVector>>(Context, &InVectorArray);
+		const FDataflowVertexSelection& InSelection = GetValue<FDataflowVertexSelection>(Context, &Selection);
+		const float InMagnitude = GetValue<float>(Context, &Magnitude);
+
+		const int32 NumElems = InInVectorArray.Num();
+
+		TArray<FVector> NewVectorArray;
+		NewVectorArray.Init(FVector(0.f), NumElems);
+
+		if (IsConnected<FDataflowVertexSelection>(&Selection))
+		{
+			if (InInVectorArray.Num() == InSelection.Num())
+			{
+				for (int32 Idx = 0; Idx < NumElems; ++Idx)
+				{
+					FVector Vector = InInVectorArray[Idx];
+
+					if (InSelection.IsSelected(Idx))
+					{
+						Vector.Normalize();
+						Vector *= InMagnitude;
+						NewVectorArray[Idx] = Vector;
+					}
+				}
+
+				SetValue<TArray<FVector>>(Context, NewVectorArray, &OutVectorArray);
+				return;
+			}
+		}
+		else
+		{
+			for (int32 Idx = 0; Idx < NumElems; ++Idx)
+			{
+				FVector Vector = InInVectorArray[Idx];
+
+				Vector.Normalize();
+				Vector *= InMagnitude;
+				NewVectorArray[Idx] = Vector;
+			}
+
+			SetValue<TArray<FVector>>(Context, NewVectorArray, &OutVectorArray);
+			return;
+		}
+
+		SetValue<TArray<FVector>>(Context, TArray<FVector>(), &OutVectorArray);
+	}
+}
+
+void FUnionIntArraysDataflowNode::Evaluate(Dataflow::FContext& Context, const FDataflowOutput* Out) const
+{
+	if (Out->IsA<TArray<int32>>(&OutArray))
+	{
+		TArray<int32> Array1 = GetValue<TArray<int32>>(Context, &InArray1);
+		TArray<int32> Array2 = GetValue<TArray<int32>>(Context, &InArray2);
+		TArray<int32> OutputArray;
+		for (int32 i = 0; i < Array1.Num(); i++)
+		{
+			OutputArray.AddUnique(Array1[i]);
+		}
+		for (int32 i = 0; i < Array2.Num(); i++)
+		{
+			OutputArray.AddUnique(Array2[i]);
+		}
+		SetValue<TArray<int32>>(Context, OutputArray, &OutArray);
+	}
+}
+
+void FRemoveFloatArrayElementDataflowNode::Evaluate(Dataflow::FContext& Context, const FDataflowOutput* Out) const
+{
+	if (Out->IsA(&FloatArray))
+	{
+		TArray<float> Array = GetValue<TArray<float>>(Context, &FloatArray);
+		if (Array.Num() > 0 && Array.IsValidIndex(Index))
+		{
+			if (bPreserveOrder)
+			{
+				Array.RemoveAt(Index);
+			}
+			else
+			{
+				Array.RemoveAtSwap(Index, 1, false);
+			}
+
+			SetValue(Context, Array, &FloatArray);
+			return;
+		}
+
+		SetValue(Context, Array, &FloatArray);
+	}
+}
+
+// -----------------------------------------------------------------------------------------------
+
+static float ComputeMin(const TArray<float>& FloatArray, TArray<int32>& IndexArr)
+{
+	float MinVal = TNumericLimits<float>::Max();
+	for (int32 Idx = 0; Idx < FloatArray.Num(); ++Idx)
+	{
+		if (FloatArray[Idx] < MinVal)
+		{
+			MinVal = FloatArray[Idx];
+		}
+	}
+
+	for (int32 Idx = 0; Idx < FloatArray.Num(); ++Idx)
+	{
+		if (FloatArray[Idx] == MinVal)
+		{
+			IndexArr.Add(Idx);
+		}
+	}
+
+	return MinVal;
+}
+
+static float ComputeMax(const TArray<float>& FloatArray, TArray<int32>& IndexArr)
+{
+	float MaxVal = TNumericLimits<float>::Min();
+	for (int32 Idx = 0; Idx < FloatArray.Num(); ++Idx)
+	{
+		if (FloatArray[Idx] > MaxVal)
+		{
+			MaxVal = FloatArray[Idx];
+		}
+	}
+
+	for (int32 Idx = 0; Idx < FloatArray.Num(); ++Idx)
+	{
+		if (FloatArray[Idx] == MaxVal)
+		{
+			IndexArr.Add(Idx);
+		}
+	}
+
+	return MaxVal;
+}
+
+static float ComputeMean(const TArray<float>& FloatArray, TArray<int32>& IndexArr)
+{
+	double SumVal = 0.f;
+
+	for (int32 Idx = 0; Idx < FloatArray.Num(); ++Idx)
+	{
+		SumVal += FloatArray[Idx];
+	}
+
+	return (float)(SumVal / (double)FloatArray.Num());
+}
+
+static float ComputeMedian(TArray<float> FloatArray, TArray<int32>& IndexArr)
+{
+	FloatArray.Sort();
+	float MedianVal;
+	int32 NumElems = FloatArray.Num();
+
+	if (NumElems % 2 == 0)
+	{
+		MedianVal = (FloatArray[NumElems / 2 - 1] + FloatArray[NumElems / 2]) / 2.f;
+
+		IndexArr.Add(NumElems / 2 - 1);
+		IndexArr.Add(NumElems / 2);
+	}
+	else
+	{
+		MedianVal = FloatArray[(NumElems - 1) / 2];
+
+		IndexArr.Add((NumElems - 1) / 2);
+	}
+
+	return MedianVal;
+}
+
+static float ComputeMode(const TArray<float>& FloatArray, TArray<int32>& IndexArr)
+{
+	float MaxVal = 0.f;
+	int32 MaxCount = 0;
+
+	for (int32 Idx1 = 0; Idx1 < FloatArray.Num(); ++Idx1)
+	{
+		int32 Count = 0;
+
+		for (int32 Idx2 = 0; Idx2 < FloatArray.Num(); ++Idx2)
+		{
+			if (FloatArray[Idx2] == FloatArray[Idx1])
+			{
+				Count++;
+			}
+
+			if (Count > MaxCount)
+			{
+				MaxCount = Count;
+				MaxVal = FloatArray[Idx1];
+			}
+		}
+	}
+
+	for (int32 Idx = 0; Idx < FloatArray.Num(); ++Idx)
+	{
+		if (FloatArray[Idx] == MaxVal)
+		{
+			IndexArr.Add(Idx);
+		}
+	}
+
+	return MaxVal;
+}
+
+static float ComputeSum(const TArray<float>& FloatArray, TArray<int32>& IndexArr)
+{
+	double SumVal = 0.f;
+
+	for (int32 Idx = 0; Idx < FloatArray.Num(); ++Idx)
+	{
+		SumVal += FloatArray[Idx];
+	}
+
+	return (float)SumVal;
+}
+
+// -----------------------------------------------------------------------------------------------
+
+void FFloatArrayComputeStatisticsDataflowNode::Evaluate(Dataflow::FContext& Context, const FDataflowOutput* Out) const
+{
+	if (Out->IsA(&Value) || Out->IsA(&Indices))
+	{
+		const TArray<float>& InFloatArray = GetValue(Context, &FloatArray);
+		const FDataflowTransformSelection& InTransformSelection = GetValue(Context, &TransformSelection);
+
+		TArray<float> Array;
+		if (IsConnected<FDataflowTransformSelection>(&TransformSelection) && InTransformSelection.Num() == InFloatArray.Num())
+		{
+			for (int32 Idx = 0; Idx < InFloatArray.Num(); ++Idx)
+			{
+				if (InTransformSelection.IsSelected(Idx))
+				{
+					Array.Add(InFloatArray[Idx]);
+				}
+			}
+		}
+		else
+		{
+			for (int32 Idx = 0; Idx < InFloatArray.Num(); ++Idx)
+			{
+				Array.Add(InFloatArray[Idx]);
+			}
+		}
+
+		float OutValue;
+		TArray<int32> OutIndices;
+
+		switch (OperationName)
+		{
+		case EStatisticsOperationEnum::Dataflow_EStatisticsOperationEnum_Min:
+			OutValue = ComputeMin(Array, OutIndices);
+			break;
+		case EStatisticsOperationEnum::Dataflow_EStatisticsOperationEnum_Max:
+			OutValue = ComputeMax(Array, OutIndices);
+			break;
+		case EStatisticsOperationEnum::Dataflow_EStatisticsOperationEnum_Mean:
+			OutValue = ComputeMean(Array, OutIndices);
+			break;
+		case EStatisticsOperationEnum::Dataflow_EStatisticsOperationEnum_Median:
+			OutValue = ComputeMedian(Array, OutIndices);
+			break;
+		case EStatisticsOperationEnum::Dataflow_EStatisticsOperationEnum_Mode:
+			OutValue = ComputeMode(Array, OutIndices);
+			break;
+		case EStatisticsOperationEnum::Dataflow_EStatisticsOperationEnum_Sum:
+			OutValue = ComputeSum(Array, OutIndices);
+			break;
+		}
+
+		SetValue(Context, OutValue, &Value);
+		SetValue(Context, OutIndices, &Indices);
+	}
+}
diff --git a/Engine/Plugins/Experimental/GeometryCollectionPlugin/Source/GeometryCollectionNodes/Private/Dataflow/GeometryCollectionNodes.cpp b/Engine/Plugins/Experimental/GeometryCollectionPlugin/Source/GeometryCollectionNodes/Private/Dataflow/GeometryCollectionNodes.cpp
--- a/Engine/Plugins/Experimental/GeometryCollectionPlugin/Source/GeometryCollectionNodes/Private/Dataflow/GeometryCollectionNodes.cpp
+++ b/Engine/Plugins/Experimental/GeometryCollectionPlugin/Source/GeometryCollectionNodes/Private/Dataflow/GeometryCollectionNodes.cpp
@@ -57,10 +57,6 @@ namespace Dataflow
 		DATAFLOW_NODE_REGISTER_CREATION_FACTORY(FStringAppendDataflowNode);
 		DATAFLOW_NODE_REGISTER_CREATION_FACTORY(FHashStringDataflowNode);
 		DATAFLOW_NODE_REGISTER_CREATION_FACTORY(FHashVectorDataflowNode);
-		DATAFLOW_NODE_REGISTER_CREATION_FACTORY(FGetFloatArrayElementDataflowNode);
-		DATAFLOW_NODE_REGISTER_CREATION_FACTORY(FFloatArrayToIntArrayDataflowNode);
-		DATAFLOW_NODE_REGISTER_CREATION_FACTORY(FGetArrayElementDataflowNode);
-		DATAFLOW_NODE_REGISTER_CREATION_FACTORY(FGetNumArrayElementsDataflowNode);
 		DATAFLOW_NODE_REGISTER_CREATION_FACTORY(FGetBoundingBoxesFromCollectionDataflowNode);
 		DATAFLOW_NODE_REGISTER_CREATION_FACTORY(FGetRootIndexFromCollectionDataflowNode);
 		DATAFLOW_NODE_REGISTER_CREATION_FACTORY(FGetCentroidsFromCollectionDataflowNode);
@@ -80,16 +76,11 @@ namespace Dataflow
 		DATAFLOW_NODE_REGISTER_CREATION_FACTORY(FGetNumElementsInCollectionGroupDataflowNode);
 		DATAFLOW_NODE_REGISTER_CREATION_FACTORY(FGetCollectionAttributeDataTypedDataflowNode);
 		DATAFLOW_NODE_REGISTER_CREATION_FACTORY(FSetCollectionAttributeDataTypedDataflowNode);
-		DATAFLOW_NODE_REGISTER_CREATION_FACTORY(FBoolArrayToFaceSelectionDataflowNode);
-		DATAFLOW_NODE_REGISTER_CREATION_FACTORY(FFloatArrayToVertexSelectionDataflowNode);
 		DATAFLOW_NODE_REGISTER_CREATION_FACTORY(FSetVertexColorInCollectionFromVertexSelectionDataflowNode);
 		DATAFLOW_NODE_REGISTER_CREATION_FACTORY(FSetVertexColorInCollectionFromFloatArrayDataflowNode);
-		DATAFLOW_NODE_REGISTER_CREATION_FACTORY(FFloatArrayNormalizeDataflowNode);
-		DATAFLOW_NODE_REGISTER_CREATION_FACTORY(FVectorArrayNormalizeDataflowNode);
 		DATAFLOW_NODE_REGISTER_CREATION_FACTORY(FMultiplyTransformDataflowNode);
 		DATAFLOW_NODE_REGISTER_CREATION_FACTORY(FInvertTransformDataflowNode);
 		DATAFLOW_NODE_REGISTER_CREATION_FACTORY(FSelectionToVertexListDataflowNode);
-		DATAFLOW_NODE_REGISTER_CREATION_FACTORY(FUnionIntArraysDataflowNode);
 
 		// GeometryCollection
 		DATAFLOW_NODE_REGISTER_CREATION_FACTORY_NODE_COLORS_BY_CATEGORY("GeometryCollection", FLinearColor(0.55f, 0.45f, 1.0f), CDefaultNodeBodyTintColor);
@@ -262,117 +253,6 @@ void FHashVectorDataflowNode::Evaluate(Dataflow::FContext& Context, const FDataf
 	}
 }
 
-void FGetFloatArrayElementDataflowNode::Evaluate(Dataflow::FContext& Context, const FDataflowOutput* Out) const
-{
-	if (Out->IsA(&FloatValue))
-	{
-		const TArray<float>& InArray = GetValue(Context, &FloatArray);
-		const int32& InIndex = GetValue(Context, &Index);
-		const float OutputValue = InArray.IsValidIndex(InIndex) ? InArray[InIndex] : 0.0f;
-
-		SetValue(Context, OutputValue, &FloatValue);
-	}
-}
-
-void FFloatArrayToIntArrayDataflowNode::Evaluate(Dataflow::FContext& Context, const FDataflowOutput* Out) const
-{
-	if (Out->IsA<TArray<int32>>(&IntArray))
-	{
-		TArray<float> FloatVal = GetValue<TArray<float>>(Context, &FloatArray);
-		TArray<int32> RetVal; RetVal.SetNumUninitialized(FloatVal.Num());
-		if (Function == EFloatArrayToIntArrayFunctionEnum::Dataflow_FloatToInt_Function_Floor)
-		{
-			for (int32 i = 0; i < FloatVal.Num(); i++)
-			{
-				RetVal[i] = FMath::FloorToInt32(FloatVal[i]);
-			}
-		}
-		else if (Function == EFloatArrayToIntArrayFunctionEnum::Dataflow_FloatToInt_Function_Ceil)
-		{
-			for (int32 i = 0; i < FloatVal.Num(); i++)
-			{
-				RetVal[i] = FMath::CeilToInt32(FloatVal[i]);
-			}
-		}
-		else if (Function == EFloatArrayToIntArrayFunctionEnum::Dataflow_FloatToInt_Function_Round)
-		{
-			for (int32 i = 0; i < FloatVal.Num(); i++)
-			{
-				RetVal[i] = FMath::RoundToInt32(FloatVal[i]);
-			}
-		}
-		else if (Function == EFloatArrayToIntArrayFunctionEnum::Dataflow_FloatToInt_Function_Truncate)
-		{
-			for (int32 i = 0; i < FloatVal.Num(); i++)
-			{
-				RetVal[i] = int32(FMath::TruncToFloat(FloatVal[i]));
-			}
-		}
-
-		else if (Function == EFloatArrayToIntArrayFunctionEnum::Dataflow_FloatToInt_NonZeroToIndex)
-		{
-			int32 RetValIndex = 0;
-			for (int32 i = 0; i < FloatVal.Num(); i++)
-			{
-				if (FloatVal[i] != 0.0)
-				{
-					RetVal[RetValIndex] = i;
-					RetValIndex++;
-				}
-			}
-			RetVal.SetNum(RetValIndex);
-		}
-		else if (Function == EFloatArrayToIntArrayFunctionEnum::Dataflow_FloatToInt_ZeroToIndex)
-		{
-			int32 RetValIndex = 0;
-			for (int32 i = 0; i < FloatVal.Num(); i++)
-			{
-				if (FloatVal[i] == 0.0)
-				{
-					RetVal[RetValIndex] = i;
-					RetValIndex++;
-				}
-			}
-			RetVal.SetNum(RetValIndex);
-		}
-
-		SetValue<TArray<int32>>(Context, RetVal, &IntArray);
-	}
-}
-
-void FGetArrayElementDataflowNode::Evaluate(Dataflow::FContext& Context, const FDataflowOutput* Out) const
-{
-	if (Out->IsA<FVector>(&Point))
-	{
-		const TArray<FVector>& Array = GetValue<TArray<FVector>>(Context, &Points);
-		if (Array.Num() > 0 && Index >= 0 && Index < Array.Num())
-		{
-			SetValue<FVector>(Context, Array[Index], &Point);
-			return;
-		}
-
-		SetValue<FVector>(Context, FVector(0.f), &Point);
-	}
-}
-
-void FGetNumArrayElementsDataflowNode::Evaluate(Dataflow::FContext& Context, const FDataflowOutput* Out) const
-{
-	if (Out->IsA<int32>(&NumElements))
-	{
-		if (IsConnected<TArray<FVector>>(&Points))
-		{
-			SetValue<int32>(Context, GetValue<TArray<FVector>>(Context, &Points).Num(), &NumElements);
-			return;
-		}
-		else if (IsConnected<TArray<FVector3f>>(&Vector3fArray))
-		{
-			SetValue<int32>(Context, GetValue<TArray<FVector3f>>(Context, &Vector3fArray).Num(), &NumElements);
-			return;
-		}
-
-		SetValue<int32>(Context, 0, &NumElements);
-	}
-}
 
 void FGetBoundingBoxesFromCollectionDataflowNode::Evaluate(Dataflow::FContext& Context, const FDataflowOutput* Out) const
 {
@@ -1160,72 +1040,6 @@ void FSetCollectionAttributeDataTypedDataflowNode::Evaluate(Dataflow::FContext&
 }
 
 
-void FBoolArrayToFaceSelectionDataflowNode::Evaluate(Dataflow::FContext& Context, const FDataflowOutput* Out) const
-{
-	if (Out->IsA<FDataflowFaceSelection>(&FaceSelection))
-	{
-		const TArray<bool>& InBoolAttributeData = GetValue<TArray<bool>>(Context, &BoolAttributeData);
-
-		FDataflowFaceSelection NewFaceSelection;
-		NewFaceSelection.Initialize(InBoolAttributeData.Num(), false);
-		NewFaceSelection.SetFromArray(InBoolAttributeData);
-
-		SetValue<FDataflowFaceSelection>(Context, NewFaceSelection, &FaceSelection);
-	}
-}
-
-
-void FFloatArrayToVertexSelectionDataflowNode::Evaluate(Dataflow::FContext& Context, const FDataflowOutput* Out) const
-{
-	if (Out->IsA<FDataflowVertexSelection>(&VertexSelection))
-	{
-		const TArray<float>& InFloatArray = GetValue<TArray<float>>(Context, &FloatArray);
-
-		FDataflowVertexSelection NewVertexSelection;
-		NewVertexSelection.Initialize(InFloatArray.Num(), false);
-
-		for (int32 Idx = 0; Idx < InFloatArray.Num(); ++Idx)
-		{
-			if (Operation == ECompareOperationEnum::Dataflow_Compare_Equal)
-			{
-				if (InFloatArray[Idx] == Threshold)
-				{
-					NewVertexSelection.SetSelected(Idx);
-				}
-			}
-			else if (Operation == ECompareOperationEnum::Dataflow_Compare_Smaller)
-			{
-				if (InFloatArray[Idx] < Threshold)
-				{
-					NewVertexSelection.SetSelected(Idx);
-				}
-			}
-			else if (Operation == ECompareOperationEnum::Dataflow_Compare_SmallerOrEqual)
-			{
-				if (InFloatArray[Idx] <= Threshold)
-				{
-					NewVertexSelection.SetSelected(Idx);
-				}
-			}
-			else if (Operation == ECompareOperationEnum::Dataflow_Compare_Greater)
-			{
-				if (InFloatArray[Idx] > Threshold)
-				{
-					NewVertexSelection.SetSelected(Idx);
-				}
-			}
-			else if (Operation == ECompareOperationEnum::Dataflow_Compare_GreaterOrEqual)
-			{
-				if (InFloatArray[Idx] >= Threshold)
-				{
-					NewVertexSelection.SetSelected(Idx);
-				}
-			}
-		}
-
-		SetValue<FDataflowVertexSelection>(Context, NewVertexSelection, &VertexSelection);
-	}
-}
 
 
 void FSetVertexColorInCollectionFromVertexSelectionDataflowNode::Evaluate(Dataflow::FContext& Context, const FDataflowOutput* Out) const
@@ -1291,146 +1105,6 @@ void FSetVertexColorInCollectionFromFloatArrayDataflowNode::Evaluate(Dataflow::F
 }
 
 
-void FFloatArrayNormalizeDataflowNode::Evaluate(Dataflow::FContext& Context, const FDataflowOutput* Out) const
-{
-	if (Out->IsA<TArray<float>>(&OutFloatArray))
-	{
-		const TArray<float>& InInFloatArray = GetValue<TArray<float>>(Context, &InFloatArray);
-		const FDataflowVertexSelection& InSelection = GetValue<FDataflowVertexSelection>(Context, &Selection);
-		const float InMinRange = GetValue<float>(Context, &MinRange);
-		const float InMaxRange = GetValue<float>(Context, &MaxRange);
-
-		const int32 NumElems = InFloatArray.Num();
-
-		TArray<float> NewFloatArray;
-		NewFloatArray.Init(false, NumElems);
-
-		if (IsConnected<FDataflowVertexSelection>(&Selection))
-		{
-			if (InInFloatArray.Num() == InSelection.Num())
-			{
-				// Compute Min/Max
-				float MinValue = 1e9, MaxValue = 1e-9;
-
-				for (int32 Idx = 0; Idx < NumElems; ++Idx)
-				{
-					if (InSelection.IsSelected(Idx))
-					{
-						if (InInFloatArray[Idx] < MinValue)
-						{
-							MinValue = InInFloatArray[Idx];
-						}
-
-						if (InInFloatArray[Idx] > MaxValue)
-						{
-							MaxValue = InInFloatArray[Idx];
-						}
-					}
-				}
-
-				for (int32 Idx = 0; Idx < NumElems; ++Idx)
-				{
-					if (InSelection.IsSelected(Idx))
-					{
-						// Normalize it
-						NewFloatArray[Idx] = (NewFloatArray[Idx] - MinValue) / (MaxValue - MinValue);
-
-						// Transform it into (MinRange, MaxRange)
-						NewFloatArray[Idx] = InMinRange + NewFloatArray[Idx] * (InMaxRange - InMinRange);
-					}
-				}
-
-				SetValue<TArray<float>>(Context, NewFloatArray, &OutFloatArray);
-				return;
-			}
-		}
-		else
-		{
-			// Compute Min/Max
-			float MinValue = 1e9, MaxValue = 1e-9;
-
-			for (int32 Idx = 0; Idx < NumElems; ++Idx)
-			{
-				if (InInFloatArray[Idx] < MinValue)
-				{
-					MinValue = InInFloatArray[Idx];
-				}
-
-				if (InInFloatArray[Idx] > MaxValue)
-				{
-					MaxValue = InInFloatArray[Idx];
-				}
-			}
-
-			for (int32 Idx = 0; Idx < NumElems; ++Idx)
-			{
-				// Normalize it
-				NewFloatArray[Idx] = (NewFloatArray[Idx] - MinValue) / (MaxValue - MinValue);
-
-				// Transform it into (MinRange, MaxRange)
-				NewFloatArray[Idx] = InMinRange + NewFloatArray[Idx] * (InMaxRange - InMinRange);
-			}
-
-			SetValue<TArray<float>>(Context, NewFloatArray, &OutFloatArray);
-			return;
-		}
-
-		SetValue<TArray<float>>(Context, TArray<float>(), &OutFloatArray);
-	}
-}
-
-
-void FVectorArrayNormalizeDataflowNode::Evaluate(Dataflow::FContext& Context, const FDataflowOutput* Out) const
-{
-	if (Out->IsA<TArray<FVector>>(&OutVectorArray))
-	{
-		const TArray<FVector>& InInVectorArray = GetValue<TArray<FVector>>(Context, &InVectorArray);
-		const FDataflowVertexSelection& InSelection = GetValue<FDataflowVertexSelection>(Context, &Selection);
-		const float InMagnitude = GetValue<float>(Context, &Magnitude);
-
-		const int32 NumElems = InInVectorArray.Num();
-
-		TArray<FVector> NewVectorArray;
-		NewVectorArray.Init(FVector(0.f), NumElems);
-
-		if (IsConnected<FDataflowVertexSelection>(&Selection))
-		{
-			if (InInVectorArray.Num() == InSelection.Num())
-			{
-				for (int32 Idx = 0; Idx < NumElems; ++Idx)
-				{
-					FVector Vector = InInVectorArray[Idx];
-
-					if (InSelection.IsSelected(Idx))
-					{
-						Vector.Normalize();
-						Vector *= InMagnitude;
-						NewVectorArray[Idx] = Vector;
-					}
-				}
-
-				SetValue<TArray<FVector>>(Context, NewVectorArray, &OutVectorArray);
-				return;
-			}
-		}
-		else
-		{
-			for (int32 Idx = 0; Idx < NumElems; ++Idx)
-			{
-				FVector Vector = InInVectorArray[Idx];
-
-				Vector.Normalize();
-				Vector *= InMagnitude;
-				NewVectorArray[Idx] = Vector;
-			}
-
-			SetValue<TArray<FVector>>(Context, NewVectorArray, &OutVectorArray);
-			return;
-		}
-
-		SetValue<TArray<FVector>>(Context, TArray<FVector>(), &OutVectorArray);
-	}
-}
 
 void FMultiplyTransformDataflowNode::Evaluate(Dataflow::FContext& Context, const FDataflowOutput* Out) const
 {
@@ -1453,21 +1127,3 @@ void FInvertTransformDataflowNode::Evaluate(Dataflow::FContext& Context, const F
 	}
 }
 
-void FUnionIntArraysDataflowNode::Evaluate(Dataflow::FContext& Context, const FDataflowOutput* Out) const
-{
-	if (Out->IsA<TArray<int32>>(&OutArray))
-	{
-		TArray<int32> Array1 = GetValue<TArray<int32>>(Context, &InArray1);
-		TArray<int32> Array2 = GetValue<TArray<int32>>(Context, &InArray2);
-		TArray<int32> OutputArray;
-		for (int32 i = 0; i < Array1.Num(); i++)
-		{
-			OutputArray.AddUnique(Array1[i]);
-		}
-		for (int32 i = 0; i < Array2.Num(); i++)
-		{
-			OutputArray.AddUnique(Array2[i]);
-		}
-		SetValue<TArray<int32>>(Context, OutputArray, &OutArray);
-	}
-}
diff --git a/Engine/Plugins/Experimental/GeometryCollectionPlugin/Source/GeometryCollectionNodes/Private/Dataflow/GeometryCollectionNodesPlugin.cpp b/Engine/Plugins/Experimental/GeometryCollectionPlugin/Source/GeometryCollectionNodes/Private/Dataflow/GeometryCollectionNodesPlugin.cpp
--- a/Engine/Plugins/Experimental/GeometryCollectionPlugin/Source/GeometryCollectionNodes/Private/Dataflow/GeometryCollectionNodesPlugin.cpp
+++ b/Engine/Plugins/Experimental/GeometryCollectionPlugin/Source/GeometryCollectionNodes/Private/Dataflow/GeometryCollectionNodesPlugin.cpp
@@ -22,6 +22,7 @@
 #include "Dataflow/GeometryCollectionMathNodes.h"
 #include "Dataflow/GeometryCollectionConversionNodes.h"
 #include "Dataflow/GeometryCollectionVerticesNodes.h"
+#include "Dataflow/GeometryCollectionArrayNodes.h"
 
 #define LOCTEXT_NAMESPACE "DataflowNodes"
 
@@ -45,6 +46,7 @@ void IGeometryCollectionNodesPlugin::StartupModule()
 	Dataflow::GeometryCollectionMathNodes();
 	Dataflow::GeometryCollectionConversionNodes();
 	Dataflow::GeometryCollectionVerticesNodes();
+	Dataflow::GeometryCollectionArrayNodes();
 }
 
 void IGeometryCollectionNodesPlugin::ShutdownModule()
diff --git a/Engine/Plugins/Experimental/GeometryCollectionPlugin/Source/GeometryCollectionNodes/Public/Dataflow/GeometryCollectionArrayNodes.h b/Engine/Plugins/Experimental/GeometryCollectionPlugin/Source/GeometryCollectionNodes/Public/Dataflow/GeometryCollectionArrayNodes.h
new file mode 100644
--- /dev/null
+++ b/Engine/Plugins/Experimental/GeometryCollectionPlugin/Source/GeometryCollectionNodes/Public/Dataflow/GeometryCollectionArrayNodes.h
@@ -0,0 +1,496 @@
+// Copyright Epic Games, Inc. All Rights Reserved.
+
+#pragma once 
+
+#include "CoreMinimal.h"
+#include "Dataflow/DataflowEngine.h"
+#include "GeometryCollection/ManagedArrayCollection.h"
+#include "Dataflow/DataflowSelection.h"
+
+#include "GeometryCollectionArrayNodes.generated.h"
+
+
+class FGeometryCollection;
+
+
+/**
+ *
+ * Returns the specified element from an array
+ *
+ */
+USTRUCT()
+struct FGetFloatArrayElementDataflowNode : public FDataflowNode
+{
+	GENERATED_USTRUCT_BODY()
+	DATAFLOW_NODE_DEFINE_INTERNAL(FGetFloatArrayElementDataflowNode, "GetFloatArrayElement", "Utilities|Array", "")
+
+public:
+	/** Element index */
+	UPROPERTY(EditAnywhere, Category = "Index");
+	int32 Index = 0;
+
+	/** Array to get the element from */
+	UPROPERTY(meta = (DataflowInput))
+	TArray<float> FloatArray;
+
+	/** Specified element */
+	UPROPERTY(meta = (DataflowOutput))
+	float FloatValue = 0;
+
+	FGetFloatArrayElementDataflowNode(const Dataflow::FNodeParameters& InParam, FGuid InGuid = FGuid::NewGuid())
+		: FDataflowNode(InParam, InGuid)
+	{
+		RegisterInputConnection(&FloatArray);
+		RegisterInputConnection(&Index);
+		RegisterOutputConnection(&FloatValue);
+	}
+
+	virtual void Evaluate(Dataflow::FContext& Context, const FDataflowOutput* Out) const override;
+
+};
+
+
+UENUM(BlueprintType)
+enum class EFloatArrayToIntArrayFunctionEnum : uint8
+{
+	Dataflow_FloatToInt_Function_Floor UMETA(DisplayName = "Floor()"),
+	Dataflow_FloatToInt_Function_Ceil UMETA(DisplayName = "Ceil()"),
+	Dataflow_FloatToInt_Function_Round UMETA(DisplayName = "Round()"),
+	Dataflow_FloatToInt_Function_Truncate UMETA(DisplayName = "Truncate()"),
+	Dataflow_FloatToInt_NonZeroToIndex UMETA(DisplayName = "Non-zero to Index"),
+	Dataflow_FloatToInt_ZeroToIndex UMETA(DisplayName = "Zero to Index"),
+	//~~~
+	//256th entry
+	Dataflow_Max                UMETA(Hidden)
+};
+
+/**
+ *
+ * Converts a Float array to Int array using the specified method.
+ *
+ */
+USTRUCT()
+struct FFloatArrayToIntArrayDataflowNode : public FDataflowNode
+{
+	GENERATED_USTRUCT_BODY()
+	DATAFLOW_NODE_DEFINE_INTERNAL(FFloatArrayToIntArrayDataflowNode, "FloatArrayToIntArray", "Math|Conversions", "")
+
+public:
+	/** Conversion method:
+	* Floor takes the floor of each input float value - 1.1 turns into 1.
+	* Ceil takes the ceil - 1.1 turns into 2.
+	* Round rounds to the nearest integer - 1.1 turns into 1.
+	* Tuncate trucates like a type cast - 1.1 turns into 1.
+	* Non-zero to Index appends the index of all non-zero values to the output array.
+	* Zero to Index appends the index of all zero values to the output array.
+	*/
+	UPROPERTY(EditAnywhere, Category = "Float");
+	EFloatArrayToIntArrayFunctionEnum Function = EFloatArrayToIntArrayFunctionEnum::Dataflow_FloatToInt_NonZeroToIndex;
+
+	/** Float array value to convert */
+	UPROPERTY(EditAnywhere, Category = "Float", meta = (DataflowInput))
+	TArray<float> FloatArray;
+
+	/** Int array output */
+	UPROPERTY(meta = (DataflowOutput))
+	TArray<int32> IntArray;
+
+	FFloatArrayToIntArrayDataflowNode(const Dataflow::FNodeParameters& InParam, FGuid InGuid = FGuid::NewGuid())
+		: FDataflowNode(InParam, InGuid)
+	{
+		RegisterInputConnection(&FloatArray);
+		RegisterOutputConnection(&IntArray);
+	}
+
+	virtual void Evaluate(Dataflow::FContext& Context, const FDataflowOutput* Out) const override;
+
+};
+
+
+/**
+ *
+ * Returns the specified element from an array
+ *
+ */
+USTRUCT()
+struct FGetArrayElementDataflowNode : public FDataflowNode
+{
+	GENERATED_USTRUCT_BODY()
+	DATAFLOW_NODE_DEFINE_INTERNAL(FGetArrayElementDataflowNode, "GetArrayElement", "Utilities|Array", "")
+
+public:
+	/** Element index */
+	UPROPERTY(EditAnywhere, Category = "Index");
+	int32 Index = 0;
+
+	/** Array to get the element from */
+	UPROPERTY(meta = (DataflowInput))
+	TArray<FVector> Points;
+
+	/** Specified element */
+	UPROPERTY(meta = (DataflowOutput))
+	FVector Point = FVector(0.0);
+
+	FGetArrayElementDataflowNode(const Dataflow::FNodeParameters& InParam, FGuid InGuid = FGuid::NewGuid())
+		: FDataflowNode(InParam, InGuid)
+	{
+		RegisterInputConnection(&Points);
+		RegisterInputConnection(&Index);
+		RegisterOutputConnection(&Point);
+	}
+
+	virtual void Evaluate(Dataflow::FContext& Context, const FDataflowOutput* Out) const override;
+
+};
+
+
+/**
+ *
+ * Returns the number of elements in an array
+ *
+ */
+USTRUCT()
+struct FGetNumArrayElementsDataflowNode : public FDataflowNode
+{
+	GENERATED_USTRUCT_BODY()
+	DATAFLOW_NODE_DEFINE_INTERNAL(FGetNumArrayElementsDataflowNode, "GetNumArrayElements", "Utilities|Array", "")
+
+public:
+	/** Float array input */
+	UPROPERTY(meta = (DataflowInput, DisplayName = "FloatArray"))
+	TArray<float> FloatArray;
+
+	/** Int32 array input */
+	UPROPERTY(meta = (DataflowInput, DisplayName = "IntArray"))
+	TArray<int32> IntArray;
+
+	/** FVector array input */
+	UPROPERTY(meta = (DataflowInput, DisplayName = "VectorArray"))
+	TArray<FVector> Points;
+
+	/** FVector3f array input */
+	UPROPERTY(meta = (DataflowInput, DisplayName = "Vector3fArray"))
+	TArray<FVector3f> Vector3fArray;
+
+	/** Number of elements in the array */
+	UPROPERTY(meta = (DataflowOutput))
+	int32 NumElements = 0;
+
+	FGetNumArrayElementsDataflowNode(const Dataflow::FNodeParameters& InParam, FGuid InGuid = FGuid::NewGuid())
+		: FDataflowNode(InParam, InGuid)
+	{
+		RegisterInputConnection(&FloatArray);
+		RegisterInputConnection(&IntArray);
+		RegisterInputConnection(&Points);
+		RegisterInputConnection(&Vector3fArray);
+		RegisterOutputConnection(&NumElements);
+	}
+
+	virtual void Evaluate(Dataflow::FContext& Context, const FDataflowOutput* Out) const override;
+
+};
+
+
+/**
+ *
+ * Converts a TArray<bool> to a FDataflowFaceSelection
+ *
+ */
+USTRUCT()
+struct FBoolArrayToFaceSelectionDataflowNode : public FDataflowNode
+{
+	GENERATED_USTRUCT_BODY()
+	DATAFLOW_NODE_DEFINE_INTERNAL(FBoolArrayToFaceSelectionDataflowNode, "BoolArrayToFaceSelection", "Utilities|Array", "")
+
+public:
+	/** TArray<bool> data */
+	UPROPERTY(meta = (DataflowInput, DataflowIntrinsic));
+	TArray<bool> BoolAttributeData;
+
+	UPROPERTY(meta = (DataflowOutput, DisplayName = "FaceSelection"))
+	FDataflowFaceSelection FaceSelection;
+
+	FBoolArrayToFaceSelectionDataflowNode(const Dataflow::FNodeParameters& InParam, FGuid InGuid = FGuid::NewGuid())
+		: FDataflowNode(InParam, InGuid)
+	{
+		RegisterInputConnection(&BoolAttributeData);
+		RegisterOutputConnection(&FaceSelection);
+	}
+
+	virtual void Evaluate(Dataflow::FContext& Context, const FDataflowOutput* Out) const override;
+
+};
+
+
+UENUM(BlueprintType)
+enum class ECompareOperation1Enum : uint8
+{
+	Dataflow_Compare_Equal UMETA(DisplayName = "=="),
+	Dataflow_Compare_Smaller UMETA(DisplayName = "<"),
+	Dataflow_Compare_SmallerOrEqual UMETA(DisplayName = "<="),
+	Dataflow_Compare_Greater UMETA(DisplayName = ">"),
+	Dataflow_Compare_GreaterOrEqual UMETA(DisplayName = ">="),
+	//~~~
+	//256th entry
+	Dataflow_Max                UMETA(Hidden)
+};
+
+/**
+ *
+ * Converts a TArray<float> to a FDataflowVertexSelection
+ *
+ */
+USTRUCT()
+struct FFloatArrayToVertexSelectionDataflowNode : public FDataflowNode
+{
+	GENERATED_USTRUCT_BODY()
+	DATAFLOW_NODE_DEFINE_INTERNAL(FFloatArrayToVertexSelectionDataflowNode, "FloatArrayToVertexSelection", "Utilities|Array", "")
+
+public:
+	/** TArray<floatl> array */
+	UPROPERTY(meta = (DataflowInput, DataflowIntrinsic));
+	TArray<float> FloatArray;
+
+	/** Comparison operation */
+	UPROPERTY(EditAnywhere, Category = "Compare");
+	ECompareOperation1Enum Operation = ECompareOperation1Enum::Dataflow_Compare_Greater;
+
+	/**  */
+	UPROPERTY(EditAnywhere, Category = "Compare")
+	float Threshold = 0.f;
+
+	UPROPERTY(meta = (DataflowOutput, DisplayName = "VertexSelection"))
+	FDataflowVertexSelection VertexSelection;
+
+	FFloatArrayToVertexSelectionDataflowNode(const Dataflow::FNodeParameters& InParam, FGuid InGuid = FGuid::NewGuid())
+		: FDataflowNode(InParam, InGuid)
+	{
+		RegisterInputConnection(&FloatArray);
+		RegisterOutputConnection(&VertexSelection);
+	}
+
+	virtual void Evaluate(Dataflow::FContext& Context, const FDataflowOutput* Out) const override;
+
+};
+
+
+/**
+ *
+ * Normalize the selected float data in a FloatArray
+ *
+ */
+USTRUCT()
+struct FFloatArrayNormalizeDataflowNode : public FDataflowNode
+{
+	GENERATED_USTRUCT_BODY()
+	DATAFLOW_NODE_DEFINE_INTERNAL(FFloatArrayNormalizeDataflowNode, "FloatArrayNormalize", "Math|Float", "")
+
+public:
+	/** Input VectorArray */
+	UPROPERTY(meta = (DataflowInput, DataflowIntrinsic))
+	TArray<float> InFloatArray;
+
+	/** Selection for the operation */
+	UPROPERTY(meta = (DataflowInput))
+	FDataflowVertexSelection Selection;
+
+	/**  */
+	UPROPERTY(EditAnywhere, Category = "Field", meta = (DataflowInput))
+	float MinRange = 0.f;
+
+	/**  */
+	UPROPERTY(EditAnywhere, Category = "Field", meta = (DataflowInput))
+	float MaxRange = 1.f;
+
+	/**  */
+	UPROPERTY(meta = (DataflowOutput))
+	TArray<float> OutFloatArray;
+
+	FFloatArrayNormalizeDataflowNode(const Dataflow::FNodeParameters& InParam, FGuid InGuid = FGuid::NewGuid())
+		: FDataflowNode(InParam, InGuid)
+	{
+		RegisterInputConnection(&InFloatArray);
+		RegisterInputConnection(&Selection);
+		RegisterInputConnection(&MinRange);
+		RegisterInputConnection(&MaxRange);
+		RegisterOutputConnection(&OutFloatArray);
+	}
+
+	virtual void Evaluate(Dataflow::FContext& Context, const FDataflowOutput* Out) const override;
+
+};
+
+
+/**
+ *
+ * Normalize all the selected vectors in a VectorArray
+ *
+ */
+USTRUCT()
+struct FVectorArrayNormalizeDataflowNode : public FDataflowNode
+{
+	GENERATED_USTRUCT_BODY()
+	DATAFLOW_NODE_DEFINE_INTERNAL(FVectorArrayNormalizeDataflowNode, "VectorArrayNormalize", "Math|Vector", "")
+
+public:
+	/** Input VectorArray */
+	UPROPERTY(meta = (DataflowInput, DataflowIntrinsic))
+	TArray<FVector> InVectorArray;
+
+	/** Selection for the operation */
+	UPROPERTY(meta = (DataflowInput))
+	FDataflowVertexSelection Selection;
+
+	/**  */
+	UPROPERTY(EditAnywhere, Category = "Field", meta = (DataflowInput))
+	float Magnitude = 1.f;
+
+	/**  */
+	UPROPERTY(meta = (DataflowOutput))
+	TArray<FVector> OutVectorArray;
+
+	FVectorArrayNormalizeDataflowNode(const Dataflow::FNodeParameters& InParam, FGuid InGuid = FGuid::NewGuid())
+		: FDataflowNode(InParam, InGuid)
+	{
+		RegisterInputConnection(&InVectorArray);
+		RegisterInputConnection(&Selection);
+		RegisterInputConnection(&Magnitude);
+		RegisterOutputConnection(&OutVectorArray);
+	}
+
+	virtual void Evaluate(Dataflow::FContext& Context, const FDataflowOutput* Out) const override;
+
+};
+
+
+USTRUCT()
+struct FUnionIntArraysDataflowNode : public FDataflowNode
+{
+	GENERATED_USTRUCT_BODY()
+	DATAFLOW_NODE_DEFINE_INTERNAL(FUnionIntArraysDataflowNode, "UnionIntArrays", "Utilities", "")
+
+public:
+	UPROPERTY(meta = (DataflowInput, DisplayName = "InArray1"));
+	TArray<int32> InArray1;
+
+	UPROPERTY(meta = (DataflowInput, DisplayName = "InArray2"));
+	TArray<int32> InArray2;
+
+	UPROPERTY(meta = (DataflowOutput, DisplayName = "OutArray", DataflowPassthrough = "InArray1"));
+	TArray<int32> OutArray;
+
+	FUnionIntArraysDataflowNode(const Dataflow::FNodeParameters& InParam, FGuid InGuid = FGuid::NewGuid())
+		: FDataflowNode(InParam, InGuid)
+	{
+		RegisterInputConnection(&InArray1);
+		RegisterInputConnection(&InArray2);
+		RegisterOutputConnection(&OutArray, &InArray1);
+	}
+
+	virtual void Evaluate(Dataflow::FContext& Context, const FDataflowOutput* Out) const override;
+};
+
+
+/**
+ *
+ * Removes the specified element from an array
+ *
+ */
+USTRUCT()
+struct FRemoveFloatArrayElementDataflowNode : public FDataflowNode
+{
+	GENERATED_USTRUCT_BODY()
+	DATAFLOW_NODE_DEFINE_INTERNAL(FRemoveFloatArrayElementDataflowNode, "RemoveFloatArrayElement", "Utilities|Array", "")
+
+public:
+	/** Element index */
+	UPROPERTY(EditAnywhere, Category = "Index");
+	int32 Index = 0;
+
+	/** Preserve order, if order not important set it to false for faster computation */
+	UPROPERTY(EditAnywhere, Category = "Order");
+	bool bPreserveOrder = true;
+
+	/** Array to remove the element from */
+	UPROPERTY(meta = (DataflowInput, DataflowOutput, DataflowPassthrough = "FloatArray", DataflowIntrinsic))
+	TArray<float> FloatArray;
+
+	FRemoveFloatArrayElementDataflowNode(const Dataflow::FNodeParameters& InParam, FGuid InGuid = FGuid::NewGuid())
+		: FDataflowNode(InParam, InGuid)
+	{
+		RegisterInputConnection(&FloatArray);
+		RegisterInputConnection(&Index);
+		RegisterOutputConnection(&FloatArray, &FloatArray);
+	}
+
+	virtual void Evaluate(Dataflow::FContext& Context, const FDataflowOutput* Out) const override;
+
+};
+
+
+UENUM(BlueprintType)
+enum class EStatisticsOperationEnum : uint8
+{
+	Dataflow_EStatisticsOperationEnum_Min		UMETA(DisplayName = "Min", ToolTip = ""),
+	Dataflow_EStatisticsOperationEnum_Max		UMETA(DisplayName = "Max", ToolTip = ""),
+	Dataflow_EStatisticsOperationEnum_Mean		UMETA(DisplayName = "Mean", ToolTip = ""),
+	Dataflow_EStatisticsOperationEnum_Median	UMETA(DisplayName = "Median", ToolTip = ""),
+	Dataflow_EStatisticsOperationEnum_Mode		UMETA(DisplayName = "Mode", ToolTip = ""),
+	Dataflow_EStatisticsOperationEnum_Sum		UMETA(DisplayName = "Sum", ToolTip = ""),
+	//~~~
+	//256th entry
+	Dataflow_Max                UMETA(Hidden)
+};
+
+/**
+ *
+ * Computes statistics of a float array
+ *
+ */
+USTRUCT()
+struct FFloatArrayComputeStatisticsDataflowNode : public FDataflowNode
+{
+	GENERATED_USTRUCT_BODY()
+	DATAFLOW_NODE_DEFINE_INTERNAL(FFloatArrayComputeStatisticsDataflowNode, "FloatArrayComputeStatistics", "Utilities|Array", "")
+
+public:
+	/** Array to compute values from */
+	UPROPERTY(meta = (DataflowInput, DataflowIntrinsic))
+	TArray<float> FloatArray;
+
+	/** TransformSelection describes which values to use, if not connected all the elements will be used */
+	UPROPERTY(meta = (DataflowInput, DisplayName = "TransformSelection"))
+	FDataflowTransformSelection TransformSelection;
+
+	/** Statistics Operation */
+	UPROPERTY(EditAnywhere, Category = "Operation", meta = (DisplayName = "Operation"))
+	EStatisticsOperationEnum OperationName = EStatisticsOperationEnum::Dataflow_EStatisticsOperationEnum_Min;
+
+	/** Computed value */
+	UPROPERTY(meta = (DataflowOutput))
+	float Value = 0.f;
+
+	/** Indices of elements with the computed value */
+	UPROPERTY(meta = (DataflowOutput))
+	TArray<int32> Indices;
+
+	FFloatArrayComputeStatisticsDataflowNode(const Dataflow::FNodeParameters& InParam, FGuid InGuid = FGuid::NewGuid())
+		: FDataflowNode(InParam, InGuid)
+	{
+		RegisterInputConnection(&FloatArray);
+		RegisterInputConnection(&TransformSelection);
+		RegisterOutputConnection(&Value);
+		RegisterOutputConnection(&Indices);
+	}
+
+	virtual void Evaluate(Dataflow::FContext& Context, const FDataflowOutput* Out) const override;
+
+};
+
+
+
+namespace Dataflow
+{
+	void GeometryCollectionArrayNodes();
+}
+
diff --git a/Engine/Plugins/Experimental/GeometryCollectionPlugin/Source/GeometryCollectionNodes/Public/Dataflow/GeometryCollectionNodes.h b/Engine/Plugins/Experimental/GeometryCollectionPlugin/Source/GeometryCollectionNodes/Public/Dataflow/GeometryCollectionNodes.h
--- a/Engine/Plugins/Experimental/GeometryCollectionPlugin/Source/GeometryCollectionNodes/Public/Dataflow/GeometryCollectionNodes.h
+++ b/Engine/Plugins/Experimental/GeometryCollectionPlugin/Source/GeometryCollectionNodes/Public/Dataflow/GeometryCollectionNodes.h
@@ -374,171 +374,6 @@ public:
 
 };
 
-UENUM(BlueprintType)
-enum class EFloatArrayToIntArrayFunctionEnum : uint8
-{
-	Dataflow_FloatToInt_Function_Floor UMETA(DisplayName = "Floor()"),
-	Dataflow_FloatToInt_Function_Ceil UMETA(DisplayName = "Ceil()"),
-	Dataflow_FloatToInt_Function_Round UMETA(DisplayName = "Round()"),
-	Dataflow_FloatToInt_Function_Truncate UMETA(DisplayName = "Truncate()"),
-	Dataflow_FloatToInt_NonZeroToIndex UMETA(DisplayName = "Non-zero to Index"),
-	Dataflow_FloatToInt_ZeroToIndex UMETA(DisplayName = "Zero to Index"),
-	//~~~
-	//256th entry
-	Dataflow_Max                UMETA(Hidden)
-};
-
-/**
- *
- * Converts a Float array to Int array using the specified method.
- *
- */
-USTRUCT()
-struct FFloatArrayToIntArrayDataflowNode : public FDataflowNode
-{
-	GENERATED_USTRUCT_BODY()
-		DATAFLOW_NODE_DEFINE_INTERNAL(FFloatArrayToIntArrayDataflowNode, "FloatArrayToIntArray", "Math|Conversions", "")
-
-public:
-	/** Conversion method:
-	* Floor takes the floor of each input float value - 1.1 turns into 1.
-	* Ceil takes the ceil - 1.1 turns into 2.
-	* Round rounds to the nearest integer - 1.1 turns into 1.
-	* Tuncate trucates like a type cast - 1.1 turns into 1.
-	* Non-zero to Index appends the index of all non-zero values to the output array.
-	* Zero to Index appends the index of all zero values to the output array.
-	*/
-	UPROPERTY(EditAnywhere, Category = "Float");
-	EFloatArrayToIntArrayFunctionEnum Function = EFloatArrayToIntArrayFunctionEnum::Dataflow_FloatToInt_NonZeroToIndex;
-
-	/** Float array value to convert */
-	UPROPERTY(EditAnywhere, Category = "Float", meta = (DataflowInput))
-	TArray<float> FloatArray;
-
-	/** Int array output */
-	UPROPERTY(meta = (DataflowOutput))
-	TArray<int32> IntArray;
-
-	FFloatArrayToIntArrayDataflowNode(const Dataflow::FNodeParameters& InParam, FGuid InGuid = FGuid::NewGuid())
-		: FDataflowNode(InParam, InGuid)
-	{
-		RegisterInputConnection(&FloatArray);
-		RegisterOutputConnection(&IntArray);
-	}
-
-	virtual void Evaluate(Dataflow::FContext& Context, const FDataflowOutput* Out) const override;
-
-};
-
-/**
- *
- * Returns the specified element from an array
- *
- */
-USTRUCT()
-struct FGetFloatArrayElementDataflowNode : public FDataflowNode
-{
-	GENERATED_USTRUCT_BODY()
-	DATAFLOW_NODE_DEFINE_INTERNAL(FGetFloatArrayElementDataflowNode, "GetFloatArrayElement", "Utilities|Array", "")
-
-public:
-	/** Element index */
-	UPROPERTY(EditAnywhere, Category = "Index");
-	int32 Index = 0;
-
-	/** Array to get the element from */
-	UPROPERTY(meta = (DataflowInput))
-	TArray<float> FloatArray;
-
-	/** Specified element */
-	UPROPERTY(meta = (DataflowOutput))
-	float FloatValue = 0;
-
-	FGetFloatArrayElementDataflowNode(const Dataflow::FNodeParameters& InParam, FGuid InGuid = FGuid::NewGuid())
-		: FDataflowNode(InParam, InGuid)
-	{
-		RegisterInputConnection(&FloatArray);
-		RegisterInputConnection(&Index);
-		RegisterOutputConnection(&FloatValue);
-	}
-
-	virtual void Evaluate(Dataflow::FContext& Context, const FDataflowOutput* Out) const override;
-
-};
-
-/**
- *
- * Returns the specified element from an array
- *
- */
-USTRUCT()
-struct FGetArrayElementDataflowNode : public FDataflowNode
-{
-	GENERATED_USTRUCT_BODY()
-	DATAFLOW_NODE_DEFINE_INTERNAL(FGetArrayElementDataflowNode, "GetArrayElement", "Utilities|Array", "")
-
-public:
-	/** Element index */
-	UPROPERTY(EditAnywhere, Category = "Index");
-	int32 Index = 0;
-
-	/** Array to get the element from */
-	UPROPERTY(meta = (DataflowInput))
-	TArray<FVector> Points;
-
-	/** Specified element */
-	UPROPERTY(meta = (DataflowOutput))
-	FVector Point = FVector(0.0);
-
-	FGetArrayElementDataflowNode(const Dataflow::FNodeParameters& InParam, FGuid InGuid = FGuid::NewGuid())
-		: FDataflowNode(InParam, InGuid)
-	{
-		RegisterInputConnection(&Points);
-		RegisterInputConnection(&Index);
-		RegisterOutputConnection(&Point);
-	}
-
-	virtual void Evaluate(Dataflow::FContext& Context, const FDataflowOutput* Out) const override;
-
-};
-
-
-/**
- *
- * Returns the number of elements in an array
- *
- */
-USTRUCT()
-struct FGetNumArrayElementsDataflowNode : public FDataflowNode
-{
-	GENERATED_USTRUCT_BODY()
-	DATAFLOW_NODE_DEFINE_INTERNAL(FGetNumArrayElementsDataflowNode, "GetNumArrayElements", "Utilities|Array", "")
-
-public:
-	/** FVector array input */
-	UPROPERTY(meta = (DataflowInput, DisplayName = "VectorArray"))
-	TArray<FVector> Points;
-
-	/** FVector3f array input */
-	UPROPERTY(meta = (DataflowInput, DisplayName = "Vector3fArray"))
-	TArray<FVector3f> Vector3fArray;
-
-	/** Number of elements in the array */
-	UPROPERTY(meta = (DataflowOutput))
-	int32 NumElements = 0;
-
-	FGetNumArrayElementsDataflowNode(const Dataflow::FNodeParameters& InParam, FGuid InGuid = FGuid::NewGuid())
-		: FDataflowNode(InParam, InGuid)
-	{
-		RegisterInputConnection(&Points);
-		RegisterInputConnection(&Vector3fArray);
-		RegisterOutputConnection(&NumElements);
-	}
-
-	virtual void Evaluate(Dataflow::FContext& Context, const FDataflowOutput* Out) const override;
-
-};
-
 
 /**
  *
@@ -1542,76 +1377,6 @@ public:
 };
 
 
-/**
- *
- * Converts a TArray<bool> to a FDataflowFaceSelection
- *
- */
-USTRUCT()
-struct FBoolArrayToFaceSelectionDataflowNode : public FDataflowNode
-{
-	GENERATED_USTRUCT_BODY()
-	DATAFLOW_NODE_DEFINE_INTERNAL(FBoolArrayToFaceSelectionDataflowNode, "BoolArrayToFaceSelection", "Utilities|Array", "")
-
-public:
-	/** TArray<bool> data */
-	UPROPERTY(meta = (DataflowInput, DataflowIntrinsic));
-	TArray<bool> BoolAttributeData;
-
-	UPROPERTY(meta = (DataflowOutput, DisplayName = "FaceSelection"))
-	FDataflowFaceSelection FaceSelection;
-
-	FBoolArrayToFaceSelectionDataflowNode(const Dataflow::FNodeParameters& InParam, FGuid InGuid = FGuid::NewGuid())
-		: FDataflowNode(InParam, InGuid)
-	{
-		RegisterInputConnection(&BoolAttributeData);
-		RegisterOutputConnection(&FaceSelection);
-	}
-
-	virtual void Evaluate(Dataflow::FContext& Context, const FDataflowOutput* Out) const override;
-
-};
-
-
-/**
- *
- * Converts a TArray<float> to a FDataflowVertexSelection
- *
- */
-USTRUCT()
-struct FFloatArrayToVertexSelectionDataflowNode : public FDataflowNode
-{
-	GENERATED_USTRUCT_BODY()
-	DATAFLOW_NODE_DEFINE_INTERNAL(FFloatArrayToVertexSelectionDataflowNode, "FloatArrayToVertexSelection", "Utilities|Array", "")
-
-public:
-	/** TArray<floatl> array */
-	UPROPERTY(meta = (DataflowInput, DataflowIntrinsic));
-	TArray<float> FloatArray;
-
-	/** Comparison operation */
-	UPROPERTY(EditAnywhere, Category = "Compare");
-	ECompareOperationEnum Operation = ECompareOperationEnum::Dataflow_Compare_Greater;
-
-	/**  */
-	UPROPERTY(EditAnywhere, Category = "Compare")
-	float Threshold = 0.f;
-
-	UPROPERTY(meta = (DataflowOutput, DisplayName = "VertexSelection"))
-	FDataflowVertexSelection VertexSelection;
-
-	FFloatArrayToVertexSelectionDataflowNode(const Dataflow::FNodeParameters& InParam, FGuid InGuid = FGuid::NewGuid())
-		: FDataflowNode(InParam, InGuid)
-	{
-		RegisterInputConnection(&FloatArray);
-		RegisterOutputConnection(&VertexSelection);
-	}
-
-	virtual void Evaluate(Dataflow::FContext& Context, const FDataflowOutput* Out) const override;
-
-};
-
-
 /**
  *
  * 
@@ -1722,94 +1487,6 @@ public:
 };
 
 
-/**
- *
- * Normalize the selected float data in a FloatArray
- *
- */
-USTRUCT()
-struct FFloatArrayNormalizeDataflowNode : public FDataflowNode
-{
-	GENERATED_USTRUCT_BODY()
-	DATAFLOW_NODE_DEFINE_INTERNAL(FFloatArrayNormalizeDataflowNode, "FloatArrayNormalize", "Math|Float", "")
-
-public:
-	/** Input VectorArray */
-	UPROPERTY(meta = (DataflowInput, DataflowIntrinsic))
-	TArray<float> InFloatArray;
-
-	/** Selection for the operation */
-	UPROPERTY(meta = (DataflowInput))
-	FDataflowVertexSelection Selection;
-
-	/**  */
-	UPROPERTY(EditAnywhere, Category = "Field", meta = (DataflowInput))
-	float MinRange = 0.f;
-
-	/**  */
-	UPROPERTY(EditAnywhere, Category = "Field", meta = (DataflowInput))
-	float MaxRange = 1.f;
-
-	/**  */
-	UPROPERTY(meta = (DataflowOutput))
-	TArray<float> OutFloatArray;
-
-	FFloatArrayNormalizeDataflowNode(const Dataflow::FNodeParameters& InParam, FGuid InGuid = FGuid::NewGuid())
-		: FDataflowNode(InParam, InGuid)
-	{
-		RegisterInputConnection(&InFloatArray);
-		RegisterInputConnection(&Selection);
-		RegisterInputConnection(&MinRange);
-		RegisterInputConnection(&MaxRange);
-		RegisterOutputConnection(&OutFloatArray);
-	}
-
-	virtual void Evaluate(Dataflow::FContext& Context, const FDataflowOutput* Out) const override;
-
-};
-
-
-/**
- *
- * Normalize all the selected vectors in a VectorArray
- *
- */
-USTRUCT()
-struct FVectorArrayNormalizeDataflowNode : public FDataflowNode
-{
-	GENERATED_USTRUCT_BODY()
-	DATAFLOW_NODE_DEFINE_INTERNAL(FVectorArrayNormalizeDataflowNode, "VectorArrayNormalize", "Math|Vector", "")
-
-public:
-	/** Input VectorArray */
-	UPROPERTY(meta = (DataflowInput, DataflowIntrinsic))
-	TArray<FVector> InVectorArray;
-
-	/** Selection for the operation */
-	UPROPERTY(meta = (DataflowInput))
-	FDataflowVertexSelection Selection;
-
-	/**  */
-	UPROPERTY(EditAnywhere, Category = "Field", meta = (DataflowInput))
-	float Magnitude = 1.f;
-
-	/**  */
-	UPROPERTY(meta = (DataflowOutput))
-	TArray<FVector> OutVectorArray;
-
-	FVectorArrayNormalizeDataflowNode(const Dataflow::FNodeParameters& InParam, FGuid InGuid = FGuid::NewGuid())
-		: FDataflowNode(InParam, InGuid)
-	{
-		RegisterInputConnection(&InVectorArray);
-		RegisterInputConnection(&Selection);
-		RegisterInputConnection(&Magnitude);
-		RegisterOutputConnection(&OutVectorArray);
-	}
-
-	virtual void Evaluate(Dataflow::FContext& Context, const FDataflowOutput* Out) const override;
-
-};
-
 /**
  *
  * Description for this node
@@ -1871,32 +1548,6 @@ public:
 	virtual void Evaluate(Dataflow::FContext& Context, const FDataflowOutput* Out) const override;
 };
 
-USTRUCT()
-struct FUnionIntArraysDataflowNode : public FDataflowNode
-{
-	GENERATED_USTRUCT_BODY()
-		DATAFLOW_NODE_DEFINE_INTERNAL(FUnionIntArraysDataflowNode, "UnionIntArrays", "Utilities", "")
-
-public:
-	UPROPERTY(meta = (DataflowInput, DisplayName = "InArray1"));
-	TArray<int32> InArray1;
-
-	UPROPERTY(meta = (DataflowInput, DisplayName = "InArray2"));
-	TArray<int32> InArray2;
-
-	UPROPERTY(meta = (DataflowOutput, DisplayName = "OutArray", DataflowPassthrough = "InArray1"));
-	TArray<int32> OutArray;
-
-	FUnionIntArraysDataflowNode(const Dataflow::FNodeParameters& InParam, FGuid InGuid = FGuid::NewGuid())
-		: FDataflowNode(InParam, InGuid)
-	{
-		RegisterInputConnection(&InArray1);
-		RegisterInputConnection(&InArray2);
-		RegisterOutputConnection(&OutArray, &InArray1);
-	}
-
-	virtual void Evaluate(Dataflow::FContext& Context, const FDataflowOutput* Out) const override;
-};
 
 namespace Dataflow
 {
