Fix uninitialized variable in FMeshToOBJStringDebugDataflowNode

diff --git a/Engine/Plugins/Experimental/GeometryCollectionPlugin/Source/GeometryCollectionNodes/Public/Dataflow/GeometryCollectionDebugNodes.h b/Engine/Plugins/Experimental/GeometryCollectionPlugin/Source/GeometryCollectionNodes/Public/Dataflow/GeometryCollectionDebugNodes.h
--- a/Engine/Plugins/Experimental/GeometryCollectionPlugin/Source/GeometryCollectionNodes/Public/Dataflow/GeometryCollectionDebugNodes.h
+++ b/Engine/Plugins/Experimental/GeometryCollectionPlugin/Source/GeometryCollectionNodes/Public/Dataflow/GeometryCollectionDebugNodes.h
@@ -85,7 +85,7 @@ public:
 
 	// Whether to flip the orientation of the triangles in the OBJ output
 	UPROPERTY(EditAnywhere, Category = Options, meta = (DataflowInput))
-	bool bInvertFaces;
+	bool bInvertFaces = false;
 
 	// A string representing the input mesh in the OBJ file format
 	UPROPERTY(meta = (DataflowOutput))
