add option for convex hull dataflow to compute the volume of the union of the hulls, so that overlapping regions are not double-counted

diff --git a/Engine/Plugins/Experimental/GeometryCollectionPlugin/Source/GeometryCollectionNodes/Private/Dataflow/GeometryCollectionUtilityNodes.cpp b/Engine/Plugins/Experimental/GeometryCollectionPlugin/Source/GeometryCollectionNodes/Private/Dataflow/GeometryCollectionUtilityNodes.cpp
--- a/Engine/Plugins/Experimental/GeometryCollectionPlugin/Source/GeometryCollectionNodes/Private/Dataflow/GeometryCollectionUtilityNodes.cpp
+++ b/Engine/Plugins/Experimental/GeometryCollectionPlugin/Source/GeometryCollectionNodes/Private/Dataflow/GeometryCollectionUtilityNodes.cpp
@@ -9,6 +9,9 @@
 #include "GeometryCollection/Facades/CollectionHierarchyFacade.h"
 #include "GeometryCollection/Facades/CollectionTransformSelectionFacade.h"
 
+#include "Operations/MeshSelfUnion.h"
+#include "MeshQueries.h"
+
 #include "MeshSimplification.h"
 
 
@@ -483,21 +486,51 @@ void FGetConvexHullVolumeDataflowNode::Evaluate(Dataflow::FContext& Context, con
 		const TManagedArray<TUniquePtr<Chaos::FConvex>>& ConvexHulls = InCollection.GetAttribute<TUniquePtr<Chaos::FConvex>>("ConvexHull", "Convex");
 
 		Chaos::Facades::FCollectionHierarchyFacade HierarchyFacade(InCollection);
-		while (!SelectionToSum.IsEmpty())
+
+		auto IterateHulls = [this, &TransformToConvexIndices, &HierarchyFacade](TArray<int32>& SelectionToSum, TFunctionRef<void(int32)> ProcessFn)
 		{
-			int32 TransformIdx = SelectionToSum.Pop(false);
-			if (!bSumChildrenForClustersWithoutHulls || !TransformToConvexIndices[TransformIdx].IsEmpty())
+			while (!SelectionToSum.IsEmpty())
 			{
-				for (int32 ConvexIdx : TransformToConvexIndices[TransformIdx])
+				int32 TransformIdx = SelectionToSum.Pop(false);
+				if (!bSumChildrenForClustersWithoutHulls || !TransformToConvexIndices[TransformIdx].IsEmpty())
 				{
-					VolumeSum += ConvexHulls[ConvexIdx]->GetVolume();
+					ProcessFn(TransformIdx);
+				}
+				else if (const TSet<int32>* Children = HierarchyFacade.FindChildren(TransformIdx))
+				{
+					SelectionToSum.Append(Children->Array());
 				}
 			}
-			else if (const TSet<int32>* Children = HierarchyFacade.FindChildren(TransformIdx))
-			{
-				SelectionToSum.Append(Children->Array());
-			}
+		};
+
+		if (!bVolumeOfUnion)
+		{
+			IterateHulls(SelectionToSum, [&VolumeSum, &ConvexHulls, &TransformToConvexIndices](int32 TransformIdx)
+				{
+					for (int32 ConvexIdx : TransformToConvexIndices[TransformIdx])
+					{
+						VolumeSum += ConvexHulls[ConvexIdx]->GetVolume();
+					}
+				});
 		}
+		else
+		{
+			TArray<int32> SelectedBones;
+			SelectedBones.Reserve(SelectionToSum.Num());
+			IterateHulls(SelectionToSum, [&SelectedBones](int32 TransformIdx)
+				{
+					SelectedBones.Add(TransformIdx);
+				});
+			UE::Geometry::FDynamicMesh3 Mesh;
+			UE::FractureEngine::Convex::GetConvexHullsAsDynamicMesh(InCollection, Mesh, true, SelectedBones);
+			UE::Geometry::FMeshSelfUnion Union(&Mesh);
+			// Disable quality-related features, since we just want the volume
+			Union.TryToImproveTriQualityThreshold = -1;
+			Union.bWeldSharedEdges = false;
+			Union.Compute();
+			VolumeSum = UE::Geometry::TMeshQueries<UE::Geometry::FDynamicMesh3>::GetVolumeNonWatertight(Mesh);
+		}
+		
 		SetValue(Context, VolumeSum, &Volume);
 	}
 }
diff --git a/Engine/Plugins/Experimental/GeometryCollectionPlugin/Source/GeometryCollectionNodes/Public/Dataflow/GeometryCollectionUtilityNodes.h b/Engine/Plugins/Experimental/GeometryCollectionPlugin/Source/GeometryCollectionNodes/Public/Dataflow/GeometryCollectionUtilityNodes.h
--- a/Engine/Plugins/Experimental/GeometryCollectionPlugin/Source/GeometryCollectionNodes/Public/Dataflow/GeometryCollectionUtilityNodes.h
+++ b/Engine/Plugins/Experimental/GeometryCollectionPlugin/Source/GeometryCollectionNodes/Public/Dataflow/GeometryCollectionUtilityNodes.h
@@ -469,6 +469,10 @@ public:
 	UPROPERTY(EditAnywhere, Category = Options)
 	bool bSumChildrenForClustersWithoutHulls = true;
 
+	/** Whether to take the volume of the union of selected hulls, rather than the sum of each hull volume separately. This is more expensive but more accurate when hulls overlap. */
+	UPROPERTY(EditAnywhere, Category = Options)
+	bool bVolumeOfUnion = false;
+
 	FGetConvexHullVolumeDataflowNode(const Dataflow::FNodeParameters& InParam, FGuid InGuid = FGuid::NewGuid());
 
 	virtual void Evaluate(Dataflow::FContext& Context, const FDataflowOutput* Out) const override;
