Dataflow:
- Allows for setting output values with move semantic to avoid uncessary copies.
- Update existing nodes SetValues with MoveTemp whenever possible.

#rnx

diff --git a/Engine/Plugins/Experimental/ChaosClothAssetEditor/Source/ChaosClothAssetDataflowNodes/Private/ChaosClothAsset/AddWeightMapNode.cpp b/Engine/Plugins/Experimental/ChaosClothAssetEditor/Source/ChaosClothAssetDataflowNodes/Private/ChaosClothAsset/AddWeightMapNode.cpp
--- a/Engine/Plugins/Experimental/ChaosClothAssetEditor/Source/ChaosClothAssetDataflowNodes/Private/ChaosClothAsset/AddWeightMapNode.cpp
+++ b/Engine/Plugins/Experimental/ChaosClothAssetEditor/Source/ChaosClothAssetDataflowNodes/Private/ChaosClothAsset/AddWeightMapNode.cpp
@@ -25,7 +25,7 @@ void FChaosClothAssetAddWeightMapNode::Evaluate(Dataflow::FContext& Context, con
 	if (Out->IsA<FManagedArrayCollection>(&Collection))
 	{
 		// Evaluate in collection
-		const FString InNameString = GetValue<FString>(Context, &Name);
+		FString InNameString = GetValue<FString>(Context, &Name);
 
 		FManagedArrayCollection InCollection = GetValue<FManagedArrayCollection>(Context, &Collection);
 		const TSharedRef<FManagedArrayCollection> ClothCollection = MakeShared<FManagedArrayCollection>(MoveTemp(InCollection));
@@ -44,14 +44,14 @@ void FChaosClothAssetAddWeightMapNode::Evaluate(Dataflow::FContext& Context, con
 						VertexWeights.Num(),
 						ClothWeights.Num()));
 			}
-				
+
 			for (int32 VertexID = 0; VertexID < MaxWeightIndex; ++VertexID)
 			{
 				ClothWeights[VertexID] = VertexWeights[VertexID];
 			}
 		}
-		SetValue<FManagedArrayCollection>(Context, *ClothCollection, &Collection);
-		SetValue<FString>(Context, InNameString, &Name);
+		SetValue(Context, MoveTemp(*ClothCollection), &Collection);
+		SetValue(Context, MoveTemp(InNameString), &Name);
 	}
 }
 
diff --git a/Engine/Plugins/Experimental/ChaosClothAssetEditor/Source/ChaosClothAssetDataflowNodes/Private/ChaosClothAsset/BindToRootBoneNode.cpp b/Engine/Plugins/Experimental/ChaosClothAssetEditor/Source/ChaosClothAssetDataflowNodes/Private/ChaosClothAsset/BindToRootBoneNode.cpp
--- a/Engine/Plugins/Experimental/ChaosClothAssetEditor/Source/ChaosClothAssetDataflowNodes/Private/ChaosClothAsset/BindToRootBoneNode.cpp
+++ b/Engine/Plugins/Experimental/ChaosClothAssetEditor/Source/ChaosClothAssetDataflowNodes/Private/ChaosClothAsset/BindToRootBoneNode.cpp
@@ -27,7 +27,7 @@ void FChaosClothAssetBindToRootBoneNode::Evaluate(Dataflow::FContext& Context, c
 
 		FClothGeometryTools::BindMeshToRootBone(ClothCollection, bBindSimMesh, bBindRenderMesh);
 
-		SetValue<FManagedArrayCollection>(Context, *ClothCollection, &Collection);
+		SetValue(Context, MoveTemp(*ClothCollection), &Collection);
 	}
 }
 
diff --git a/Engine/Plugins/Experimental/ChaosClothAssetEditor/Source/ChaosClothAssetDataflowNodes/Private/ChaosClothAsset/CopySimulationToRenderMeshNode.cpp b/Engine/Plugins/Experimental/ChaosClothAssetEditor/Source/ChaosClothAssetDataflowNodes/Private/ChaosClothAsset/CopySimulationToRenderMeshNode.cpp
--- a/Engine/Plugins/Experimental/ChaosClothAssetEditor/Source/ChaosClothAssetDataflowNodes/Private/ChaosClothAsset/CopySimulationToRenderMeshNode.cpp
+++ b/Engine/Plugins/Experimental/ChaosClothAssetEditor/Source/ChaosClothAssetDataflowNodes/Private/ChaosClothAsset/CopySimulationToRenderMeshNode.cpp
@@ -38,7 +38,7 @@ void FChaosClothAssetCopySimulationToRenderMeshNode::Evaluate(Dataflow::FContext
 		// Copy the mesh data
 		FClothGeometryTools::CopySimMeshToRenderMesh(ClothCollection, MaterialPathName, bGenerateSingleRenderPattern);
 
-		SetValue<FManagedArrayCollection>(Context, *ClothCollection, &Collection);
+		SetValue(Context, MoveTemp(*ClothCollection), &Collection);
 	}
 }
 
diff --git a/Engine/Plugins/Experimental/ChaosClothAssetEditor/Source/ChaosClothAssetDataflowNodes/Private/ChaosClothAsset/DatasmithImportNode.cpp b/Engine/Plugins/Experimental/ChaosClothAssetEditor/Source/ChaosClothAssetDataflowNodes/Private/ChaosClothAsset/DatasmithImportNode.cpp
--- a/Engine/Plugins/Experimental/ChaosClothAssetEditor/Source/ChaosClothAssetDataflowNodes/Private/ChaosClothAsset/DatasmithImportNode.cpp
+++ b/Engine/Plugins/Experimental/ChaosClothAssetEditor/Source/ChaosClothAssetDataflowNodes/Private/ChaosClothAsset/DatasmithImportNode.cpp
@@ -123,7 +123,7 @@ void FChaosClothAssetDatasmithImportNode::Evaluate(Dataflow::FContext& Context,
 				}
 			}
 		}
-		SetValue<FManagedArrayCollection>(Context, ImportCache, &Collection);
+		SetValue(Context, ImportCache, &Collection);
 	}
 }
 
diff --git a/Engine/Plugins/Experimental/ChaosClothAssetEditor/Source/ChaosClothAssetDataflowNodes/Private/ChaosClothAsset/DeleteRenderMeshNode.cpp b/Engine/Plugins/Experimental/ChaosClothAssetEditor/Source/ChaosClothAssetDataflowNodes/Private/ChaosClothAsset/DeleteRenderMeshNode.cpp
--- a/Engine/Plugins/Experimental/ChaosClothAssetEditor/Source/ChaosClothAssetDataflowNodes/Private/ChaosClothAsset/DeleteRenderMeshNode.cpp
+++ b/Engine/Plugins/Experimental/ChaosClothAssetEditor/Source/ChaosClothAssetDataflowNodes/Private/ChaosClothAsset/DeleteRenderMeshNode.cpp
@@ -35,7 +35,7 @@ void FChaosClothAssetDeleteRenderMeshNode::Evaluate(Dataflow::FContext& Context,
 			DataflowNodes::LogAndToastWarning(LOCTEXT("DeletePerPatternNotImplemented", "FClothAssetDeleteRenderMeshNode: Delete per patterns not implemented."));
 		}
 
-		SetValue<FManagedArrayCollection>(Context, *ClothCollection, &Collection);
+		SetValue(Context, MoveTemp(*ClothCollection), &Collection);
 	}
 }
 
diff --git a/Engine/Plugins/Experimental/ChaosClothAssetEditor/Source/ChaosClothAssetDataflowNodes/Private/ChaosClothAsset/ImportNode.cpp b/Engine/Plugins/Experimental/ChaosClothAssetEditor/Source/ChaosClothAssetDataflowNodes/Private/ChaosClothAsset/ImportNode.cpp
--- a/Engine/Plugins/Experimental/ChaosClothAssetEditor/Source/ChaosClothAssetDataflowNodes/Private/ChaosClothAsset/ImportNode.cpp
+++ b/Engine/Plugins/Experimental/ChaosClothAssetEditor/Source/ChaosClothAssetDataflowNodes/Private/ChaosClothAsset/ImportNode.cpp
@@ -71,7 +71,7 @@ void FChaosClothAssetImportNode::Evaluate(Dataflow::FContext& Context, const FDa
 				ClothFacade.Initialize(InClothFacade);
 			}
 		}
-		SetValue<FManagedArrayCollection>(Context, *ClothCollection, &Collection);
+		SetValue(Context, MoveTemp(*ClothCollection), &Collection);
 }
 }
 
diff --git a/Engine/Plugins/Experimental/ChaosClothAssetEditor/Source/ChaosClothAssetDataflowNodes/Private/ChaosClothAsset/MergeClothCollectionsNode.cpp b/Engine/Plugins/Experimental/ChaosClothAssetEditor/Source/ChaosClothAssetDataflowNodes/Private/ChaosClothAsset/MergeClothCollectionsNode.cpp
--- a/Engine/Plugins/Experimental/ChaosClothAssetEditor/Source/ChaosClothAssetDataflowNodes/Private/ChaosClothAsset/MergeClothCollectionsNode.cpp
+++ b/Engine/Plugins/Experimental/ChaosClothAssetEditor/Source/ChaosClothAssetDataflowNodes/Private/ChaosClothAsset/MergeClothCollectionsNode.cpp
@@ -38,7 +38,7 @@ void FChaosClothAssetMergeClothCollectionsNode::Evaluate(Dataflow::FContext& Con
 			ClothFacade.Append(OtherClothFacade);
 		}
 
-		SetValue<FManagedArrayCollection>(Context, *ClothCollection, &Collection);
+		SetValue(Context, MoveTemp(*ClothCollection), &Collection);
 	}
 }
 
diff --git a/Engine/Plugins/Experimental/ChaosClothAssetEditor/Source/ChaosClothAssetDataflowNodes/Private/ChaosClothAsset/PatternSelectionNode.cpp b/Engine/Plugins/Experimental/ChaosClothAssetEditor/Source/ChaosClothAssetDataflowNodes/Private/ChaosClothAsset/PatternSelectionNode.cpp
--- a/Engine/Plugins/Experimental/ChaosClothAssetEditor/Source/ChaosClothAssetDataflowNodes/Private/ChaosClothAsset/PatternSelectionNode.cpp
+++ b/Engine/Plugins/Experimental/ChaosClothAssetEditor/Source/ChaosClothAssetDataflowNodes/Private/ChaosClothAsset/PatternSelectionNode.cpp
@@ -18,8 +18,8 @@ void FChaosClothAssetPatternSelectionNode::Evaluate(Dataflow::FContext& Context,
 	if (Out->IsA<TArray<int32>>(&Patterns))
 	{
 		// Init with empty pattern indices
-		const TArray<int32> Selection;
-		SetValue<TArray<int32>>(Context, Selection, &Patterns);
+		TArray<int32> Selection;
+		SetValue(Context, MoveTemp(Selection), &Patterns);
 	}
 }
 
diff --git a/Engine/Plugins/Experimental/ChaosClothAssetEditor/Source/ChaosClothAssetDataflowNodes/Private/ChaosClothAsset/ReverseNormalsNode.cpp b/Engine/Plugins/Experimental/ChaosClothAssetEditor/Source/ChaosClothAssetDataflowNodes/Private/ChaosClothAsset/ReverseNormalsNode.cpp
--- a/Engine/Plugins/Experimental/ChaosClothAssetEditor/Source/ChaosClothAssetDataflowNodes/Private/ChaosClothAsset/ReverseNormalsNode.cpp
+++ b/Engine/Plugins/Experimental/ChaosClothAssetEditor/Source/ChaosClothAssetDataflowNodes/Private/ChaosClothAsset/ReverseNormalsNode.cpp
@@ -36,7 +36,7 @@ void FChaosClothAssetReverseNormalsNode::Evaluate(Dataflow::FContext& Context, c
 			GetValue<TArray<int32>>(Context, &SimPatterns),
 			GetValue<TArray<int32>>(Context, &RenderPatterns));
 
-		SetValue<FManagedArrayCollection>(Context, *ClothCollection, &Collection);
+		SetValue(Context, MoveTemp(*ClothCollection), &Collection);
 	}
 }
 
diff --git a/Engine/Plugins/Experimental/ChaosClothAssetEditor/Source/ChaosClothAssetDataflowNodes/Private/ChaosClothAsset/SetPhysicsAssetNode.cpp b/Engine/Plugins/Experimental/ChaosClothAssetEditor/Source/ChaosClothAssetDataflowNodes/Private/ChaosClothAsset/SetPhysicsAssetNode.cpp
--- a/Engine/Plugins/Experimental/ChaosClothAssetEditor/Source/ChaosClothAssetDataflowNodes/Private/ChaosClothAsset/SetPhysicsAssetNode.cpp
+++ b/Engine/Plugins/Experimental/ChaosClothAssetEditor/Source/ChaosClothAssetDataflowNodes/Private/ChaosClothAsset/SetPhysicsAssetNode.cpp
@@ -30,7 +30,7 @@ void FChaosClothAssetSetPhysicsAssetNode::Evaluate(Dataflow::FContext& Context,
 		FCollectionClothFacade CollectionClothFacade(ClothCollection);
 		CollectionClothFacade.SetPhysicsAssetPathName(PhysicsAsset.GetPathName());
 
-		SetValue<FManagedArrayCollection>(Context, *ClothCollection, &Collection);
+		SetValue(Context, MoveTemp(*ClothCollection), &Collection);
 	}
 }
 
diff --git a/Engine/Plugins/Experimental/ChaosClothAssetEditor/Source/ChaosClothAssetDataflowNodes/Private/ChaosClothAsset/SimulationBaseConfigNode.cpp b/Engine/Plugins/Experimental/ChaosClothAssetEditor/Source/ChaosClothAssetDataflowNodes/Private/ChaosClothAsset/SimulationBaseConfigNode.cpp
--- a/Engine/Plugins/Experimental/ChaosClothAssetEditor/Source/ChaosClothAssetDataflowNodes/Private/ChaosClothAsset/SimulationBaseConfigNode.cpp
+++ b/Engine/Plugins/Experimental/ChaosClothAssetEditor/Source/ChaosClothAssetDataflowNodes/Private/ChaosClothAsset/SimulationBaseConfigNode.cpp
@@ -66,7 +66,7 @@ void FChaosClothAssetSimulationBaseConfigNode::Evaluate(Dataflow::FContext& Cont
 
 		EvaluateClothCollection(Context, ClothCollection);
 
-		SetValue<FManagedArrayCollection>(Context, *ClothCollection, &Collection);
+		SetValue(Context, MoveTemp(*ClothCollection), &Collection);
 	}
 }
 
diff --git a/Engine/Plugins/Experimental/ChaosClothAssetEditor/Source/ChaosClothAssetDataflowNodes/Private/ChaosClothAsset/SimulationDefaultConfigNode.cpp b/Engine/Plugins/Experimental/ChaosClothAssetEditor/Source/ChaosClothAssetDataflowNodes/Private/ChaosClothAsset/SimulationDefaultConfigNode.cpp
--- a/Engine/Plugins/Experimental/ChaosClothAssetEditor/Source/ChaosClothAssetDataflowNodes/Private/ChaosClothAsset/SimulationDefaultConfigNode.cpp
+++ b/Engine/Plugins/Experimental/ChaosClothAssetEditor/Source/ChaosClothAssetDataflowNodes/Private/ChaosClothAsset/SimulationDefaultConfigNode.cpp
@@ -60,7 +60,7 @@ void FChaosClothAssetSimulationDefaultConfigNode::Evaluate(Dataflow::FContext& C
 
 		UE::Chaos::ClothAsset::FClothEngineTools::GenerateTethers(ClothCollection, MaxDistanceName, bUseGeodesicTethers);
 
-		SetValue<FManagedArrayCollection>(Context, *ClothCollection, &Collection);
+		SetValue(Context, MoveTemp(*ClothCollection), &Collection);
 	}
 }
 
diff --git a/Engine/Plugins/Experimental/ChaosClothAssetEditor/Source/ChaosClothAssetDataflowNodes/Private/ChaosClothAsset/SkeletalMeshImportNode.cpp b/Engine/Plugins/Experimental/ChaosClothAssetEditor/Source/ChaosClothAssetDataflowNodes/Private/ChaosClothAsset/SkeletalMeshImportNode.cpp
--- a/Engine/Plugins/Experimental/ChaosClothAssetEditor/Source/ChaosClothAssetDataflowNodes/Private/ChaosClothAsset/SkeletalMeshImportNode.cpp
+++ b/Engine/Plugins/Experimental/ChaosClothAssetEditor/Source/ChaosClothAssetDataflowNodes/Private/ChaosClothAsset/SkeletalMeshImportNode.cpp
@@ -37,7 +37,7 @@ void FChaosClothAssetSkeletalMeshImportNode::Evaluate(Dataflow::FContext& Contex
 			const bool bIsValidLOD = ImportedModel && ImportedModel->LODModels.IsValidIndex(LODIndex);
 			if (!ensureAlways(bIsValidLOD))
 			{
-				SetValue<FManagedArrayCollection>(Context, *ClothCollection, &Collection);
+				SetValue(Context, MoveTemp(*ClothCollection), &Collection);
 				return;
 			}
 			const FSkeletalMeshLODModel &LODModel = ImportedModel->LODModels[LODIndex];
@@ -45,7 +45,7 @@ void FChaosClothAssetSkeletalMeshImportNode::Evaluate(Dataflow::FContext& Contex
 			const bool bIsValidSection = LODModel.Sections.IsValidIndex(SectionIndex);;
 			if (!ensureAlways(bIsValidSection))
 			{
-				SetValue<FManagedArrayCollection>(Context, *ClothCollection, &Collection);	
+				SetValue(Context, MoveTemp(*ClothCollection), &Collection);
 				return;
 			}
 
@@ -75,7 +75,7 @@ void FChaosClothAssetSkeletalMeshImportNode::Evaluate(Dataflow::FContext& Contex
 				ClothFacade.SetSkeletonAssetPathName(Skeleton->GetPathName());
 			}
 		}
-		SetValue<FManagedArrayCollection>(Context, *ClothCollection, &Collection);	
+		SetValue(Context, MoveTemp(*ClothCollection), &Collection);
 	}
 }
 
diff --git a/Engine/Plugins/Experimental/ChaosClothAssetEditor/Source/ChaosClothAssetDataflowNodes/Private/ChaosClothAsset/StaticMeshImportNode.cpp b/Engine/Plugins/Experimental/ChaosClothAssetEditor/Source/ChaosClothAssetDataflowNodes/Private/ChaosClothAsset/StaticMeshImportNode.cpp
--- a/Engine/Plugins/Experimental/ChaosClothAssetEditor/Source/ChaosClothAssetDataflowNodes/Private/ChaosClothAsset/StaticMeshImportNode.cpp
+++ b/Engine/Plugins/Experimental/ChaosClothAssetEditor/Source/ChaosClothAssetDataflowNodes/Private/ChaosClothAsset/StaticMeshImportNode.cpp
@@ -167,7 +167,7 @@ void FChaosClothAssetStaticMeshImportNode::Evaluate(Dataflow::FContext& Context,
 			}
 		}
 
-		SetValue<FManagedArrayCollection>(Context, *ClothCollection, &Collection);
+		SetValue(Context, MoveTemp(*ClothCollection), &Collection);
 	}
 }
 
diff --git a/Engine/Plugins/Experimental/ChaosClothAssetEditor/Source/ChaosClothAssetDataflowNodes/Private/ChaosClothAsset/TransferSkinWeightsNode.cpp b/Engine/Plugins/Experimental/ChaosClothAssetEditor/Source/ChaosClothAssetDataflowNodes/Private/ChaosClothAsset/TransferSkinWeightsNode.cpp
--- a/Engine/Plugins/Experimental/ChaosClothAssetEditor/Source/ChaosClothAssetDataflowNodes/Private/ChaosClothAsset/TransferSkinWeightsNode.cpp
+++ b/Engine/Plugins/Experimental/ChaosClothAssetEditor/Source/ChaosClothAssetDataflowNodes/Private/ChaosClothAsset/TransferSkinWeightsNode.cpp
@@ -206,7 +206,7 @@ void FChaosClothAssetTransferSkinWeightsNode::Evaluate(Dataflow::FContext& Conte
 			if (!InputSkeletalMesh->IsValidLODIndex(SkeletalMeshLOD))
 			{
 				UE::Chaos::ClothAsset::DataflowNodes::LogAndToastWarning(LOCTEXT("Warning_TransferWeightsInvalidSkeletalMeshLOD", "TransferSkinWeightsNode: The specified input SkeletalMesh LOD is not valid."));
-				SetValue<FManagedArrayCollection>(Context, *ClothCollection, &Collection);
+				SetValue(Context, MoveTemp(*ClothCollection), &Collection);
 				return;
 			}
 
@@ -248,7 +248,7 @@ void FChaosClothAssetTransferSkinWeightsNode::Evaluate(Dataflow::FContext& Conte
 			}
 		}
 
-		SetValue<FManagedArrayCollection>(Context, *ClothCollection, &Collection);
+		SetValue(Context, MoveTemp(*ClothCollection), &Collection);
 	}
 }
 
diff --git a/Engine/Plugins/Experimental/ChaosFlesh/Source/ChaosFleshNodes/Private/Dataflow/ChaosFleshBindingsNodes.cpp b/Engine/Plugins/Experimental/ChaosFlesh/Source/ChaosFleshNodes/Private/Dataflow/ChaosFleshBindingsNodes.cpp
--- a/Engine/Plugins/Experimental/ChaosFlesh/Source/ChaosFleshNodes/Private/Dataflow/ChaosFleshBindingsNodes.cpp
+++ b/Engine/Plugins/Experimental/ChaosFlesh/Source/ChaosFleshNodes/Private/Dataflow/ChaosFleshBindingsNodes.cpp
@@ -555,6 +555,6 @@ FGenerateSurfaceBindings::Evaluate(Dataflow::FContext& Context, const FDataflowO
 				}
 			} // end for TetMeshIdx
 		}
-		SetValue<DataType>(Context, OutCollection, &Collection);
+		SetValue(Context, MoveTemp(OutCollection), &Collection);
 	}
 }
diff --git a/Engine/Plugins/Experimental/ChaosFlesh/Source/ChaosFleshNodes/Private/Dataflow/ChaosFleshCoreNodes.cpp b/Engine/Plugins/Experimental/ChaosFlesh/Source/ChaosFleshNodes/Private/Dataflow/ChaosFleshCoreNodes.cpp
--- a/Engine/Plugins/Experimental/ChaosFlesh/Source/ChaosFleshNodes/Private/Dataflow/ChaosFleshCoreNodes.cpp
+++ b/Engine/Plugins/Experimental/ChaosFlesh/Source/ChaosFleshNodes/Private/Dataflow/ChaosFleshCoreNodes.cpp
@@ -37,9 +37,9 @@ void FAppendTetrahedralCollectionDataflowNode::Evaluate(Dataflow::FContext& Cont
 		{
 			GeometryGroupGuidsLocal2 = GuidArray2->GetConstArray();
 		}
-		SetValue<FManagedArrayCollection>(Context, *InCollection1, &Collection1);
-		SetValue<TArray<FString>>(Context, GeometryGroupGuidsLocal1, &GeometryGroupGuidsOut1);
-		SetValue<TArray<FString>>(Context, GeometryGroupGuidsLocal2, &GeometryGroupGuidsOut2);
+		SetValue(Context, MoveTemp((FManagedArrayCollection&)(*InCollection1)), &Collection1);
+		SetValue(Context, MoveTemp(GeometryGroupGuidsLocal1), &GeometryGroupGuidsOut1);
+		SetValue(Context, MoveTemp(GeometryGroupGuidsLocal2), &GeometryGroupGuidsOut2);
 	}
 }
 
diff --git a/Engine/Plugins/Experimental/ChaosFlesh/Source/ChaosFleshNodes/Private/Dataflow/ChaosFleshEngineAssetNodes.cpp b/Engine/Plugins/Experimental/ChaosFlesh/Source/ChaosFleshNodes/Private/Dataflow/ChaosFleshEngineAssetNodes.cpp
--- a/Engine/Plugins/Experimental/ChaosFlesh/Source/ChaosFleshNodes/Private/Dataflow/ChaosFleshEngineAssetNodes.cpp
+++ b/Engine/Plugins/Experimental/ChaosFlesh/Source/ChaosFleshNodes/Private/Dataflow/ChaosFleshEngineAssetNodes.cpp
@@ -36,7 +36,7 @@ void FGetFleshAssetDataflowNode::Evaluate(Dataflow::FContext& Context, const FDa
 	if (Out->IsA<FManagedArrayCollection>(&Output))
 	{
 		FManagedArrayCollection Collection;
-		SetValue<FManagedArrayCollection>(Context, Collection, &Output);
+		SetValue(Context, MoveTemp(Collection), &Output);
 
 		const UFleshAsset* FleshAssetValue = FleshAsset;
 		if (!FleshAssetValue)
@@ -51,7 +51,7 @@ void FGetFleshAssetDataflowNode::Evaluate(Dataflow::FContext& Context, const FDa
 		{
 			if (const FFleshCollection* AssetCollection = FleshAssetValue->GetCollection())
 			{
-				SetValue<FManagedArrayCollection>(Context, *AssetCollection, &Output);
+				SetValue(Context, (const FManagedArrayCollection&)(*AssetCollection), &Output);
 			}
 		}
 	}
@@ -69,7 +69,7 @@ void FFleshAssetTerminalDataflowNode::SetAssetValue(TObjectPtr<UObject> Asset, D
 void FFleshAssetTerminalDataflowNode::Evaluate(Dataflow::FContext& Context, const FDataflowOutput* Out) const
 {
 	const FManagedArrayCollection& InCollection = GetValue<FManagedArrayCollection>(Context, &Collection);
-	SetValue<FManagedArrayCollection>(Context, InCollection, &Collection);
+	SetValue(Context, InCollection, &Collection);
 }
 
 template <class T> using MType = FManagedArrayCollection::TManagedType<T>;
@@ -208,7 +208,7 @@ void FSetFleshDefaultPropertiesNode::Evaluate(Dataflow::FContext& Context, const
 					MinV, AvgV, MaxV);
 			}
 		}
-		SetValue<FManagedArrayCollection>(Context, InCollection, &Collection);
+		SetValue(Context, MoveTemp(InCollection), &Collection);
 	}
 }
 
@@ -232,7 +232,7 @@ void FComputeFiberFieldNode::Evaluate(Dataflow::FContext& Context, const FDatafl
 			UE_LOG(LogChaosFlesh, Warning,
 				TEXT("ComputeFiberFieldNode: Failed to find geometry collection attr '%s' in group '%s'"),
 				*FTetrahedralCollection::TetrahedronAttribute.ToString(), *FTetrahedralCollection::TetrahedralGroup.ToString());
-			Out->SetValue<FManagedArrayCollection>(InCollection, Context);
+			Out->SetValue(MoveTemp(InCollection), Context);
 			return;
 		}
 
@@ -242,7 +242,7 @@ void FComputeFiberFieldNode::Evaluate(Dataflow::FContext& Context, const FDatafl
 		{
 			UE_LOG(LogChaosFlesh, Warning,
 				TEXT("ComputeFiberFieldNode: Failed to find geometry collection attr 'Vertex' in group 'Vertices'"));
-			Out->SetValue<FManagedArrayCollection>(InCollection, Context);
+			Out->SetValue(MoveTemp(InCollection), Context);
 			return;
 		}
 
@@ -254,7 +254,7 @@ void FComputeFiberFieldNode::Evaluate(Dataflow::FContext& Context, const FDatafl
 			UE_LOG(LogChaosFlesh, Warning,
 				TEXT("ComputeFiberFieldNode: Failed to find geometry collection attr '%s' in group '%s'"),
 				*FTetrahedralCollection::IncidentElementsAttribute.ToString(), *FGeometryCollection::VerticesGroup.ToString());
-			Out->SetValue<FManagedArrayCollection>(InCollection, Context);
+			Out->SetValue(MoveTemp(InCollection), Context);
 			return;
 		}
 		TManagedArray<TArray<int32>>* IncidentElementsLocalIndex = InCollection.FindAttribute<TArray<int32>>(
@@ -264,7 +264,7 @@ void FComputeFiberFieldNode::Evaluate(Dataflow::FContext& Context, const FDatafl
 			UE_LOG(LogChaosFlesh, Warning,
 				TEXT("ComputeFiberFieldNode: Failed to find geometry collection attr '%s' in group '%s'"),
 				*FTetrahedralCollection::IncidentElementsLocalIndexAttribute.ToString(), *FGeometryCollection::VerticesGroup.ToString());
-			Out->SetValue<FManagedArrayCollection>(InCollection, Context);
+			Out->SetValue(MoveTemp(InCollection), Context);
 			return;
 		}
 
@@ -282,7 +282,7 @@ void FComputeFiberFieldNode::Evaluate(Dataflow::FContext& Context, const FDatafl
 			if (OriginInsertionGroupName.IsEmpty())
 			{
 				UE_LOG(LogChaosFlesh, Warning, TEXT("ComputeFiberFieldNode: Attr 'OriginInsertionGroupName' cannot be empty."));
-				Out->SetValue<FManagedArrayCollection>(InCollection, Context);
+				Out->SetValue(MoveTemp(InCollection), Context);
 				return;
 			}
 
@@ -292,7 +292,7 @@ void FComputeFiberFieldNode::Evaluate(Dataflow::FContext& Context, const FDatafl
 				if (OriginVertexFieldName.IsEmpty())
 				{
 					UE_LOG(LogChaosFlesh, Warning, TEXT("ComputeFiberFieldNode: Attr 'OriginVertexFieldName' cannot be empty."));
-					Out->SetValue<FManagedArrayCollection>(InCollection, Context);
+					Out->SetValue(MoveTemp(InCollection), Context);
 					return;
 				}
 				Origin = InCollection.FindAttribute<int32>(FName(OriginVertexFieldName), FName(OriginInsertionGroupName));
@@ -301,7 +301,7 @@ void FComputeFiberFieldNode::Evaluate(Dataflow::FContext& Context, const FDatafl
 					UE_LOG(LogChaosFlesh, Warning,
 						TEXT("ComputeFiberFieldNode: Failed to find geometry collection attr '%s' in group '%s'"),
 						*OriginVertexFieldName, *OriginInsertionGroupName);
-					Out->SetValue<FManagedArrayCollection>(InCollection, Context);
+					Out->SetValue(MoveTemp(InCollection), Context);
 					return;
 				}
 			}
@@ -312,7 +312,7 @@ void FComputeFiberFieldNode::Evaluate(Dataflow::FContext& Context, const FDatafl
 				if (InsertionVertexFieldName.IsEmpty())
 				{
 					UE_LOG(LogChaosFlesh, Warning, TEXT("ComputeFiberFieldNode: Attr 'InsertionVertexFieldName' cannot be empty."));
-					Out->SetValue<FManagedArrayCollection>(InCollection, Context);
+					Out->SetValue(MoveTemp(InCollection), Context);
 					return;
 				}
 				Insertion = InCollection.FindAttribute<int32>(FName(InsertionVertexFieldName), FName(OriginInsertionGroupName));
@@ -321,7 +321,7 @@ void FComputeFiberFieldNode::Evaluate(Dataflow::FContext& Context, const FDatafl
 					UE_LOG(LogChaosFlesh, Warning,
 						TEXT("ComputeFiberFieldNode: Failed to find geometry collection attr '%s' in group '%s'"),
 						*InsertionVertexFieldName, *OriginInsertionGroupName);
-					Out->SetValue<FManagedArrayCollection>(InCollection, Context);
+					Out->SetValue(MoveTemp(InCollection), Context);
 					return;
 				}
 			}
@@ -349,7 +349,7 @@ void FComputeFiberFieldNode::Evaluate(Dataflow::FContext& Context, const FDatafl
 		}
 		(*FiberDirections) = MoveTemp(FiberDirs);
 
-		Out->SetValue<FManagedArrayCollection>(InCollection, Context);
+		Out->SetValue(MoveTemp(InCollection), Context);
 	}
 }
 
@@ -436,7 +436,7 @@ void FComputeIslandsNode::Evaluate(Dataflow::FContext& Context, const FDataflowO
 			}
 		}
 		
-		SetValue<FManagedArrayCollection>(Context, InCollection, &Collection);
+		SetValue(Context, MoveTemp(InCollection), &Collection);
 	}
 }
 
@@ -462,7 +462,7 @@ void FGenerateOriginInsertionNode::Evaluate(Dataflow::FContext& Context, const F
 			UE_LOG(LogChaosFlesh, Warning,
 				TEXT("GenerateOriginInsertionNode: Failed to find geometry collection attr '%s' in group '%s'"),
 				*FTetrahedralCollection::TetrahedronAttribute.ToString(), *FTetrahedralCollection::TetrahedralGroup.ToString());
-			Out->SetValue<FManagedArrayCollection>(InCollection, Context);
+			Out->SetValue(MoveTemp(InCollection), Context);
 			return;
 		}
 
@@ -473,7 +473,7 @@ void FGenerateOriginInsertionNode::Evaluate(Dataflow::FContext& Context, const F
 		{
 			UE_LOG(LogChaosFlesh, Warning,
 				TEXT("GenerateOriginInsertionNode: Failed to find geometry collection attr 'Vertex' in group 'Vertices'"));
-			Out->SetValue<FManagedArrayCollection>(InCollection, Context);
+			Out->SetValue(MoveTemp(InCollection), Context);
 			return;
 		}
 
@@ -485,7 +485,7 @@ void FGenerateOriginInsertionNode::Evaluate(Dataflow::FContext& Context, const F
 			UE_LOG(LogChaosFlesh, Warning,
 				TEXT("GenerateOriginInsertionNode: Failed to find geometry collection attr '%s' in group '%s'"),
 				*FTetrahedralCollection::IncidentElementsAttribute.ToString(), *FGeometryCollection::VerticesGroup.ToString());
-			Out->SetValue<FManagedArrayCollection>(InCollection, Context);
+			Out->SetValue(MoveTemp(InCollection), Context);
 			return;
 		}
 		TManagedArray<TArray<int32>>* IncidentElementsLocalIndex = InCollection.FindAttribute<TArray<int32>>(
@@ -495,7 +495,7 @@ void FGenerateOriginInsertionNode::Evaluate(Dataflow::FContext& Context, const F
 			UE_LOG(LogChaosFlesh, Warning,
 				TEXT("GenerateOriginInsertionNode: Failed to find geometry collection attr '%s' in group '%s'"),
 				*FTetrahedralCollection::IncidentElementsLocalIndexAttribute.ToString(), *FGeometryCollection::VerticesGroup.ToString());
-			Out->SetValue<FManagedArrayCollection>(InCollection, Context);
+			Out->SetValue(MoveTemp(InCollection), Context);
 			return;
 		}
 
@@ -553,9 +553,9 @@ void FGenerateOriginInsertionNode::Evaluate(Dataflow::FContext& Context, const F
 		// Set output(s)
 		//
 
-		SetValue<FManagedArrayCollection>(Context, InCollection, &Collection);
-		SetValue<TArray<int32>>(Context, OutOriginIndices, &OriginIndicesOut);
-		SetValue<TArray<int32>>(Context, OutInsertionIndices, &InsertionIndicesOut);
+		SetValue(Context, MoveTemp(InCollection), &Collection);
+		SetValue(Context, MoveTemp(OutOriginIndices), &OriginIndicesOut);
+		SetValue(Context, MoveTemp(OutInsertionIndices), &InsertionIndicesOut);
 	}
 }
 
@@ -597,7 +597,7 @@ void FIsolateComponentNode::Evaluate(Dataflow::FContext& Context, const FDataflo
 			DeleteList.Sort();
 			InCollection.RemoveElements(FGeometryCollection::FacesGroup, DeleteList);
 		}
-		SetValue<FManagedArrayCollection>(Context, InCollection, &Collection);
+		SetValue(Context, MoveTemp(InCollection), &Collection);
 	}
 }
 
@@ -643,5 +643,5 @@ void FGetSurfaceIndicesNode::Evaluate(Dataflow::FContext& Context, const FDatafl
 			}
 		}
 	}
-	SetValue<TArray<int32>>(Context, SurfaceIndicesLocal, &SurfaceIndicesOut);
+	SetValue(Context, MoveTemp(SurfaceIndicesLocal), &SurfaceIndicesOut);
 }
\ No newline at end of file
diff --git a/Engine/Plugins/Experimental/ChaosFlesh/Source/ChaosFleshNodes/Private/Dataflow/ChaosFleshFiberDirectionInitializationNodes.cpp b/Engine/Plugins/Experimental/ChaosFlesh/Source/ChaosFleshNodes/Private/Dataflow/ChaosFleshFiberDirectionInitializationNodes.cpp
--- a/Engine/Plugins/Experimental/ChaosFlesh/Source/ChaosFleshNodes/Private/Dataflow/ChaosFleshFiberDirectionInitializationNodes.cpp
+++ b/Engine/Plugins/Experimental/ChaosFlesh/Source/ChaosFleshNodes/Private/Dataflow/ChaosFleshFiberDirectionInitializationNodes.cpp
@@ -23,7 +23,7 @@ void FGenerateFiberDirectionsDataflowNode::Evaluate(Dataflow::FContext& Context,
 		// @todo(dataflow) : Implemention fiber direction
 		//
 
-		SetValue<DataType>(Context, InCollection, &Collection);
+		SetValue(Context, MoveTemp(InCollection), &Collection);
 	}
 }
 
diff --git a/Engine/Plugins/Experimental/ChaosFlesh/Source/ChaosFleshNodes/Private/Dataflow/ChaosFleshKinematicInitializationNodes.cpp b/Engine/Plugins/Experimental/ChaosFlesh/Source/ChaosFleshNodes/Private/Dataflow/ChaosFleshKinematicInitializationNodes.cpp
--- a/Engine/Plugins/Experimental/ChaosFlesh/Source/ChaosFleshNodes/Private/Dataflow/ChaosFleshKinematicInitializationNodes.cpp
+++ b/Engine/Plugins/Experimental/ChaosFlesh/Source/ChaosFleshNodes/Private/Dataflow/ChaosFleshKinematicInitializationNodes.cpp
@@ -143,7 +143,7 @@ void FKinematicTetrahedralBindingsDataflowNode::Evaluate(Dataflow::FContext& Con
 			GeometryCollection::Facades::FVertexBoneWeightsFacade(InCollection).AddBoneWeightsFromKinematicBindings();
 		}
 
-		SetValue<DataType>(Context, InCollection, &Collection);
+		SetValue(Context, MoveTemp(InCollection), &Collection);
 	}
 }
 
@@ -247,7 +247,7 @@ void FKinematicInitializationDataflowNode::Evaluate(Dataflow::FContext& Context,
 				}
 			}
 		}
-		SetValue<DataType>(Context, InCollection, &Collection);
+		SetValue(Context, MoveTemp(InCollection), &Collection);
 	}
 }
 
@@ -372,7 +372,7 @@ void FKinematicOriginInsertionInitializationDataflowNode::Evaluate(Dataflow::FCo
 				}
 			}
 		}
-		SetValue<DataType>(Context, InCollection, &Collection);
+		SetValue(Context, MoveTemp(InCollection), &Collection);
 	}
 }
 
@@ -425,7 +425,7 @@ void FSetVerticesKinematicDataflowNode::Evaluate(Dataflow::FContext& Context, co
 			GeometryCollection::Facades::FKinematicBindingFacade Kinematics(InCollection);
 			Kinematics.AddKinematicBinding(Kinematics.SetBoneBindings(INDEX_NONE, BoundVerts, BoundWeights));
 		}
-		SetValue<DataType>(Context, InCollection, &Collection);
+		SetValue(Context, MoveTemp(InCollection), &Collection);
 	}
 }
 
@@ -493,7 +493,7 @@ void FKinematicBodySetupInitializationDataflowNode::Evaluate(Dataflow::FContext&
 				GeometryCollection::Facades::FVertexBoneWeightsFacade(InCollection).AddBoneWeightsFromKinematicBindings();
 			}
 		}
-		SetValue<DataType>(Context, InCollection, &Collection);
+		SetValue(Context, MoveTemp(InCollection), &Collection);
 	}
 }
 
@@ -579,8 +579,8 @@ void FKinematicSkeletalMeshInitializationDataflowNode::Evaluate(Dataflow::FConte
 				GeometryCollection::Facades::FVertexBoneWeightsFacade(InCollection).AddBoneWeightsFromKinematicBindings();
 			}
 		}
-		SetValue<DataType>(Context, InCollection, &Collection);
-		SetValue<TArray<int32>>(Context, Indices, &IndicesOut);
+		SetValue(Context, MoveTemp(InCollection), &Collection);
+		SetValue(Context, MoveTemp(Indices), &IndicesOut);
 	}
 }
 
@@ -615,7 +615,7 @@ FBindVerticesToSkeleton::Evaluate(Dataflow::FContext& Context, const FDataflowOu
 			}
 		}
 		
-		SetValue<DataType>(Context, InCollection, &Collection);
+		SetValue(Context, MoveTemp(InCollection), &Collection);
 	}
 }
 
@@ -693,7 +693,7 @@ FAuthorSceneCollisionCandidates::Evaluate(Dataflow::FContext& Context, const FDa
 		UE_LOG(LogKinematicInit, Display,
 			TEXT("'%s' - Added %d scene collision candidates."),
 			*GetName().ToString(), Num);
-		SetValue<DataType>(Context, InCollection, &Collection);
+		SetValue(Context, MoveTemp(InCollection), &Collection);
 	}
 }
 
@@ -717,6 +717,6 @@ FAppendToCollectionTransformAttributeDataflowNode::Evaluate(Dataflow::FContext&
 			TManagedArray<FTransform>& Transforms = CollectionValue.ModifyAttribute<FTransform>(FName(AttributeName), FName(GroupName));
 			Transforms[Index] = Transform;
 		}
-		SetValue<DataType>(Context, CollectionValue, &Collection);
+		SetValue(Context, MoveTemp(CollectionValue), &Collection);
 	}
 }
\ No newline at end of file
diff --git a/Engine/Plugins/Experimental/ChaosFlesh/Source/ChaosFleshNodes/Private/Dataflow/ChaosFleshPositionTargetInitializationNodes.cpp b/Engine/Plugins/Experimental/ChaosFlesh/Source/ChaosFleshNodes/Private/Dataflow/ChaosFleshPositionTargetInitializationNodes.cpp
--- a/Engine/Plugins/Experimental/ChaosFlesh/Source/ChaosFleshNodes/Private/Dataflow/ChaosFleshPositionTargetInitializationNodes.cpp
+++ b/Engine/Plugins/Experimental/ChaosFlesh/Source/ChaosFleshNodes/Private/Dataflow/ChaosFleshPositionTargetInitializationNodes.cpp
@@ -145,8 +145,8 @@ void FAddKinematicParticlesDataflowNode::Evaluate(Dataflow::FContext& Context, c
 				}
 			}
 		}
-		SetValue<DataType>(Context, InCollection, &Collection);
-		SetValue<TArray<int32>>(Context, TargetIndices, &TargetIndicesOut);
+		SetValue(Context, MoveTemp(InCollection), &Collection);
+		SetValue(Context, MoveTemp(TargetIndices), &TargetIndicesOut);
 	}
 }
 
@@ -263,7 +263,7 @@ void FSetVertexVertexPositionTargetBindingDataflowNode::Evaluate(Dataflow::FCont
 			}
 			
 		}
-		SetValue<DataType>(Context, InCollection, &Collection);
+		SetValue(Context, MoveTemp(InCollection), &Collection);
 	}
 }
 
@@ -379,7 +379,7 @@ void FSetVertexTetrahedraPositionTargetBindingDataflowNode::Evaluate(Dataflow::F
 				}
 			}
 		}
-		SetValue<DataType>(Context, InCollection, &Collection);
+		SetValue(Context, MoveTemp(InCollection), &Collection);
 	}
 
 
@@ -557,7 +557,7 @@ void FSetVertexTrianglePositionTargetBindingDataflowNode::Evaluate(Dataflow::FCo
 					}
 				}
 		}
-		SetValue<DataType>(Context, InCollection, &Collection);
+		SetValue(Context, MoveTemp(InCollection), &Collection);
 	}
 }
 
@@ -888,6 +888,6 @@ void FSetFleshBonePositionTargetBindingDataflowNode::Evaluate(Dataflow::FContext
 				}
 			}
 		}
-		SetValue<DataType>(Context, InCollection, &Collection);
+		SetValue(Context, MoveTemp(InCollection), &Collection);
 	}
 }
\ No newline at end of file
diff --git a/Engine/Plugins/Experimental/ChaosFlesh/Source/ChaosFleshNodes/Private/Dataflow/ChaosFleshRadialTetrahedronNodes.cpp b/Engine/Plugins/Experimental/ChaosFlesh/Source/ChaosFleshNodes/Private/Dataflow/ChaosFleshRadialTetrahedronNodes.cpp
--- a/Engine/Plugins/Experimental/ChaosFlesh/Source/ChaosFleshNodes/Private/Dataflow/ChaosFleshRadialTetrahedronNodes.cpp
+++ b/Engine/Plugins/Experimental/ChaosFlesh/Source/ChaosFleshNodes/Private/Dataflow/ChaosFleshRadialTetrahedronNodes.cpp
@@ -48,7 +48,7 @@ void FRadialTetrahedronDataflowNodes::Evaluate(Dataflow::FContext& Context, cons
 		TUniquePtr<FTetrahedralCollection> RadialCollection(FTetrahedralCollection::NewTetrahedralCollection(TetVertices, SurfaceElements, TetElements));
 		InCollection->AppendGeometry(*RadialCollection.Get());
 
-		SetValue<DataType>(Context, *(FManagedArrayCollection*)InCollection.Get(), &Collection);
+		SetValue(Context, MoveTemp((FManagedArrayCollection&)(*InCollection)), &Collection);
 	}
 }
 
diff --git a/Engine/Plugins/Experimental/ChaosFlesh/Source/ChaosFleshNodes/Private/Dataflow/ChaosFleshRenderInitializationNodes.cpp b/Engine/Plugins/Experimental/ChaosFlesh/Source/ChaosFleshNodes/Private/Dataflow/ChaosFleshRenderInitializationNodes.cpp
--- a/Engine/Plugins/Experimental/ChaosFlesh/Source/ChaosFleshNodes/Private/Dataflow/ChaosFleshRenderInitializationNodes.cpp
+++ b/Engine/Plugins/Experimental/ChaosFlesh/Source/ChaosFleshNodes/Private/Dataflow/ChaosFleshRenderInitializationNodes.cpp
@@ -25,7 +25,7 @@ void FBindForRenderToSkeletalMeshDataflowNode::Evaluate(Dataflow::FContext& Cont
 		// @todo(dataflow) : Implemention binding for optimus deformation pipeline. 
 		//
 
-		SetValue<DataType>(Context, InCollection, &Collection);
+		SetValue(Context, MoveTemp(InCollection), &Collection);
 	}
 }
 
diff --git a/Engine/Plugins/Experimental/ChaosFlesh/Source/ChaosFleshNodes/Private/Dataflow/ChaosFleshSkeletalBindingsNode.cpp b/Engine/Plugins/Experimental/ChaosFlesh/Source/ChaosFleshNodes/Private/Dataflow/ChaosFleshSkeletalBindingsNode.cpp
--- a/Engine/Plugins/Experimental/ChaosFlesh/Source/ChaosFleshNodes/Private/Dataflow/ChaosFleshSkeletalBindingsNode.cpp
+++ b/Engine/Plugins/Experimental/ChaosFlesh/Source/ChaosFleshNodes/Private/Dataflow/ChaosFleshSkeletalBindingsNode.cpp
@@ -179,6 +179,6 @@ FGenerateSkeletalBindings::Evaluate(Dataflow::FContext& Context, const FDataflow
 				}
 			} // end for TetMeshIdx
 		}
-		SetValue<DataType>(Context, OutCollection, &Collection);
+		SetValue(Context, MoveTemp(OutCollection), &Collection);
 	}
 }
diff --git a/Engine/Plugins/Experimental/ChaosFlesh/Source/ChaosFleshNodes/Private/Dataflow/ChaosFleshTetrahedralNodes.cpp b/Engine/Plugins/Experimental/ChaosFlesh/Source/ChaosFleshNodes/Private/Dataflow/ChaosFleshTetrahedralNodes.cpp
--- a/Engine/Plugins/Experimental/ChaosFlesh/Source/ChaosFleshNodes/Private/Dataflow/ChaosFleshTetrahedralNodes.cpp
+++ b/Engine/Plugins/Experimental/ChaosFlesh/Source/ChaosFleshNodes/Private/Dataflow/ChaosFleshTetrahedralNodes.cpp
@@ -166,10 +166,10 @@ FCalculateTetMetrics::Evaluate(Dataflow::FContext& Context, const FDataflowOutpu
 				TetMeshCount,
 				MinVol, AvgVol, MaxVol,
 				MinAR, AvgAR, MaxAR);
-        }
+		}
 
-		//SetValue<DataType>(Context, *(FManagedArrayCollection*)OutCollection.Get(), &Collection);
-		SetValue<DataType>(Context, *(FManagedArrayCollection*)InCollection.Get(), &Collection);
+		//SetValue(Context, MoveTemp((FManagedArrayCollection&)(*OutCollection)), &Collection);
+		SetValue(Context, MoveTemp((FManagedArrayCollection&)(*InCollection)), &Collection);
 	}
 }
 
@@ -200,7 +200,7 @@ void FConstructTetGridNode::Evaluate(Dataflow::FContext& Context, const FDataflo
 			FTetrahedralCollection::NewTetrahedralCollection(X, Tris, Tets));
 		InCollection->AppendGeometry(*TetCollection.Get());
 
-		SetValue<DataType>(Context, *(FManagedArrayCollection*)InCollection.Get(), &Collection);
+		SetValue(Context, MoveTemp((FManagedArrayCollection&)(*InCollection)), &Collection);
 	}
 }
 
@@ -327,7 +327,7 @@ void FGenerateTetrahedralCollectionDataflowNodes::Evaluate(Dataflow::FContext& C
 			ensureMsgf(false, TEXT("FGenerateTetrahedralCollectionDataflowNodes is an editor only node."));
 #endif
 		} // end if InStaticMesh || InSkeletalMesh
-		SetValue<DataType>(Context, *(FManagedArrayCollection*)InCollection.Get(), &Collection);
+		SetValue(Context, MoveTemp((FManagedArrayCollection&)(*InCollection)), &Collection);
 	}
 }
 
diff --git a/Engine/Plugins/Experimental/Dataflow/Source/DataflowNodes/Private/Dataflow/DataflowContextOverridesNodes.cpp b/Engine/Plugins/Experimental/Dataflow/Source/DataflowNodes/Private/Dataflow/DataflowContextOverridesNodes.cpp
--- a/Engine/Plugins/Experimental/Dataflow/Source/DataflowNodes/Private/Dataflow/DataflowContextOverridesNodes.cpp
+++ b/Engine/Plugins/Experimental/Dataflow/Source/DataflowNodes/Private/Dataflow/DataflowContextOverridesNodes.cpp
@@ -24,7 +24,7 @@ void FFloatOverrideDataflowNode::Evaluate(Dataflow::FContext& Context, const FDa
 			FString Result = Dataflow::Reflection::FindOverrideProperty< FString >(EngineContext->Owner, PropertyName, KeyName);
 			Value = FCString::Atof(*Result);
 		}
-		SetValue<float>(Context, Value, &ValueOut);
+		SetValue(Context, Value, &ValueOut);
 	}
 }
 
diff --git a/Engine/Plugins/Experimental/Dataflow/Source/DataflowNodes/Private/Dataflow/DataflowSelectionNodes.cpp b/Engine/Plugins/Experimental/Dataflow/Source/DataflowNodes/Private/Dataflow/DataflowSelectionNodes.cpp
--- a/Engine/Plugins/Experimental/Dataflow/Source/DataflowNodes/Private/Dataflow/DataflowSelectionNodes.cpp
+++ b/Engine/Plugins/Experimental/Dataflow/Source/DataflowNodes/Private/Dataflow/DataflowSelectionNodes.cpp
@@ -28,7 +28,7 @@ void FSelectionSetDataflowNode::Evaluate(Dataflow::FContext& Context, const FDat
 			}
 		}
 
-		Out->SetValue<DataType>(IntArray, Context);
+		Out->SetValue<DataType>(MoveTemp(IntArray), Context);
 	}
 }
 
diff --git a/Engine/Plugins/Experimental/Dataflow/Source/DataflowNodes/Private/Dataflow/DataflowSkeletalMeshNodes.cpp b/Engine/Plugins/Experimental/Dataflow/Source/DataflowNodes/Private/Dataflow/DataflowSkeletalMeshNodes.cpp
--- a/Engine/Plugins/Experimental/Dataflow/Source/DataflowNodes/Private/Dataflow/DataflowSkeletalMeshNodes.cpp
+++ b/Engine/Plugins/Experimental/Dataflow/Source/DataflowNodes/Private/Dataflow/DataflowSkeletalMeshNodes.cpp
@@ -28,7 +28,7 @@ void FGetSkeletalMeshDataflowNode::Evaluate(Dataflow::FContext& Context, const F
 	typedef TObjectPtr<const USkeletalMesh> DataType;
 	if (Out->IsA<DataType>(&SkeletalMesh))
 	{
-		SetValue<DataType>(Context, SkeletalMesh, &SkeletalMesh);
+		SetValue(Context, SkeletalMesh, &SkeletalMesh);
 
 		if (!SkeletalMesh)
 		{
@@ -50,7 +50,7 @@ void FGetSkeletonDataflowNode::Evaluate(Dataflow::FContext& Context, const FData
 	typedef TObjectPtr<const USkeleton> DataType;
 	if (Out->IsA<DataType>(&Skeleton))
 	{
-		SetValue<DataType>(Context, Skeleton, &Skeleton);
+		SetValue(Context, Skeleton, &Skeleton);
 
 		if (!Skeleton)
 		{
@@ -85,7 +85,7 @@ void FSkeletalMeshBoneDataflowNode::Evaluate(Dataflow::FContext& Context, const
 			}
 
 			int32 Index = InSkeletalMesh->GetRefSkeleton().FindBoneIndex(LocalBoneName);
-			SetValue<int>(Context, Index, &BoneIndexOut);
+			SetValue(Context, Index, &BoneIndexOut);
 		}
 
 	}
@@ -97,7 +97,7 @@ void FSkeletalMeshReferenceTransformDataflowNode::Evaluate(Dataflow::FContext& C
 	typedef TObjectPtr<const USkeletalMesh> InDataType;
 	if (Out->IsA<FTransform>(&TransformOut))
 	{
-		SetValue<FTransform>(Context, FTransform::Identity, &TransformOut);
+		SetValue(Context, FTransform::Identity, &TransformOut);
 		
 		int32 BoneIndex = GetValue<int32>(Context, &BoneIndexIn);
 		if (0 <= BoneIndex)
@@ -108,8 +108,7 @@ void FSkeletalMeshReferenceTransformDataflowNode::Evaluate(Dataflow::FContext& C
 				Dataflow::Animation::GlobalTransforms(SkeletalMesh->GetRefSkeleton(), ComponentPose);
 				if (BoneIndex < ComponentPose.Num())
 				{
-					SetValue<FTransform>(Context, ComponentPose[BoneIndex], &TransformOut);
-
+					SetValue(Context, ComponentPose[BoneIndex], &TransformOut);
 				}
 			}
 		}
diff --git a/Engine/Plugins/Experimental/Dataflow/Source/DataflowNodes/Private/Dataflow/DataflowStaticMeshNodes.cpp b/Engine/Plugins/Experimental/Dataflow/Source/DataflowNodes/Private/Dataflow/DataflowStaticMeshNodes.cpp
--- a/Engine/Plugins/Experimental/Dataflow/Source/DataflowNodes/Private/Dataflow/DataflowStaticMeshNodes.cpp
+++ b/Engine/Plugins/Experimental/Dataflow/Source/DataflowNodes/Private/Dataflow/DataflowStaticMeshNodes.cpp
@@ -18,18 +18,18 @@ void FGetStaticMeshDataflowNode::Evaluate(Dataflow::FContext& Context, const FDa
 	typedef TObjectPtr<const UStaticMesh> DataType;
 	if (Out->IsA<DataType>(&StaticMesh))
 	{
-		SetValue<DataType>(Context, StaticMesh, &StaticMesh); // prime to avoid ensure
+		SetValue(Context, StaticMesh, &StaticMesh); // prime to avoid ensure
 
 		if (StaticMesh)
 		{
-			SetValue<DataType>(Context, StaticMesh, &StaticMesh);
+			SetValue(Context, StaticMesh, &StaticMesh);
 		}
 		else if (const Dataflow::FEngineContext* EngineContext = Context.AsType<Dataflow::FEngineContext>())
 		{
 			if (const UStaticMesh* StaticMeshFromOwner = Dataflow::Reflection::FindObjectPtrProperty<UStaticMesh>(
 				EngineContext->Owner, PropertyName))
 			{
-				SetValue<DataType>(Context, DataType(StaticMeshFromOwner), &StaticMesh);
+				SetValue(Context, DataType(StaticMeshFromOwner), &StaticMesh);
 
 			}
 		}
diff --git a/Engine/Plugins/Experimental/GeometryCollectionPlugin/Source/GeometryCollectionNodes/Private/Dataflow/GeometryCollectionArrayNodes.cpp b/Engine/Plugins/Experimental/GeometryCollectionPlugin/Source/GeometryCollectionNodes/Private/Dataflow/GeometryCollectionArrayNodes.cpp
--- a/Engine/Plugins/Experimental/GeometryCollectionPlugin/Source/GeometryCollectionNodes/Private/Dataflow/GeometryCollectionArrayNodes.cpp
+++ b/Engine/Plugins/Experimental/GeometryCollectionPlugin/Source/GeometryCollectionNodes/Private/Dataflow/GeometryCollectionArrayNodes.cpp
@@ -102,7 +102,7 @@ void FFloatArrayToIntArrayDataflowNode::Evaluate(Dataflow::FContext& Context, co
 			RetVal.SetNum(RetValIndex);
 		}
 
-		SetValue<TArray<int32>>(Context, RetVal, &IntArray);
+		SetValue(Context, MoveTemp(RetVal), &IntArray);
 	}
 }
 
@@ -113,11 +113,11 @@ void FGetArrayElementDataflowNode::Evaluate(Dataflow::FContext& Context, const F
 		const TArray<FVector>& Array = GetValue<TArray<FVector>>(Context, &Points);
 		if (Array.Num() > 0 && Index >= 0 && Index < Array.Num())
 		{
-			SetValue<FVector>(Context, Array[Index], &Point);
+			SetValue(Context, Array[Index], &Point);
 			return;
 		}
 
-		SetValue<FVector>(Context, FVector(0.f), &Point);
+		SetValue(Context, FVector(0.f), &Point);
 	}
 }
 
@@ -127,26 +127,26 @@ void FGetNumArrayElementsDataflowNode::Evaluate(Dataflow::FContext& Context, con
 	{
 		if (IsConnected<TArray<float>>(&FloatArray))
 		{
-			SetValue<int32>(Context, GetValue<TArray<float>>(Context, &FloatArray).Num(), &NumElements);
+			SetValue(Context, GetValue<TArray<float>>(Context, &FloatArray).Num(), &NumElements);
 			return;
 		}
 		else if (IsConnected<TArray<int32>>(&IntArray))
 		{
-			SetValue<int32>(Context, GetValue<TArray<int32>>(Context, &IntArray).Num(), &NumElements);
+			SetValue(Context, GetValue<TArray<int32>>(Context, &IntArray).Num(), &NumElements);
 			return;
 		}
 		else if (IsConnected<TArray<FVector>>(&Points))
 		{
-			SetValue<int32>(Context, GetValue<TArray<FVector>>(Context, &Points).Num(), &NumElements);
+			SetValue(Context, GetValue<TArray<FVector>>(Context, &Points).Num(), &NumElements);
 			return;
 		}
 		else if (IsConnected<TArray<FVector3f>>(&Vector3fArray))
 		{
-			SetValue<int32>(Context, GetValue<TArray<FVector3f>>(Context, &Vector3fArray).Num(), &NumElements);
+			SetValue(Context, GetValue<TArray<FVector3f>>(Context, &Vector3fArray).Num(), &NumElements);
 			return;
 		}
 
-		SetValue<int32>(Context, 0, &NumElements);
+		SetValue(Context, 0, &NumElements);
 	}
 }
 
@@ -160,7 +160,7 @@ void FBoolArrayToFaceSelectionDataflowNode::Evaluate(Dataflow::FContext& Context
 		NewFaceSelection.Initialize(InBoolAttributeData.Num(), false);
 		NewFaceSelection.SetFromArray(InBoolAttributeData);
 
-		SetValue<FDataflowFaceSelection>(Context, NewFaceSelection, &FaceSelection);
+		SetValue(Context, MoveTemp(NewFaceSelection), &FaceSelection);
 	}
 }
 
@@ -213,7 +213,7 @@ void FFloatArrayToVertexSelectionDataflowNode::Evaluate(Dataflow::FContext& Cont
 			}
 		}
 
-		SetValue<FDataflowVertexSelection>(Context, NewVertexSelection, &VertexSelection);
+		SetValue(Context, MoveTemp(NewVertexSelection), &VertexSelection);
 	}
 }
 
@@ -266,7 +266,7 @@ void FFloatArrayNormalizeDataflowNode::Evaluate(Dataflow::FContext& Context, con
 					}
 				}
 
-				SetValue<TArray<float>>(Context, NewFloatArray, &OutFloatArray);
+				SetValue(Context, MoveTemp(NewFloatArray), &OutFloatArray);
 				return;
 			}
 		}
@@ -297,11 +297,11 @@ void FFloatArrayNormalizeDataflowNode::Evaluate(Dataflow::FContext& Context, con
 				NewFloatArray[Idx] = InMinRange + NewFloatArray[Idx] * (InMaxRange - InMinRange);
 			}
 
-			SetValue<TArray<float>>(Context, NewFloatArray, &OutFloatArray);
+			SetValue(Context, MoveTemp(NewFloatArray), &OutFloatArray);
 			return;
 		}
 
-		SetValue<TArray<float>>(Context, TArray<float>(), &OutFloatArray);
+		SetValue(Context, TArray<float>(), &OutFloatArray);
 	}
 }
 
@@ -335,7 +335,7 @@ void FVectorArrayNormalizeDataflowNode::Evaluate(Dataflow::FContext& Context, co
 					}
 				}
 
-				SetValue<TArray<FVector>>(Context, NewVectorArray, &OutVectorArray);
+				SetValue(Context, MoveTemp(NewVectorArray), &OutVectorArray);
 				return;
 			}
 		}
@@ -350,7 +350,7 @@ void FVectorArrayNormalizeDataflowNode::Evaluate(Dataflow::FContext& Context, co
 				NewVectorArray[Idx] = Vector;
 			}
 
-			SetValue<TArray<FVector>>(Context, NewVectorArray, &OutVectorArray);
+			SetValue(Context, MoveTemp(NewVectorArray), &OutVectorArray);
 			return;
 		}
 
@@ -373,7 +373,7 @@ void FUnionIntArraysDataflowNode::Evaluate(Dataflow::FContext& Context, const FD
 		{
 			OutputArray.AddUnique(Array2[i]);
 		}
-		SetValue<TArray<int32>>(Context, OutputArray, &OutArray);
+		SetValue(Context, MoveTemp(OutputArray), &OutArray);
 	}
 }
 
diff --git a/Engine/Plugins/Experimental/GeometryCollectionPlugin/Source/GeometryCollectionNodes/Private/Dataflow/GeometryCollectionClusteringNodes.cpp b/Engine/Plugins/Experimental/GeometryCollectionPlugin/Source/GeometryCollectionNodes/Private/Dataflow/GeometryCollectionClusteringNodes.cpp
--- a/Engine/Plugins/Experimental/GeometryCollectionPlugin/Source/GeometryCollectionNodes/Private/Dataflow/GeometryCollectionClusteringNodes.cpp
+++ b/Engine/Plugins/Experimental/GeometryCollectionPlugin/Source/GeometryCollectionNodes/Private/Dataflow/GeometryCollectionClusteringNodes.cpp
@@ -87,7 +87,7 @@ void FAutoClusterDataflowNode::Evaluate(Dataflow::FContext& Context, const FData
 				InEnforceSiteParameters,
 				InGridX, InGridY, InGridZ, InMinimumClusterSize, InKMeansIterations);
 
-			SetValue<FManagedArrayCollection>(Context, (const FManagedArrayCollection&)(*GeomCollection), &Collection);
+			SetValue(Context, MoveTemp((FManagedArrayCollection&)(*GeomCollection)), &Collection);
 		}
 	}
 }
@@ -140,14 +140,14 @@ void FClusterUnclusterDataflowNode::Evaluate(Dataflow::FContext& Context, const
 			SelectionFacade.RemoveRootNodes(Selection);
 			if (Selection.IsEmpty())
 			{
-				SetValue<FManagedArrayCollection>(Context, (const FManagedArrayCollection&)(*GeomCollection), &Collection);
+				SetValue(Context, MoveTemp((FManagedArrayCollection&)(*GeomCollection)), &Collection);
 				return;
 			}
 			FGeometryCollectionClusteringUtility::CollapseHierarchyOneLevel(GeomCollection.Get(), Selection);
 			FGeometryCollectionClusteringUtility::RemoveDanglingClusters(GeomCollection.Get());
 
 			HierarchyFacade.GenerateLevelAttribute();
-			SetValue<FManagedArrayCollection>(Context, (const FManagedArrayCollection&)(*GeomCollection), &Collection);
+			SetValue(Context, MoveTemp((FManagedArrayCollection&)(*GeomCollection)), &Collection);
 		}
 	}
 }
diff --git a/Engine/Plugins/Experimental/GeometryCollectionPlugin/Source/GeometryCollectionNodes/Private/Dataflow/GeometryCollectionConversionNodes.cpp b/Engine/Plugins/Experimental/GeometryCollectionPlugin/Source/GeometryCollectionNodes/Private/Dataflow/GeometryCollectionConversionNodes.cpp
--- a/Engine/Plugins/Experimental/GeometryCollectionPlugin/Source/GeometryCollectionNodes/Private/Dataflow/GeometryCollectionConversionNodes.cpp
+++ b/Engine/Plugins/Experimental/GeometryCollectionPlugin/Source/GeometryCollectionNodes/Private/Dataflow/GeometryCollectionConversionNodes.cpp
@@ -33,7 +33,7 @@ void FVectorToStringDataflowNode::Evaluate(Dataflow::FContext& Context, const FD
 	if (Out->IsA<FString>(&String))
 	{
 		FString Value = GetValue<FVector>(Context, &Vector).ToString();
-		SetValue<FString>(Context, Value, &String);
+		SetValue(Context, Value, &String);
 	}
 }
 
@@ -42,7 +42,7 @@ void FFloatToStringDataflowNode::Evaluate(Dataflow::FContext& Context, const FDa
 	if (Out->IsA<FString>(&String))
 	{
 		FString Value = FString::Printf(TEXT("%f"), GetValue<float>(Context, &Float));
-		SetValue<FString>(Context, Value, &String);
+		SetValue(Context, Value, &String);
 	}
 }
 
@@ -51,7 +51,7 @@ void FIntToStringDataflowNode::Evaluate(Dataflow::FContext& Context, const FData
 	if (Out->IsA<FString>(&String))
 	{
 		FString Value = FString::Printf(TEXT("%d"), GetValue<int32>(Context, &Int));
-		SetValue<FString>(Context, Value, &String);
+		SetValue(Context, Value, &String);
 	}
 }
 
@@ -60,7 +60,7 @@ void FBoolToStringDataflowNode::Evaluate(Dataflow::FContext& Context, const FDat
 	if (Out->IsA<FString>(&String))
 	{
 		FString Value = FString::Printf(TEXT("%s"), GetValue<bool>(Context, &Bool) ? TEXT("true") : TEXT("false"));
-		SetValue<FString>(Context, Value, &String);
+		SetValue(Context, Value, &String);
 	}
 }
 
@@ -120,7 +120,7 @@ void FIntToBoolDataflowNode::Evaluate(Dataflow::FContext& Context, const FDatafl
 	if (Out->IsA<bool>(&Bool))
 	{
 		bool Value = GetValue<int32>(Context, &Int, Int) == 0 ? false : true;
-		SetValue<bool>(Context, Value, &Bool);
+		SetValue(Context, Value, &Bool);
 	}
 }
 
@@ -129,7 +129,7 @@ void FBoolToIntDataflowNode::Evaluate(Dataflow::FContext& Context, const FDatafl
 	if (Out->IsA<int32>(&Int))
 	{
 		int32 Value = (int32)GetValue<bool>(Context, &Bool, Bool);
-		SetValue<int32>(Context, Value, &Int);
+		SetValue(Context, Value, &Int);
 	}
 }
 
diff --git a/Engine/Plugins/Experimental/GeometryCollectionPlugin/Source/GeometryCollectionNodes/Private/Dataflow/GeometryCollectionEditNodes.cpp b/Engine/Plugins/Experimental/GeometryCollectionPlugin/Source/GeometryCollectionNodes/Private/Dataflow/GeometryCollectionEditNodes.cpp
--- a/Engine/Plugins/Experimental/GeometryCollectionPlugin/Source/GeometryCollectionNodes/Private/Dataflow/GeometryCollectionEditNodes.cpp
+++ b/Engine/Plugins/Experimental/GeometryCollectionPlugin/Source/GeometryCollectionNodes/Private/Dataflow/GeometryCollectionEditNodes.cpp
@@ -66,7 +66,7 @@ void FPruneInCollectionDataflowNode::Evaluate(Dataflow::FContext& Context, const
 
 					FFractureEngineEdit::DeleteBranch(*GeomCollection, BoneIndices);
 
-					SetValue<FManagedArrayCollection>(Context, (const FManagedArrayCollection&)(*GeomCollection), &Collection);
+					SetValue(Context, MoveTemp((FManagedArrayCollection&)(*GeomCollection)), &Collection);
 					return;
 				}
 			}
@@ -78,7 +78,7 @@ void FPruneInCollectionDataflowNode::Evaluate(Dataflow::FContext& Context, const
 			UE_LOG(LogTemp, Error, TEXT("[Dataflow ERROR] %s"), *ErrorStr);
 		}
 
-		SetValue<FManagedArrayCollection>(Context, InCollection, &Collection);
+		SetValue(Context, InCollection, &Collection);
 	}
 }
 
@@ -104,7 +104,7 @@ void FSetVisibilityInCollectionDataflowNode::Evaluate(Dataflow::FContext& Contex
 
 					FFractureEngineEdit::SetVisibilityInCollectionFromTransformSelection(InCollection, BoneIndices, Visibility == EVisibiltyOptionsEnum::Dataflow_VisibilityOptions_Visible);
 
-					SetValue<FManagedArrayCollection>(Context, InCollection, &Collection);
+					SetValue(Context, MoveTemp(InCollection), &Collection);
 					return;
 				}
 			}
@@ -129,7 +129,7 @@ void FSetVisibilityInCollectionDataflowNode::Evaluate(Dataflow::FContext& Contex
 
 					FFractureEngineEdit::SetVisibilityInCollectionFromFaceSelection(InCollection, FaceIndexArr, Visibility == EVisibiltyOptionsEnum::Dataflow_VisibilityOptions_Visible);
 
-					SetValue<FManagedArrayCollection>(Context, InCollection, &Collection);
+					SetValue(Context, MoveTemp(InCollection), &Collection);
 					return;
 				}
 			}
@@ -148,7 +148,7 @@ void FSetVisibilityInCollectionDataflowNode::Evaluate(Dataflow::FContext& Contex
 		}
 
 		const FManagedArrayCollection& InCollection = GetValue<FManagedArrayCollection>(Context, &Collection);
-		SetValue<FManagedArrayCollection>(Context, InCollection, &Collection);
+		SetValue(Context, InCollection, &Collection);
 	}
 }
 
@@ -157,7 +157,7 @@ void FMergeInCollectionDataflowNode::Evaluate(Dataflow::FContext& Context, const
 {
 	if (Out->IsA<FManagedArrayCollection>(&Collection))
 	{
-		FManagedArrayCollection InCollection = GetValue<FManagedArrayCollection>(Context, &Collection);
+		const FManagedArrayCollection& InCollection = GetValue<FManagedArrayCollection>(Context, &Collection);
 		const FDataflowTransformSelection& InTransformSelection = GetValue<FDataflowTransformSelection>(Context, &TransformSelection);
 
 		const int32 NumTransforms = InCollection.NumElements(FGeometryCollection::TransformGroup);
@@ -172,7 +172,7 @@ void FMergeInCollectionDataflowNode::Evaluate(Dataflow::FContext& Context, const
 
 					FFractureEngineEdit::Merge(*GeomCollection, BoneIndices);
 
-					SetValue<FManagedArrayCollection>(Context, (const FManagedArrayCollection&)(*GeomCollection), &Collection);
+					SetValue(Context, MoveTemp((FManagedArrayCollection&)(*GeomCollection)), &Collection);
 					return;
 				}
 			}
@@ -184,6 +184,6 @@ void FMergeInCollectionDataflowNode::Evaluate(Dataflow::FContext& Context, const
 			UE_LOG(LogTemp, Error, TEXT("[Dataflow ERROR] %s"), *ErrorStr);
 		}
 
-		SetValue<FManagedArrayCollection>(Context, InCollection, &Collection);
+		SetValue(Context, InCollection, &Collection);
 	}
 }
diff --git a/Engine/Plugins/Experimental/GeometryCollectionPlugin/Source/GeometryCollectionNodes/Private/Dataflow/GeometryCollectionFieldNodes.cpp b/Engine/Plugins/Experimental/GeometryCollectionPlugin/Source/GeometryCollectionNodes/Private/Dataflow/GeometryCollectionFieldNodes.cpp
--- a/Engine/Plugins/Experimental/GeometryCollectionPlugin/Source/GeometryCollectionNodes/Private/Dataflow/GeometryCollectionFieldNodes.cpp
+++ b/Engine/Plugins/Experimental/GeometryCollectionPlugin/Source/GeometryCollectionNodes/Private/Dataflow/GeometryCollectionFieldNodes.cpp
@@ -254,9 +254,9 @@ void FRadialFalloffFieldDataflowNode::Evaluate(Dataflow::FContext& Context, cons
 				NewFieldFloatResult);
 
 			// Set the outputs
-			SetValue<TArray<float>>(Context, NewFieldFloatResult, &FieldFloatResult);
+			SetValue<TArray<float>>(Context, MoveTemp(NewFieldFloatResult), &FieldFloatResult);
 			SetValue<TArray<int32>>(Context, TArray<int32>(), &FieldRemap);
-			SetValue<FDataflowVertexSelection>(Context, NewFieldSelectionMask, &FieldSelectionMask);
+			SetValue<FDataflowVertexSelection>(Context, MoveTemp(NewFieldSelectionMask), &FieldSelectionMask);
 			SetValue<int32>(Context, InSamplePositions.Num(), &NumSamplePositions);
 
 			return;
@@ -298,9 +298,9 @@ void FRadialFalloffFieldDataflowNode::Evaluate(Dataflow::FContext& Context, cons
 					NewFieldFloatResult);
 
 				// Set the outputs
-				SetValue<TArray<float>>(Context, NewFieldFloatResult, &FieldFloatResult);
-				SetValue<TArray<int32>>(Context, NewRemapArray, &FieldRemap);
-				SetValue<FDataflowVertexSelection>(Context, NewFieldSelectionMask, &FieldSelectionMask);
+				SetValue<TArray<float>>(Context, MoveTemp(NewFieldFloatResult), &FieldFloatResult);
+				SetValue<TArray<int32>>(Context, MoveTemp(NewRemapArray), &FieldRemap);
+				SetValue<FDataflowVertexSelection>(Context, MoveTemp(NewFieldSelectionMask), &FieldSelectionMask);
 				SetValue<int32>(Context, InSamplePositions.Num(), &NumSamplePositions);
 
 				return;
@@ -482,9 +482,9 @@ void FBoxFalloffFieldDataflowNode::Evaluate(Dataflow::FContext& Context, const F
 				NewFieldFloatResult);
 
 			// Set the outputs
-			SetValue<TArray<float>>(Context, NewFieldFloatResult, &FieldFloatResult);
+			SetValue<TArray<float>>(Context, MoveTemp(NewFieldFloatResult), &FieldFloatResult);
 			SetValue<TArray<int32>>(Context, TArray<int32>(), &FieldRemap);
-			SetValue<FDataflowVertexSelection>(Context, NewFieldSelectionMask, &FieldSelectionMask);
+			SetValue<FDataflowVertexSelection>(Context, MoveTemp(NewFieldSelectionMask), &FieldSelectionMask);
 			SetValue<int32>(Context, InSamplePositions.Num(), &NumSamplePositions);
 
 			return;
@@ -526,9 +526,9 @@ void FBoxFalloffFieldDataflowNode::Evaluate(Dataflow::FContext& Context, const F
 					NewFieldFloatResult);
 
 				// Set the outputs
-				SetValue<TArray<float>>(Context, NewFieldFloatResult, &FieldFloatResult);
-				SetValue<TArray<int32>>(Context, NewRemapArray, &FieldRemap);
-				SetValue<FDataflowVertexSelection>(Context, NewFieldSelectionMask, &FieldSelectionMask);
+				SetValue<TArray<float>>(Context, MoveTemp(NewFieldFloatResult), &FieldFloatResult);
+				SetValue<TArray<int32>>(Context, MoveTemp(NewRemapArray), &FieldRemap);
+				SetValue<FDataflowVertexSelection>(Context, MoveTemp(NewFieldSelectionMask), &FieldSelectionMask);
 				SetValue<int32>(Context, InSamplePositions.Num(), &NumSamplePositions);
 
 				return;
@@ -702,9 +702,9 @@ void FPlaneFalloffFieldDataflowNode::Evaluate(Dataflow::FContext& Context, const
 				NewFieldFloatResult);
 
 			// Set the outputs
-			SetValue<TArray<float>>(Context, NewFieldFloatResult, &FieldFloatResult);
+			SetValue<TArray<float>>(Context, MoveTemp(NewFieldFloatResult), &FieldFloatResult);
 			SetValue<TArray<int32>>(Context, TArray<int32>(), &FieldRemap);
-			SetValue<FDataflowVertexSelection>(Context, NewFieldSelectionMask, &FieldSelectionMask);
+			SetValue<FDataflowVertexSelection>(Context, MoveTemp(NewFieldSelectionMask), &FieldSelectionMask);
 			SetValue<int32>(Context, InSamplePositions.Num(), &NumSamplePositions);
 
 			return;
@@ -748,9 +748,9 @@ void FPlaneFalloffFieldDataflowNode::Evaluate(Dataflow::FContext& Context, const
 					NewFieldFloatResult);
 
 				// Set the outputs
-				SetValue<TArray<float>>(Context, NewFieldFloatResult, &FieldFloatResult);
-				SetValue<TArray<int32>>(Context, NewRemapArray, &FieldRemap);
-				SetValue<FDataflowVertexSelection>(Context, NewFieldSelectionMask, &FieldSelectionMask);
+				SetValue<TArray<float>>(Context, MoveTemp(NewFieldFloatResult), &FieldFloatResult);
+				SetValue<TArray<int32>>(Context, MoveTemp(NewRemapArray), &FieldRemap);
+				SetValue<FDataflowVertexSelection>(Context, MoveTemp(NewFieldSelectionMask), &FieldSelectionMask);
 				SetValue<int32>(Context, InSamplePositions.Num(), &NumSamplePositions);
 
 				return;
@@ -864,7 +864,7 @@ void FRadialIntMaskFieldDataflowNode::Evaluate(Dataflow::FContext& Context, cons
 				NewFieldIntResult);
 
 			// Set the outputs
-			SetValue<TArray<int32>>(Context, NewFieldIntResult, &FieldIntResult);
+			SetValue<TArray<int32>>(Context, MoveTemp(NewFieldIntResult), &FieldIntResult);
 			SetValue<TArray<int32>>(Context, TArray<int32>(), &FieldRemap);
 			SetValue<int32>(Context, InSamplePositions.Num(), &NumSamplePositions);
 
@@ -896,8 +896,8 @@ void FRadialIntMaskFieldDataflowNode::Evaluate(Dataflow::FContext& Context, cons
 					NewFieldIntResult);
 
 				// Set the outputs
-				SetValue<TArray<int32>>(Context, NewFieldIntResult, &FieldIntResult);
-				SetValue<TArray<int32>>(Context, NewRemapArray, &FieldRemap);
+				SetValue<TArray<int32>>(Context, MoveTemp(NewFieldIntResult), &FieldIntResult);
+				SetValue<TArray<int32>>(Context, MoveTemp(NewRemapArray), &FieldRemap);
 				SetValue<int32>(Context, InSamplePositions.Num(), &NumSamplePositions);
 
 				return;
@@ -984,7 +984,7 @@ void FUniformScalarFieldDataflowNode::Evaluate(Dataflow::FContext& Context, cons
 			UniformScalarFieldProcess(InSamplePositions, InMagnitude, NewFieldFloatResult);
 
 			// Set the outputs
-			SetValue<TArray<float>>(Context, NewFieldFloatResult, &FieldFloatResult);
+			SetValue<TArray<float>>(Context, MoveTemp(NewFieldFloatResult), &FieldFloatResult);
 			SetValue<TArray<int32>>(Context, TArray<int32>(), &FieldRemap);
 			SetValue<int32>(Context, InSamplePositions.Num(), &NumSamplePositions);
 
@@ -1009,8 +1009,8 @@ void FUniformScalarFieldDataflowNode::Evaluate(Dataflow::FContext& Context, cons
 				UniformScalarFieldProcess(NewSamplePositions, InMagnitude, NewFieldFloatResult);
 
 				// Set the outputs
-				SetValue<TArray<float>>(Context, NewFieldFloatResult, &FieldFloatResult);
-				SetValue<TArray<int32>>(Context, NewRemapArray, &FieldRemap);
+				SetValue<TArray<float>>(Context, MoveTemp(NewFieldFloatResult), &FieldFloatResult);
+				SetValue<TArray<int32>>(Context, MoveTemp(NewRemapArray), &FieldRemap);
 				SetValue<int32>(Context, InSamplePositions.Num(), &NumSamplePositions);
 
 				return;
@@ -1100,7 +1100,7 @@ void FUniformVectorFieldDataflowNode::Evaluate(Dataflow::FContext& Context, cons
 			UniformVectorFieldProcess(InSamplePositions, InMagnitude, InDirection, NewFieldVectorResult);
 
 			// Set the outputs
-			SetValue<TArray<FVector>>(Context, NewFieldVectorResult, &FieldVectorResult);
+			SetValue<TArray<FVector>>(Context, MoveTemp(NewFieldVectorResult), &FieldVectorResult);
 			SetValue<TArray<int32>>(Context, TArray<int32>(), &FieldRemap);
 			SetValue<int32>(Context, InSamplePositions.Num(), &NumSamplePositions);
 
@@ -1125,8 +1125,8 @@ void FUniformVectorFieldDataflowNode::Evaluate(Dataflow::FContext& Context, cons
 				UniformVectorFieldProcess(NewSamplePositions, InMagnitude, InDirection, NewFieldVectorResult);
 
 				// Set the outputs
-				SetValue<TArray<FVector>>(Context, NewFieldVectorResult, &FieldVectorResult);
-				SetValue<TArray<int32>>(Context, NewRemapArray, &FieldRemap);
+				SetValue<TArray<FVector>>(Context, MoveTemp(NewFieldVectorResult), &FieldVectorResult);
+				SetValue<TArray<int32>>(Context, MoveTemp(NewRemapArray), &FieldRemap);
 				SetValue<int32>(Context, InSamplePositions.Num(), &NumSamplePositions);
 
 				return;
@@ -1216,7 +1216,7 @@ void FRadialVectorFieldDataflowNode::Evaluate(Dataflow::FContext& Context, const
 			RadialVectorFieldProcess(InSamplePositions, InMagnitude, InPosition, NewFieldVectorResult);
 
 			// Set the outputs
-			SetValue<TArray<FVector>>(Context, NewFieldVectorResult, &FieldVectorResult);
+			SetValue<TArray<FVector>>(Context, MoveTemp(NewFieldVectorResult), &FieldVectorResult);
 			SetValue<TArray<int32>>(Context, TArray<int32>(), &FieldRemap);
 			SetValue<int32>(Context, InSamplePositions.Num(), &NumSamplePositions);
 
@@ -1241,8 +1241,8 @@ void FRadialVectorFieldDataflowNode::Evaluate(Dataflow::FContext& Context, const
 				RadialVectorFieldProcess(NewSamplePositions, InMagnitude, InPosition, NewFieldVectorResult);
 
 				// Set the outputs
-				SetValue<TArray<FVector>>(Context, NewFieldVectorResult, &FieldVectorResult);
-				SetValue<TArray<int32>>(Context, NewRemapArray, &FieldRemap);
+				SetValue<TArray<FVector>>(Context, MoveTemp(NewFieldVectorResult), &FieldVectorResult);
+				SetValue<TArray<int32>>(Context, MoveTemp(NewRemapArray), &FieldRemap);
 				SetValue<int32>(Context, InSamplePositions.Num(), &NumSamplePositions);
 
 				return;
@@ -1329,7 +1329,7 @@ void FRandomVectorFieldDataflowNode::Evaluate(Dataflow::FContext& Context, const
 			RandomVectorFieldProcess(InSamplePositions, InMagnitude, NewFieldVectorResult);
 
 			// Set the outputs
-			SetValue<TArray<FVector>>(Context, NewFieldVectorResult, &FieldVectorResult);
+			SetValue<TArray<FVector>>(Context, MoveTemp(NewFieldVectorResult), &FieldVectorResult);
 			SetValue<TArray<int32>>(Context, TArray<int32>(), &FieldRemap);
 			SetValue<int32>(Context, InSamplePositions.Num(), &NumSamplePositions);
 
@@ -1354,8 +1354,8 @@ void FRandomVectorFieldDataflowNode::Evaluate(Dataflow::FContext& Context, const
 				RandomVectorFieldProcess(NewSamplePositions, InMagnitude, NewFieldVectorResult);
 
 				// Set the outputs
-				SetValue<TArray<FVector>>(Context, NewFieldVectorResult, &FieldVectorResult);
-				SetValue<TArray<int32>>(Context, NewRemapArray, &FieldRemap);
+				SetValue<TArray<FVector>>(Context, MoveTemp(NewFieldVectorResult), &FieldVectorResult);
+				SetValue<TArray<int32>>(Context, MoveTemp(NewRemapArray), &FieldRemap);
 				SetValue<int32>(Context, InSamplePositions.Num(), &NumSamplePositions);
 
 				return;
@@ -1448,7 +1448,7 @@ void FNoiseFieldDataflowNode::Evaluate(Dataflow::FContext& Context, const FDataf
 			NoiseFieldProcess(InSamplePositions, InMinRange, InMaxRange, InTransform, NewFieldFloatResult);
 
 			// Set the outputs
-			SetValue<TArray<float>>(Context, NewFieldFloatResult, &FieldFloatResult);
+			SetValue<TArray<float>>(Context, MoveTemp(NewFieldFloatResult), &FieldFloatResult);
 			SetValue<TArray<int32>>(Context, TArray<int32>(), &FieldRemap);
 			SetValue<int32>(Context, InSamplePositions.Num(), &NumSamplePositions);
 
@@ -1473,8 +1473,8 @@ void FNoiseFieldDataflowNode::Evaluate(Dataflow::FContext& Context, const FDataf
 				NoiseFieldProcess(NewSamplePositions, InMinRange, InMaxRange, InTransform, NewFieldFloatResult);
 
 				// Set the outputs
-				SetValue<TArray<float>>(Context, NewFieldFloatResult, &FieldFloatResult);
-				SetValue<TArray<int32>>(Context, NewRemapArray, &FieldRemap);
+				SetValue<TArray<float>>(Context, MoveTemp(NewFieldFloatResult), &FieldFloatResult);
+				SetValue<TArray<int32>>(Context, MoveTemp(NewRemapArray), &FieldRemap);
 				SetValue<int32>(Context, InSamplePositions.Num(), &NumSamplePositions);
 
 				return;
@@ -1561,7 +1561,7 @@ void FUniformIntegerFieldDataflowNode::Evaluate(Dataflow::FContext& Context, con
 			UniformIntegerFieldProcess(InSamplePositions, InMagnitude, NewFieldIntResult);
 
 			// Set the outputs
-			SetValue<TArray<int32>>(Context, NewFieldIntResult, &FieldIntResult);
+			SetValue<TArray<int32>>(Context, MoveTemp(NewFieldIntResult), &FieldIntResult);
 			SetValue<TArray<int32>>(Context, TArray<int32>(), &FieldRemap);
 			SetValue<int32>(Context, InSamplePositions.Num(), &NumSamplePositions);
 
@@ -1586,8 +1586,8 @@ void FUniformIntegerFieldDataflowNode::Evaluate(Dataflow::FContext& Context, con
 				UniformIntegerFieldProcess(NewSamplePositions, InMagnitude, NewFieldIntResult);
 
 				// Set the outputs
-				SetValue<TArray<int32>>(Context, NewFieldIntResult, &FieldIntResult);
-				SetValue<TArray<int32>>(Context, NewRemapArray, &FieldRemap);
+				SetValue<TArray<int32>>(Context, MoveTemp(NewFieldIntResult), &FieldIntResult);
+				SetValue<TArray<int32>>(Context, MoveTemp(NewRemapArray), &FieldRemap);
 				SetValue<int32>(Context, InSamplePositions.Num(), &NumSamplePositions);
 
 				return;
@@ -1704,7 +1704,7 @@ void FWaveScalarFieldDataflowNode::Evaluate(Dataflow::FContext& Context, const F
 				NewFieldFloatResult);
 
 			// Set the outputs
-			SetValue<TArray<float>>(Context, NewFieldFloatResult, &FieldFloatResult);
+			SetValue<TArray<float>>(Context, MoveTemp(NewFieldFloatResult), &FieldFloatResult);
 			SetValue<TArray<int32>>(Context, TArray<int32>(), &FieldRemap);
 			SetValue<int32>(Context, InSamplePositions.Num(), &NumSamplePositions);
 
@@ -1737,8 +1737,8 @@ void FWaveScalarFieldDataflowNode::Evaluate(Dataflow::FContext& Context, const F
 					NewFieldFloatResult);
 
 				// Set the outputs
-				SetValue<TArray<float>>(Context, NewFieldFloatResult, &FieldFloatResult);
-				SetValue<TArray<int32>>(Context, NewRemapArray, &FieldRemap);
+				SetValue<TArray<float>>(Context, MoveTemp(NewFieldFloatResult), &FieldFloatResult);
+				SetValue<TArray<int32>>(Context, MoveTemp(NewRemapArray), &FieldRemap);
 				SetValue<int32>(Context, InSamplePositions.Num(), &NumSamplePositions);
 
 				return;
@@ -1908,8 +1908,8 @@ void FSumScalarFieldDataflowNode::Evaluate(Dataflow::FContext& Context, const FD
 					Magnitude,
 					bSwapInputs);
 
-				SetValue<TArray<float>>(Context, NewFieldFloatOutput, &FieldFloatResult);
-				SetValue<TArray<int32>>(Context, NewRemapArray, &FieldRemap);
+				SetValue<TArray<float>>(Context, MoveTemp(NewFieldFloatOutput), &FieldFloatResult);
+				SetValue<TArray<int32>>(Context, MoveTemp(NewRemapArray), &FieldRemap);
 
 				return;
 			}
@@ -1933,7 +1933,7 @@ void FSumScalarFieldDataflowNode::Evaluate(Dataflow::FContext& Context, const FD
 					Magnitude,
 					bSwapInputs);
 
-				SetValue<TArray<float>>(Context, NewFieldFloatOutput, &FieldFloatResult);
+				SetValue<TArray<float>>(Context, MoveTemp(NewFieldFloatOutput), &FieldFloatResult);
 				SetValue<TArray<int32>>(Context, TArray<int32>(), &FieldRemap);
 
 				return;
@@ -2149,8 +2149,8 @@ void FSumVectorFieldDataflowNode::Evaluate(Dataflow::FContext& Context, const FD
 						Magnitude,
 						bSwapVectorInputs);
 
-					SetValue<TArray<FVector>>(Context, NewFieldVectorResult, &FieldVectorResult);
-					SetValue<TArray<int32>>(Context, NewRemapArray, &FieldRemap);
+					SetValue<TArray<FVector>>(Context, MoveTemp(NewFieldVectorResult), &FieldVectorResult);
+					SetValue<TArray<int32>>(Context, MoveTemp(NewRemapArray), &FieldRemap);
 
 					return;
 				}
@@ -2175,7 +2175,7 @@ void FSumVectorFieldDataflowNode::Evaluate(Dataflow::FContext& Context, const FD
 						Magnitude,
 						bSwapVectorInputs);
 
-					SetValue<TArray<FVector>>(Context, NewFieldVectorResult, &FieldVectorResult);
+					SetValue<TArray<FVector>>(Context, MoveTemp(NewFieldVectorResult), &FieldVectorResult);
 					SetValue<TArray<int32>>(Context, TArray<int32>(), &FieldRemap);
 
 					return;
@@ -2236,8 +2236,8 @@ void FSumVectorFieldDataflowNode::Evaluate(Dataflow::FContext& Context, const FD
 						Magnitude,
 						bSwapVectorInputs);
 
-					SetValue<TArray<FVector>>(Context, NewFieldVectorResult, &FieldVectorResult);
-					SetValue<TArray<int32>>(Context, NewRemapArray, &FieldRemap);
+					SetValue<TArray<FVector>>(Context, MoveTemp(NewFieldVectorResult), &FieldVectorResult);
+					SetValue<TArray<int32>>(Context, MoveTemp(NewRemapArray), &FieldRemap);
 
 					return;
 				}
@@ -2262,7 +2262,7 @@ void FSumVectorFieldDataflowNode::Evaluate(Dataflow::FContext& Context, const FD
 						Magnitude,
 						bSwapVectorInputs);
 
-					SetValue<TArray<FVector>>(Context, NewFieldVectorResult, &FieldVectorResult);
+					SetValue<TArray<FVector>>(Context, MoveTemp(NewFieldVectorResult), &FieldVectorResult);
 					SetValue<TArray<int32>>(Context, TArray<int32>(), &FieldRemap);
 
 					return;
@@ -2322,8 +2322,8 @@ void FSumVectorFieldDataflowNode::Evaluate(Dataflow::FContext& Context, const FD
 						Magnitude,
 						bSwapVectorInputs);
 
-					SetValue<TArray<FVector>>(Context, NewFieldVectorResult, &FieldVectorResult);
-					SetValue<TArray<int32>>(Context, NewRemapArray, &FieldRemap);
+					SetValue<TArray<FVector>>(Context, MoveTemp(NewFieldVectorResult), &FieldVectorResult);
+					SetValue<TArray<int32>>(Context, MoveTemp(NewRemapArray), &FieldRemap);
 
 					return;
 				}
@@ -2348,7 +2348,7 @@ void FSumVectorFieldDataflowNode::Evaluate(Dataflow::FContext& Context, const FD
 						Magnitude,
 						bSwapVectorInputs);
 
-					SetValue<TArray<FVector>>(Context, NewFieldVectorResult, &FieldVectorResult);
+					SetValue<TArray<FVector>>(Context, MoveTemp(NewFieldVectorResult), &FieldVectorResult);
 					SetValue<TArray<int32>>(Context, TArray<int32>(), &FieldRemap);
 
 					return;
@@ -2415,8 +2415,8 @@ void FSumVectorFieldDataflowNode::Evaluate(Dataflow::FContext& Context, const FD
 						Magnitude,
 						bSwapVectorInputs);
 
-					SetValue<TArray<FVector>>(Context, NewFieldVectorResult, &FieldVectorResult);
-					SetValue<TArray<int32>>(Context, NewRemapArray, &FieldRemap);
+					SetValue<TArray<FVector>>(Context, MoveTemp(NewFieldVectorResult), &FieldVectorResult);
+					SetValue<TArray<int32>>(Context, MoveTemp(NewRemapArray), &FieldRemap);
 
 					return;
 				}
@@ -2442,7 +2442,7 @@ void FSumVectorFieldDataflowNode::Evaluate(Dataflow::FContext& Context, const FD
 						Magnitude,
 						bSwapVectorInputs);
 
-					SetValue<TArray<FVector>>(Context, NewFieldVectorResult, &FieldVectorResult);
+					SetValue<TArray<FVector>>(Context, MoveTemp(NewFieldVectorResult), &FieldVectorResult);
 					SetValue<TArray<int32>>(Context, TArray<int32>(), &FieldRemap);
 
 					return;
@@ -2479,7 +2479,7 @@ void FFieldMakeDenseFloatArrayDataflowNode::Evaluate(Dataflow::FContext& Context
 				}
 			}
 
-			SetValue<TArray<float>>(Context, NewFieldFloatOutput, &FieldFloatResult);
+			SetValue<TArray<float>>(Context, MoveTemp(NewFieldFloatOutput), &FieldFloatResult);
 
 			return;
 		}
diff --git a/Engine/Plugins/Experimental/GeometryCollectionPlugin/Source/GeometryCollectionNodes/Private/Dataflow/GeometryCollectionFracturingNodes.cpp b/Engine/Plugins/Experimental/GeometryCollectionPlugin/Source/GeometryCollectionNodes/Private/Dataflow/GeometryCollectionFracturingNodes.cpp
--- a/Engine/Plugins/Experimental/GeometryCollectionPlugin/Source/GeometryCollectionNodes/Private/Dataflow/GeometryCollectionFracturingNodes.cpp
+++ b/Engine/Plugins/Experimental/GeometryCollectionPlugin/Source/GeometryCollectionNodes/Private/Dataflow/GeometryCollectionFracturingNodes.cpp
@@ -68,12 +68,12 @@ void FUniformScatterPointsDataflowNode::Evaluate(Dataflow::FContext& Context, co
 				PointsArr.Emplace(BBox.Min + FVector(RandStream.FRand(), RandStream.FRand(), RandStream.FRand()) * Extent);
 			}
 
-			SetValue<TArray<FVector>>(Context, PointsArr, &Points);
+			SetValue(Context, MoveTemp(PointsArr), &Points);
 		}
 		else
 		{
 			// ERROR: Invalid BoundingBox input
-			SetValue<TArray<FVector>>(Context, TArray<FVector>(), &Points);
+			SetValue(Context, TArray<FVector>(), &Points);
 		}
 	}
 }
@@ -104,7 +104,7 @@ void FRadialScatterPointsDataflowNode::Evaluate(Dataflow::FContext& Context, con
 			}
 		}
 
-		SetValue<TArray<FVector>>(Context, PointsArr, &Points);
+		SetValue(Context, MoveTemp(PointsArr), &Points);
 	}
 }
 
@@ -137,14 +137,14 @@ void FVoronoiFractureDataflowNode::Evaluate(Dataflow::FContext& Context, const F
 					AddSamplesForCollision,
 					GetValue<float>(Context, &CollisionSampleSpacing));
 
-				SetValue<FManagedArrayCollection>(Context, InCollection, &Collection);
+				SetValue(Context, MoveTemp(InCollection), &Collection);
 
 				return;
 			}
 		}
 
 		const FManagedArrayCollection& InCollection = GetValue<FManagedArrayCollection>(Context, &Collection);
-		SetValue<FManagedArrayCollection>(Context, InCollection, &Collection);
+		SetValue(Context, InCollection, &Collection);
 	}
 }
 
@@ -176,14 +176,14 @@ void FPlaneCutterDataflowNode::Evaluate(Dataflow::FContext& Context, const FData
 					GetValue<bool>(Context, &AddSamplesForCollision),
 					GetValue<float>(Context, &CollisionSampleSpacing));
 
-				SetValue<FManagedArrayCollection>(Context, InCollection, &Collection);
+				SetValue(Context, MoveTemp(InCollection), &Collection);
 
 				return;
 			}
 		}
 
 		const FManagedArrayCollection& InCollection = GetValue<FManagedArrayCollection>(Context, &Collection);
-		SetValue<FManagedArrayCollection>(Context, InCollection, &Collection);
+		SetValue(Context, InCollection, &Collection);
 	}
 }
 
diff --git a/Engine/Plugins/Experimental/GeometryCollectionPlugin/Source/GeometryCollectionNodes/Private/Dataflow/GeometryCollectionMakeNodes.cpp b/Engine/Plugins/Experimental/GeometryCollectionPlugin/Source/GeometryCollectionNodes/Private/Dataflow/GeometryCollectionMakeNodes.cpp
--- a/Engine/Plugins/Experimental/GeometryCollectionPlugin/Source/GeometryCollectionNodes/Private/Dataflow/GeometryCollectionMakeNodes.cpp
+++ b/Engine/Plugins/Experimental/GeometryCollectionPlugin/Source/GeometryCollectionNodes/Private/Dataflow/GeometryCollectionMakeNodes.cpp
@@ -66,7 +66,7 @@ void FMakeLiteralStringDataflowNode::Evaluate(Dataflow::FContext& Context, const
 {
 	if (Out->IsA<FString>(&String))
 	{
-		SetValue<FString>(Context, Value, &String);
+		SetValue(Context, Value, &String);
 	}
 }
 
@@ -74,7 +74,7 @@ void FMakePointsDataflowNode::Evaluate(Dataflow::FContext& Context, const FDataf
 {
 	if (Out->IsA<TArray<FVector>>(&Points))
 	{
-		SetValue<TArray<FVector>>(Context, Point, &Points);
+		SetValue(Context, Point, &Points);
 	}
 }
 
@@ -87,14 +87,14 @@ void FMakeBoxDataflowNode::Evaluate(Dataflow::FContext& Context, const FDataflow
 			FVector MinVal = GetValue<FVector>(Context, &Min);
 			FVector MaxVal = GetValue<FVector>(Context, &Max);
 
-			SetValue<FBox>(Context, FBox(MinVal, MaxVal), &Box);
+			SetValue(Context, FBox(MinVal, MaxVal), &Box);
 		}
 		else if (DataType == EMakeBoxDataTypeEnum::Dataflow_MakeBox_DataType_CenterSize)
 		{
 			FVector CenterVal = GetValue<FVector>(Context, &Center);
 			FVector SizeVal = GetValue<FVector>(Context, &Size);
 
-			SetValue<FBox>(Context, FBox(CenterVal - 0.5 * SizeVal, CenterVal + 0.5 * SizeVal), &Box);
+			SetValue(Context, FBox(CenterVal - 0.5 * SizeVal, CenterVal + 0.5 * SizeVal), &Box);
 		}
 	}
 }
@@ -107,7 +107,7 @@ void FMakeSphereDataflowNode::Evaluate(Dataflow::FContext& Context, const FDataf
 		FVector CenterVal = GetValue<FVector>(Context, &Center);
 		float RadiusVal = GetValue<float>(Context, &Radius);
 
-		SetValue<FSphere>(Context, FSphere(CenterVal, RadiusVal), &Sphere);
+		SetValue(Context, FSphere(CenterVal, RadiusVal), &Sphere);
 	}
 }
 
@@ -116,7 +116,7 @@ void FMakeLiteralFloatDataflowNode::Evaluate(Dataflow::FContext& Context, const
 {
 	if (Out->IsA<float>(&Float))
 	{
-		SetValue<float>(Context, Value, &Float);
+		SetValue(Context, Value, &Float);
 	}
 }
 
@@ -124,7 +124,7 @@ void FMakeLiteralIntDataflowNode::Evaluate(Dataflow::FContext& Context, const FD
 {
 	if (Out->IsA<int32>(&Int))
 	{
-		SetValue<int32>(Context, Value, &Int);
+		SetValue(Context, Value, &Int);
 	}
 }
 
@@ -132,7 +132,7 @@ void FMakeLiteralBoolDataflowNode::Evaluate(Dataflow::FContext& Context, const F
 {
 	if (Out->IsA<bool>(&Bool))
 	{
-		SetValue<bool>(Context, Value, &Bool);
+		SetValue(Context, Value, &Bool);
 	}
 }
 
@@ -140,8 +140,8 @@ void FMakeLiteralVectorDataflowNode::Evaluate(Dataflow::FContext& Context, const
 {
 	if (Out->IsA<FVector>(&Vector))
 	{
-		FVector Value(GetValue<float>(Context, &X, X), GetValue<float>(Context, &Y, Y), GetValue<float>(Context, &Z, Z));
-		SetValue<FVector>(Context, Value, &Vector);
+		const FVector Value(GetValue<float>(Context, &X, X), GetValue<float>(Context, &Y, Y), GetValue<float>(Context, &Z, Z));
+		SetValue(Context, Value, &Vector);
 	}
 }
 
@@ -149,7 +149,7 @@ void FMakeTransformDataflowNode::Evaluate(Dataflow::FContext& Context, const FDa
 {
 	if (Out->IsA<FTransform>(&OutTransform))
 	{
-		SetValue<FTransform>(Context,
+		SetValue(Context,
 			FTransform(FQuat::MakeFromEuler(GetValue<FVector>(Context, &InRotation))
 				, GetValue<FVector>(Context, &InTranslation)
 				, GetValue<FVector>(Context, &InScale))
@@ -161,8 +161,8 @@ void FMakeQuaternionDataflowNode::Evaluate(Dataflow::FContext& Context, const FD
 {
 	if (Out->IsA<FQuat>(&Quaternion))
 	{
-		FQuat Value(GetValue<float>(Context, &X, X), GetValue<float>(Context, &Y, Y), GetValue<float>(Context, &Z, Z), GetValue<float>(Context, &W, W));
-		SetValue<FQuat>(Context, Value, &Quaternion);
+		const FQuat Value(GetValue<float>(Context, &X, X), GetValue<float>(Context, &Y, Y), GetValue<float>(Context, &Z, Z), GetValue<float>(Context, &W, W));
+		SetValue(Context, Value, &Quaternion);
 	}
 }
 
diff --git a/Engine/Plugins/Experimental/GeometryCollectionPlugin/Source/GeometryCollectionNodes/Private/Dataflow/GeometryCollectionMaterialNodes.cpp b/Engine/Plugins/Experimental/GeometryCollectionPlugin/Source/GeometryCollectionNodes/Private/Dataflow/GeometryCollectionMaterialNodes.cpp
--- a/Engine/Plugins/Experimental/GeometryCollectionPlugin/Source/GeometryCollectionNodes/Private/Dataflow/GeometryCollectionMaterialNodes.cpp
+++ b/Engine/Plugins/Experimental/GeometryCollectionPlugin/Source/GeometryCollectionNodes/Private/Dataflow/GeometryCollectionMaterialNodes.cpp
@@ -110,8 +110,8 @@ void FAddMaterialToCollectionDataflowNode::Evaluate(Dataflow::FContext& Context,
 			}
 		}
 
-		SetValue<TArray<TObjectPtr<UMaterial>>>(Context, MoveTemp(InMaterials), &Materials);
-		SetValue<FManagedArrayCollection>(Context, MoveTemp(InCollection), &Collection);
+		SetValue(Context, MoveTemp(InMaterials), &Materials);
+		SetValue(Context, MoveTemp(InCollection), &Collection);
 	}
 }
 
@@ -170,8 +170,8 @@ void FReAssignMaterialInCollectionDataflowNode::Evaluate(Dataflow::FContext& Con
 		}
 
 		// move the collection to the output to avoid making another copy
-		SetValue<FManagedArrayCollection>(Context, MoveTemp(InCollection), &Collection);
-		SetValue<TArray<TObjectPtr<UMaterial>>>(Context, InMaterials, &Materials);
+		SetValue(Context, MoveTemp(InCollection), &Collection);
+		SetValue(Context, InMaterials, &Materials);
 	}
 }
 
@@ -197,7 +197,7 @@ void FMaterialsInfoDataflowNode::Evaluate(Dataflow::FContext& Context, const FDa
 
 		OutputStr.Appendf(TEXT("----------------------------------------\n"));
 
-		SetValue<FString>(Context, OutputStr, &String);
+		SetValue(Context, OutputStr, &String);
 	}
 }
 
@@ -213,13 +213,13 @@ void FGetMaterialFromMaterialsArrayDataflowNode::Evaluate(Dataflow::FContext& Co
 		{
 			if (InMaterialIdx >= 0 && InMaterialIdx < InMaterials.Num())
 			{
-				SetValue<TObjectPtr<UMaterial>>(Context, InMaterials[InMaterialIdx], &Material);
+				SetValue(Context, InMaterials[InMaterialIdx], &Material);
 
 				return;
 			}
 		}
 
-		SetValue<TObjectPtr<UMaterial>>(Context, TObjectPtr<UMaterial>(), &Material);
+		SetValue(Context, TObjectPtr<UMaterial>(), &Material);
 	}
 }
 
@@ -244,7 +244,7 @@ void FSetMaterialInMaterialsArrayDataflowNode::Evaluate(Dataflow::FContext& Cont
 			}
 		}
 
-		SetValue<TArray<TObjectPtr<UMaterial>>>(Context, MoveTemp(InMaterials), &Materials);
+		SetValue(Context, MoveTemp(InMaterials), &Materials);
 	}
 }
 
@@ -253,7 +253,7 @@ void FMakeMaterialDataflowNode::Evaluate(Dataflow::FContext& Context, const FDat
 {
 	if (Out->IsA<TObjectPtr<UMaterial>>(&Material))
 	{
-		SetValue<TObjectPtr<UMaterial>>(Context, InMaterial, &Material);
+		SetValue(Context, InMaterial, &Material);
 	}
 }
 
@@ -262,7 +262,7 @@ void FMakeMaterialsArrayDataflowNode::Evaluate(Dataflow::FContext& Context, cons
 {
 	if (Out->IsA<TArray<TObjectPtr<UMaterial>>>(&Materials))
 	{
-		SetValue<TArray<TObjectPtr<UMaterial>>>(Context, TArray<TObjectPtr<UMaterial>>(), &Materials);
+		SetValue(Context, TArray<TObjectPtr<UMaterial>>(), &Materials);
 	}
 }
 
diff --git a/Engine/Plugins/Experimental/GeometryCollectionPlugin/Source/GeometryCollectionNodes/Private/Dataflow/GeometryCollectionMathNodes.cpp b/Engine/Plugins/Experimental/GeometryCollectionPlugin/Source/GeometryCollectionNodes/Private/Dataflow/GeometryCollectionMathNodes.cpp
--- a/Engine/Plugins/Experimental/GeometryCollectionPlugin/Source/GeometryCollectionNodes/Private/Dataflow/GeometryCollectionMathNodes.cpp
+++ b/Engine/Plugins/Experimental/GeometryCollectionPlugin/Source/GeometryCollectionNodes/Private/Dataflow/GeometryCollectionMathNodes.cpp
@@ -80,8 +80,8 @@ void FAddDataflowNode::Evaluate(Dataflow::FContext& Context, const FDataflowOutp
 		const float InFloatA = GetValue<float>(Context, &FloatA, FloatA);
 		const float InFloatB = GetValue<float>(Context, &FloatB, FloatB);
 
-		float Result = InFloatA + InFloatB;
-		SetValue<float>(Context, Result, &ReturnValue);
+		const float Result = InFloatA + InFloatB;
+		SetValue(Context, Result, &ReturnValue);
 	}
 }
 
@@ -92,8 +92,8 @@ void FSubtractDataflowNode::Evaluate(Dataflow::FContext& Context, const FDataflo
 		const float InFloatA = GetValue<float>(Context, &FloatA, FloatA);
 		const float InFloatB = GetValue<float>(Context, &FloatB, FloatB);
 
-		float Result = InFloatA - InFloatB;
-		SetValue<float>(Context, Result, &ReturnValue);
+		const float Result = InFloatA - InFloatB;
+		SetValue(Context, Result, &ReturnValue);
 	}
 }
 
@@ -104,8 +104,8 @@ void FMultiplyDataflowNode::Evaluate(Dataflow::FContext& Context, const FDataflo
 		const float InFloatA = GetValue<float>(Context, &FloatA, FloatA);
 		const float InFloatB = GetValue<float>(Context, &FloatB, FloatB);
 
-		float Result = InFloatA * InFloatB;
-		SetValue<float>(Context, Result, &ReturnValue);
+		const float Result = InFloatA * InFloatB;
+		SetValue(Context, Result, &ReturnValue);
 	}
 }
 
@@ -121,7 +121,7 @@ void FSafeDivideDataflowNode::Evaluate(Dataflow::FContext& Context, const FDataf
 		{
 			Result = InFloatA / InFloatB;
 		}
-		SetValue<float>(Context, Result, &ReturnValue);
+		SetValue(Context, Result, &ReturnValue);
 	}
 }
 
@@ -141,8 +141,8 @@ void FDivisionDataflowNode::Evaluate(Dataflow::FContext& Context, const FDataflo
 			Result = (int32)(InDividend / InDivisor);
 			ResultRemainder = InDividend - (float)Result * InDivisor;
 		}
-		SetValue<float>(Context, ResultRemainder, &Remainder);
-		SetValue<int32>(Context, Result, &ReturnValue);
+		SetValue(Context, ResultRemainder, &Remainder);
+		SetValue(Context, Result, &ReturnValue);
 	}
 }
 
@@ -157,7 +157,7 @@ void FSafeReciprocalDataflowNode::Evaluate(Dataflow::FContext& Context, const FD
 		{
 			Result = 1.f / InFloat;
 		}
-		SetValue<float>(Context, Result, &ReturnValue);
+		SetValue(Context, Result, &ReturnValue);
 	}
 }
 
@@ -167,8 +167,8 @@ void FSquareDataflowNode::Evaluate(Dataflow::FContext& Context, const FDataflowO
 	{
 		const float InFloat = GetValue<float>(Context, &Float, Float);
 
-		float Result = InFloat * InFloat;
-		SetValue<float>(Context, Result, &ReturnValue);
+		const float Result = InFloat * InFloat;
+		SetValue(Context, Result, &ReturnValue);
 	}
 }
 
@@ -183,7 +183,7 @@ void FSquareRootDataflowNode::Evaluate(Dataflow::FContext& Context, const FDataf
 		{
 			Result = FMath::Sqrt(InFloat);
 		}
-		SetValue<float>(Context, Result, &ReturnValue);
+		SetValue(Context, Result, &ReturnValue);
 	}
 }
 
@@ -198,7 +198,7 @@ void FInverseSqrtDataflowNode::Evaluate(Dataflow::FContext& Context, const FData
 		{
 			Result = FMath::InvSqrt(InFloat);
 		}
-		SetValue<float>(Context, Result, &ReturnValue);
+		SetValue(Context, Result, &ReturnValue);
 	}
 }
 
@@ -208,8 +208,8 @@ void FCubeDataflowNode::Evaluate(Dataflow::FContext& Context, const FDataflowOut
 	{
 		const float InFloat = GetValue<float>(Context, &Float, Float);
 
-		float Result = InFloat * InFloat * InFloat;
-		SetValue<float>(Context, Result, &ReturnValue);
+		const float Result = InFloat * InFloat * InFloat;
+		SetValue(Context, Result, &ReturnValue);
 	}
 }
 
@@ -219,8 +219,8 @@ void FNegateDataflowNode::Evaluate(Dataflow::FContext& Context, const FDataflowO
 	{
 		const float InFloat = GetValue<float>(Context, &Float, Float);
 
-		float Result = -1.f * InFloat;
-		SetValue<float>(Context, Result, &ReturnValue);
+		const float Result = -1.f * InFloat;
+		SetValue(Context, Result, &ReturnValue);
 	}
 }
 
@@ -230,8 +230,8 @@ void FAbsDataflowNode::Evaluate(Dataflow::FContext& Context, const FDataflowOutp
 	{
 		const float InFloat = GetValue<float>(Context, &Float, Float);
 
-		float Result = FMath::Abs(InFloat);
-		SetValue<float>(Context, Result, &ReturnValue);
+		const float Result = FMath::Abs(InFloat);
+		SetValue(Context, Result, &ReturnValue);
 	}
 }
 
@@ -241,8 +241,8 @@ void FFloorDataflowNode::Evaluate(Dataflow::FContext& Context, const FDataflowOu
 	{
 		const float InFloat = GetValue<float>(Context, &Float, Float);
 
-		float Result = FMath::Floor(InFloat);
-		SetValue<float>(Context, Result, &ReturnValue);
+		const float Result = FMath::Floor(InFloat);
+		SetValue(Context, Result, &ReturnValue);
 	}
 }
 
@@ -253,7 +253,7 @@ void FCeilDataflowNode::Evaluate(Dataflow::FContext& Context, const FDataflowOut
 		const float InFloat = GetValue<float>(Context, &Float, Float);
 
 		float Result = FMath::CeilToFloat(InFloat);
-		SetValue<float>(Context, Result, &ReturnValue);
+		SetValue(Context, Result, &ReturnValue);
 	}
 }
 
@@ -264,7 +264,7 @@ void FRoundDataflowNode::Evaluate(Dataflow::FContext& Context, const FDataflowOu
 		const float InFloat = GetValue<float>(Context, &Float, Float);
 
 		float Result = FMath::RoundToFloat(InFloat);
-		SetValue<float>(Context, Result, &ReturnValue);
+		SetValue(Context, Result, &ReturnValue);
 	}
 }
 
@@ -274,8 +274,8 @@ void FTruncDataflowNode::Evaluate(Dataflow::FContext& Context, const FDataflowOu
 	{
 		const float InFloat = GetValue<float>(Context, &Float, Float);
 
-		float Result = FMath::TruncToFloat(InFloat);
-		SetValue<float>(Context, Result, &ReturnValue);
+		const float Result = FMath::TruncToFloat(InFloat);
+		SetValue(Context, Result, &ReturnValue);
 	}
 }
 
@@ -285,8 +285,8 @@ void FFracDataflowNode::Evaluate(Dataflow::FContext& Context, const FDataflowOut
 	{
 		const float InFloat = GetValue<float>(Context, &Float, Float);
 
-		float Result = FMath::Frac(InFloat);
-		SetValue<float>(Context, Result, &ReturnValue);
+		const float Result = FMath::Frac(InFloat);
+		SetValue(Context, Result, &ReturnValue);
 	}
 }
 
@@ -297,8 +297,8 @@ void FMinDataflowNode::Evaluate(Dataflow::FContext& Context, const FDataflowOutp
 		const float InFloatA = GetValue<float>(Context, &FloatA, FloatA);
 		const float InFloatB = GetValue<float>(Context, &FloatB, FloatB);
 
-		float Result = FMath::Min(InFloatA, InFloatB);
-		SetValue<float>(Context, Result, &ReturnValue);
+		const float Result = FMath::Min(InFloatA, InFloatB);
+		SetValue(Context, Result, &ReturnValue);
 	}
 }
 
@@ -309,8 +309,8 @@ void FMaxDataflowNode::Evaluate(Dataflow::FContext& Context, const FDataflowOutp
 		const float InFloatA = GetValue<float>(Context, &FloatA, FloatA);
 		const float InFloatB = GetValue<float>(Context, &FloatB, FloatB);
 
-		float Result = FMath::Max(InFloatA, InFloatB);
-		SetValue<float>(Context, Result, &ReturnValue);
+		const float Result = FMath::Max(InFloatA, InFloatB);
+		SetValue(Context, Result, &ReturnValue);
 	}
 }
 
@@ -322,8 +322,8 @@ void FMin3DataflowNode::Evaluate(Dataflow::FContext& Context, const FDataflowOut
 		const float InFloatB = GetValue<float>(Context, &FloatB, FloatB);
 		const float InFloatC = GetValue<float>(Context, &FloatC, FloatB);
 
-		float Result = FMath::Min3(InFloatA, InFloatB, InFloatC);
-		SetValue<float>(Context, Result, &ReturnValue);
+		const float Result = FMath::Min3(InFloatA, InFloatB, InFloatC);
+		SetValue(Context, Result, &ReturnValue);
 	}
 }
 
@@ -335,8 +335,8 @@ void FMax3DataflowNode::Evaluate(Dataflow::FContext& Context, const FDataflowOut
 		const float InFloatB = GetValue<float>(Context, &FloatB, FloatB);
 		const float InFloatC = GetValue<float>(Context, &FloatC, FloatB);
 
-		float Result = FMath::Max3(InFloatA, InFloatB, InFloatC);
-		SetValue<float>(Context, Result, &ReturnValue);
+		const float Result = FMath::Max3(InFloatA, InFloatB, InFloatC);
+		SetValue(Context, Result, &ReturnValue);
 	}
 }
 
@@ -346,8 +346,8 @@ void FSignDataflowNode::Evaluate(Dataflow::FContext& Context, const FDataflowOut
 	{
 		const float InFloat = GetValue<float>(Context, &Float, Float);
 
-		float Result = FMath::Sign(InFloat);
-		SetValue<float>(Context, Result, &ReturnValue);
+		const float Result = FMath::Sign(InFloat);
+		SetValue(Context, Result, &ReturnValue);
 	}
 }
 
@@ -359,8 +359,8 @@ void FClampDataflowNode::Evaluate(Dataflow::FContext& Context, const FDataflowOu
 		const float InMin = GetValue<float>(Context, &Min, Min);
 		const float InMax = GetValue<float>(Context, &Max, Max);
 
-		float Result = FMath::Clamp(InFloat, InMin, InMax);
-		SetValue<float>(Context, Result, &ReturnValue);
+		const float Result = FMath::Clamp(InFloat, InMin, InMax);
+		SetValue(Context, Result, &ReturnValue);
 	}
 }
 
@@ -381,7 +381,7 @@ void FFitDataflowNode::Evaluate(Dataflow::FContext& Context, const FDataflowOutp
 			float Q = (InFloat - InOldMin) / (InOldMax - InOldMin);
 			Result = InNewMin + Q * (InNewMax - InNewMin);
 		}
-		SetValue<float>(Context, Result, &ReturnValue);
+		SetValue(Context, Result, &ReturnValue);
 	}
 }
 
@@ -401,7 +401,7 @@ void FEFitDataflowNode::Evaluate(Dataflow::FContext& Context, const FDataflowOut
 			float Q = (InFloat - InOldMin) / (InOldMax - InOldMin);
 			Result = InNewMin + Q * (InNewMax - InNewMin);
 		}
-		SetValue<float>(Context, Result, &ReturnValue);
+		SetValue(Context, Result, &ReturnValue);
 	}
 }
 
@@ -412,8 +412,8 @@ void FPowDataflowNode::Evaluate(Dataflow::FContext& Context, const FDataflowOutp
 		const float InBase = GetValue<float>(Context, &Base, Base);
 		const float InExp = GetValue<float>(Context, &Exp, Exp);
 
-		float Result = FMath::Pow(InBase, InExp);
-		SetValue<float>(Context, Result, &ReturnValue);
+		const float Result = FMath::Pow(InBase, InExp);
+		SetValue(Context, Result, &ReturnValue);
 	}
 }
 
@@ -429,7 +429,7 @@ void FLogDataflowNode::Evaluate(Dataflow::FContext& Context, const FDataflowOutp
 		{
 			Result = FMath::LogX(InBase, A);
 		}
-		SetValue<float>(Context, Result, &ReturnValue);
+		SetValue(Context, Result, &ReturnValue);
 	}
 }
 
@@ -439,8 +439,8 @@ void FLogeDataflowNode::Evaluate(Dataflow::FContext& Context, const FDataflowOut
 	{
 		const float InA = GetValue<float>(Context, &A, A);
 
-		float Result = FMath::Loge(A);
-		SetValue<float>(Context, Result, &ReturnValue);
+		const float Result = FMath::Loge(A);
+		SetValue(Context, Result, &ReturnValue);
 	}
 }
 
@@ -457,7 +457,7 @@ void FLerpDataflowNode::Evaluate(Dataflow::FContext& Context, const FDataflowOut
 		{
 			Result = FMath::Lerp(InA, InB, InAlpha);
 		}
-		SetValue<float>(Context, Result, &ReturnValue);
+		SetValue(Context, Result, &ReturnValue);
 	}
 }
 
@@ -474,7 +474,7 @@ void FWrapDataflowNode::Evaluate(Dataflow::FContext& Context, const FDataflowOut
 		{
 			Result = FMath::Wrap(InFloat, InMin, InMax);
 		}
-		SetValue<float>(Context, Result, &ReturnValue);
+		SetValue(Context, Result, &ReturnValue);
 	}
 }
 
@@ -484,8 +484,8 @@ void FExpDataflowNode::Evaluate(Dataflow::FContext& Context, const FDataflowOutp
 	{
 		const float InFloat = GetValue<float>(Context, &Float, Float);
 
-		float Result = FMath::Exp(InFloat);
-		SetValue<float>(Context, Result, &ReturnValue);
+		const float Result = FMath::Exp(InFloat);
+		SetValue(Context, Result, &ReturnValue);
 	}
 }
 
@@ -495,8 +495,8 @@ void FSinDataflowNode::Evaluate(Dataflow::FContext& Context, const FDataflowOutp
 	{
 		const float InFloat = GetValue<float>(Context, &Float, Float);
 
-		float Result = FMath::Sin(InFloat);
-		SetValue<float>(Context, Result, &ReturnValue);
+		const float Result = FMath::Sin(InFloat);
+		SetValue(Context, Result, &ReturnValue);
 	}
 }
 
@@ -511,7 +511,7 @@ void FArcSinDataflowNode::Evaluate(Dataflow::FContext& Context, const FDataflowO
 		{
 			Result = FMath::Asin(InFloat);
 		}
-		SetValue<float>(Context, Result, &ReturnValue);
+		SetValue(Context, Result, &ReturnValue);
 	}
 }
 
@@ -521,8 +521,8 @@ void FCosDataflowNode::Evaluate(Dataflow::FContext& Context, const FDataflowOutp
 	{
 		const float InFloat = GetValue<float>(Context, &Float, Float);
 
-		float Result = FMath::Cos(InFloat);
-		SetValue<float>(Context, Result, &ReturnValue);
+		const float Result = FMath::Cos(InFloat);
+		SetValue(Context, Result, &ReturnValue);
 	}
 }
 
@@ -537,7 +537,7 @@ void FArcCosDataflowNode::Evaluate(Dataflow::FContext& Context, const FDataflowO
 		{
 			Result = FMath::Acos(InFloat);
 		}
-		SetValue<float>(Context, Result, &ReturnValue);
+		SetValue(Context, Result, &ReturnValue);
 	}
 }
 
@@ -547,8 +547,8 @@ void FTanDataflowNode::Evaluate(Dataflow::FContext& Context, const FDataflowOutp
 	{
 		const float InFloat = GetValue<float>(Context, &Float, Float);
 
-		float Result = FMath::Tan(InFloat);
-		SetValue<float>(Context, Result, &ReturnValue);
+		const float Result = FMath::Tan(InFloat);
+		SetValue(Context, Result, &ReturnValue);
 	}
 }
 
@@ -558,8 +558,8 @@ void FArcTanDataflowNode::Evaluate(Dataflow::FContext& Context, const FDataflowO
 	{
 		const float InFloat = GetValue<float>(Context, &Float, Float);
 
-		float Result = FMath::Atan(InFloat);
-		SetValue<float>(Context, Result, &ReturnValue);
+		const float Result = FMath::Atan(InFloat);
+		SetValue(Context, Result, &ReturnValue);
 	}
 }
 
@@ -570,8 +570,8 @@ void FArcTan2DataflowNode::Evaluate(Dataflow::FContext& Context, const FDataflow
 		const float InY = GetValue<float>(Context, &Y, Y);
 		const float InX = GetValue<float>(Context, &X, X);
 
-		float Result = FMath::Atan2(InY, InX);
-		SetValue<float>(Context, Result, &ReturnValue);
+		const float Result = FMath::Atan2(InY, InX);
+		SetValue(Context, Result, &ReturnValue);
 	}
 }
 
@@ -602,7 +602,7 @@ void FNormalizeToRangeDataflowNode::Evaluate(Dataflow::FContext& Context, const
 		}
 		Result = (InFloat - IRangeMin) / (InRangeMax - IRangeMin);
 
-		SetValue<float>(Context, Result, &ReturnValue);
+		SetValue(Context, Result, &ReturnValue);
 	}
 }
 
@@ -626,8 +626,8 @@ void FDotProductDataflowNode::Evaluate(Dataflow::FContext& Context, const FDataf
 		const FVector InVectorA = GetValue<FVector>(Context, &VectorA, VectorA);
 		const FVector InVectorB = GetValue<FVector>(Context, &VectorB, VectorB);
 
-		float Result = FVector::DotProduct(InVectorA, InVectorB);
-		SetValue<float>(Context, Result, &ReturnValue);
+		const float Result = FVector::DotProduct(InVectorA, InVectorB);
+		SetValue(Context, Result, &ReturnValue);
 	}
 }
 
@@ -638,8 +638,8 @@ void FCrossProductDataflowNode::Evaluate(Dataflow::FContext& Context, const FDat
 		const FVector InVectorA = GetValue<FVector>(Context, &VectorA, VectorA);
 		const FVector InVectorB = GetValue<FVector>(Context, &VectorB, VectorB);
 
-		FVector Result = FVector::CrossProduct(InVectorA, InVectorB);
-		SetValue<FVector>(Context, Result, &ReturnValue);
+		const FVector Result = FVector::CrossProduct(InVectorA, InVectorB);
+		SetValue(Context, Result, &ReturnValue);
 	}
 }
 
@@ -650,8 +650,8 @@ void FNormalizeDataflowNode::Evaluate(Dataflow::FContext& Context, const FDatafl
 		const FVector InVectorA = GetValue<FVector>(Context, &VectorA, VectorA);
 		const float InTolerance = GetValue<float>(Context, &Tolerance, Tolerance);
 
-		FVector Result = InVectorA.GetSafeNormal(Tolerance);
-		SetValue<FVector>(Context, Result, &ReturnValue);
+		const FVector Result = InVectorA.GetSafeNormal(Tolerance);
+		SetValue(Context, Result, &ReturnValue);
 	}
 }
 
@@ -661,8 +661,8 @@ void FLengthDataflowNode::Evaluate(Dataflow::FContext& Context, const FDataflowO
 	{
 		const FVector InVector = GetValue<FVector>(Context, &Vector, Vector);
 
-		float Result = InVector.Length();
-		SetValue<float>(Context, Result, &ReturnValue);
+		const float Result = InVector.Length();
+		SetValue(Context, Result, &ReturnValue);
 	}
 }
 
@@ -673,8 +673,8 @@ void FDistanceDataflowNode::Evaluate(Dataflow::FContext& Context, const FDataflo
 		const FVector InPointA = GetValue<FVector>(Context, &PointA, PointA);
 		const FVector InPointB = GetValue<FVector>(Context, &PointB, PointB);
 
-		float Result = (InPointB - InPointA).Length();
-		SetValue<float>(Context, Result, &ReturnValue);
+		const float Result = (InPointB - InPointA).Length();
+		SetValue(Context, Result, &ReturnValue);
 	}
 }
 
@@ -684,8 +684,8 @@ void FIsNearlyZeroDataflowNode::Evaluate(Dataflow::FContext& Context, const FDat
 	{
 		const float InFloat = GetValue<float>(Context, &Float, Float);
 
-		bool Result = FMath::IsNearlyZero(InFloat);
-		SetValue<bool>(Context, Result, &ReturnValue);
+		const bool Result = FMath::IsNearlyZero(InFloat);
+		SetValue(Context, Result, &ReturnValue);
 	}
 }
 
@@ -695,14 +695,14 @@ void FRandomFloatDataflowNode::Evaluate(Dataflow::FContext& Context, const FData
 	{
 		if (bDeterministic)
 		{
-			float RandomSeedVal = GetValue<float>(Context, &RandomSeed);
+			const float RandomSeedVal = GetValue<float>(Context, &RandomSeed);
 
-			FRandomStream Stream(RandomSeedVal);
-			SetValue<float>(Context, Stream.FRand(), &ReturnValue);
+			const FRandomStream Stream(RandomSeedVal);
+			SetValue(Context, Stream.FRand(), &ReturnValue);
 		}
 		else
 		{
-			SetValue<float>(Context, FMath::FRand(), &ReturnValue);
+			SetValue(Context, FMath::FRand(), &ReturnValue);
 		}
 	}
 }
@@ -716,14 +716,14 @@ void FRandomFloatInRangeDataflowNode::Evaluate(Dataflow::FContext& Context, cons
 
 		if (bDeterministic)
 		{
-			float RandomSeedVal = GetValue<float>(Context, &RandomSeed);
+			const float RandomSeedVal = GetValue<float>(Context, &RandomSeed);
 
-			FRandomStream Stream(RandomSeedVal);
-			SetValue<float>(Context, Stream.FRandRange(MinVal, MaxVal), &ReturnValue);
+			const FRandomStream Stream(RandomSeedVal);
+			SetValue(Context, (float)Stream.FRandRange(MinVal, MaxVal), &ReturnValue);
 		}
 		else
 		{
-			SetValue<float>(Context, FMath::FRandRange(MinVal, MaxVal), &ReturnValue);
+			SetValue(Context, FMath::FRandRange(MinVal, MaxVal), &ReturnValue);
 		}
 	}
 }
@@ -734,14 +734,14 @@ void FRandomUnitVectorDataflowNode::Evaluate(Dataflow::FContext& Context, const
 	{
 		if (bDeterministic)
 		{
-			float RandomSeedVal = GetValue<float>(Context, &RandomSeed);
+			const float RandomSeedVal = GetValue<float>(Context, &RandomSeed);
 
-			FRandomStream Stream(RandomSeedVal);
-			SetValue<FVector>(Context, Stream.VRand(), &ReturnValue);
+			const FRandomStream Stream(RandomSeedVal);
+			SetValue(Context, Stream.VRand(), &ReturnValue);
 		}
 		else
 		{
-			SetValue<FVector>(Context, FMath::VRand(), &ReturnValue);
+			SetValue(Context, FMath::VRand(), &ReturnValue);
 		}
 	}
 }
@@ -755,14 +755,14 @@ void FRandomUnitVectorInConeDataflowNode::Evaluate(Dataflow::FContext& Context,
 
 		if (bDeterministic)
 		{
-			float RandomSeedVal = GetValue<float>(Context, &RandomSeed);
+			const float RandomSeedVal = GetValue<float>(Context, &RandomSeed);
 
-			FRandomStream Stream(RandomSeedVal);
-			SetValue<FVector>(Context, Stream.VRandCone(ConeDirectionVal, ConeHalfAngleVal), &ReturnValue);
+			const FRandomStream Stream(RandomSeedVal);
+			SetValue(Context, Stream.VRandCone(ConeDirectionVal, ConeHalfAngleVal), &ReturnValue);
 		}
 		else
 		{
-			SetValue<FVector>(Context, FMath::VRandCone(ConeDirectionVal, ConeHalfAngleVal), &ReturnValue);
+			SetValue(Context, FMath::VRandCone(ConeDirectionVal, ConeHalfAngleVal), &ReturnValue);
 		}
 	}
 }
@@ -771,7 +771,7 @@ void FRadiansToDegreesDataflowNode::Evaluate(Dataflow::FContext& Context, const
 {
 	if (Out->IsA<float>(&Degrees))
 	{
-		SetValue<float>(Context, FMath::RadiansToDegrees(GetValue<float>(Context, &Radians)), &Degrees);
+		SetValue(Context, FMath::RadiansToDegrees(GetValue<float>(Context, &Radians)), &Degrees);
 	}
 }
 
@@ -779,7 +779,7 @@ void FDegreesToRadiansDataflowNode::Evaluate(Dataflow::FContext& Context, const
 {
 	if (Out->IsA<float>(&Radians))
 	{
-		SetValue<float>(Context, FMath::DegreesToRadians(GetValue<float>(Context, &Degrees)), &Radians);
+		SetValue(Context, FMath::DegreesToRadians(GetValue<float>(Context, &Degrees)), &Radians);
 	}
 }
 
@@ -789,59 +789,59 @@ void FMathConstantsDataflowNode::Evaluate(Dataflow::FContext& Context, const FDa
 	{
 		if (Constant == EMathConstantsEnum::Dataflow_MathConstants_Pi)
 		{
-			SetValue<float>(Context, FMathf::Pi, &ReturnValue);
+			SetValue(Context, FMathf::Pi, &ReturnValue);
 		}
 		else if (Constant == EMathConstantsEnum::Dataflow_MathConstants_HalfPi)
 		{
-			SetValue<float>(Context, FMathf::HalfPi, &ReturnValue);
+			SetValue(Context, FMathf::HalfPi, &ReturnValue);
 		}
 		else if (Constant == EMathConstantsEnum::Dataflow_MathConstants_TwoPi)
 		{
-			SetValue<float>(Context, FMathf::TwoPi, &ReturnValue);
+			SetValue(Context, FMathf::TwoPi, &ReturnValue);
 		}
 		else if (Constant == EMathConstantsEnum::Dataflow_MathConstants_FourPi)
 		{
-			SetValue<float>(Context, FMathf::FourPi, &ReturnValue);
+			SetValue(Context, FMathf::FourPi, &ReturnValue);
 		}
 		else if (Constant == EMathConstantsEnum::Dataflow_MathConstants_InvPi)
 		{
-			SetValue<float>(Context, FMathf::InvPi, &ReturnValue);
+			SetValue(Context, FMathf::InvPi, &ReturnValue);
 		}
 		else if (Constant == EMathConstantsEnum::Dataflow_MathConstants_InvTwoPi)
 		{
-			SetValue<float>(Context, FMathf::InvTwoPi, &ReturnValue);
+			SetValue(Context, FMathf::InvTwoPi, &ReturnValue);
 		}
 		else if (Constant == EMathConstantsEnum::Dataflow_MathConstants_Sqrt2)
 		{
-			SetValue<float>(Context, FMathf::Sqrt2, &ReturnValue);
+			SetValue(Context, FMathf::Sqrt2, &ReturnValue);
 		}
 		else if (Constant == EMathConstantsEnum::Dataflow_MathConstants_InvSqrt2)
 		{
-			SetValue<float>(Context, FMathf::InvSqrt2, &ReturnValue);
+			SetValue(Context, FMathf::InvSqrt2, &ReturnValue);
 		}
 		else if (Constant == EMathConstantsEnum::Dataflow_MathConstants_Sqrt3)
 		{
-			SetValue<float>(Context, FMathf::Sqrt3, &ReturnValue);
+			SetValue(Context, FMathf::Sqrt3, &ReturnValue);
 		}
 		else if (Constant == EMathConstantsEnum::Dataflow_MathConstants_InvSqrt3)
 		{
-			SetValue<float>(Context, FMathf::InvSqrt3, &ReturnValue);
+			SetValue(Context, FMathf::InvSqrt3, &ReturnValue);
 		}
 		else if (Constant == EMathConstantsEnum::Dataflow_FloatToInt_Function_E)
 		{
-			SetValue<float>(Context, 2.71828182845904523536f, &ReturnValue);
+			SetValue(Context, 2.71828182845904523536f, &ReturnValue);
 		}
 		else if (Constant == EMathConstantsEnum::Dataflow_FloatToInt_Function_Gamma)
 		{
-			SetValue<float>(Context, 0.577215664901532860606512090082f, &ReturnValue);
+			SetValue(Context, 0.577215664901532860606512090082f, &ReturnValue);
 		}
 		else if (Constant == EMathConstantsEnum::Dataflow_FloatToInt_Function_GoldenRatio)
 		{
-			SetValue<float>(Context, 1.618033988749894f, &ReturnValue);
+			SetValue(Context, 1.618033988749894f, &ReturnValue);
 		}
 		else if (Constant == EMathConstantsEnum::Dataflow_FloatToInt_Function_ZeroTolerance)
 		{
-			SetValue<float>(Context, FMathf::ZeroTolerance, &ReturnValue);
+			SetValue(Context, FMathf::ZeroTolerance, &ReturnValue);
 		}
 	}
 }
diff --git a/Engine/Plugins/Experimental/GeometryCollectionPlugin/Source/GeometryCollectionNodes/Private/Dataflow/GeometryCollectionMeshNodes.cpp b/Engine/Plugins/Experimental/GeometryCollectionPlugin/Source/GeometryCollectionNodes/Private/Dataflow/GeometryCollectionMeshNodes.cpp
--- a/Engine/Plugins/Experimental/GeometryCollectionPlugin/Source/GeometryCollectionNodes/Private/Dataflow/GeometryCollectionMeshNodes.cpp
+++ b/Engine/Plugins/Experimental/GeometryCollectionPlugin/Source/GeometryCollectionNodes/Private/Dataflow/GeometryCollectionMeshNodes.cpp
@@ -75,13 +75,13 @@ void FPointsToMeshDataflowNode::Evaluate(Dataflow::FContext& Context, const FDat
 				DynMesh.AppendVertex(Point);
 			}
 
-			SetValue<TObjectPtr<UDynamicMesh>>(Context, DynamicMesh, &Mesh);
-			SetValue<int32>(Context, DynamicMesh->GetTriangleCount(), &TriangleCount);
+			SetValue(Context, DynamicMesh, &Mesh);
+			SetValue(Context, DynamicMesh->GetTriangleCount(), &TriangleCount);
 		}
 		else
 		{
-			SetValue<TObjectPtr<UDynamicMesh>>(Context, NewObject<UDynamicMesh>(), &Mesh);
-			SetValue<int32>(Context, 0, &TriangleCount);
+			SetValue(Context, TObjectPtr<UDynamicMesh>(NewObject<UDynamicMesh>()), &Mesh);
+			SetValue(Context, 0, &TriangleCount);
 		}
 	}
 }
@@ -104,8 +104,8 @@ void FBoxToMeshDataflowNode::Evaluate(Dataflow::FContext& Context, const FDatafl
 		FFractureEngineUtility::ConvertBoxToVertexAndTriangleData(InBox, Vertices, Triangles);
 		FFractureEngineUtility::ConstructMesh(DynMesh, Vertices, Triangles);
 
-		SetValue<TObjectPtr<UDynamicMesh>>(Context, NewMesh, &Mesh);
-		SetValue<int32>(Context, NewMesh->GetTriangleCount(), &TriangleCount);
+		SetValue(Context, NewMesh, &Mesh);
+		SetValue(Context, NewMesh->GetTriangleCount(), &TriangleCount);
 	}
 }
 
@@ -118,11 +118,11 @@ void FMeshInfoDataflowNode::Evaluate(Dataflow::FContext& Context, const FDataflo
 		{
 			const UE::Geometry::FDynamicMesh3& DynMesh = InMesh->GetMeshRef();
 
-			SetValue<FString>(Context, DynMesh.MeshInfoString(), &InfoString);
+			SetValue(Context, DynMesh.MeshInfoString(), &InfoString);
 		}
 		else
 		{
-			SetValue<FString>(Context, FString(""), &InfoString);
+			SetValue(Context, FString(""), &InfoString);
 		}
 	}
 }
@@ -151,13 +151,13 @@ void FMeshToCollectionDataflowNode::Evaluate(Dataflow::FContext& Context, const
 				FManagedArrayCollection NewCollection = FManagedArrayCollection();
 				NewGeometryCollection.CopyTo(&NewCollection);
 
-				SetValue<FManagedArrayCollection>(Context, NewCollection, &Collection);
+				SetValue(Context, MoveTemp(NewCollection), &Collection);
 
 				return;
 			}
 		}
 
-		SetValue<FManagedArrayCollection>(Context, FManagedArrayCollection(), &Collection);
+		SetValue(Context, FManagedArrayCollection(), &Collection);
 	}
 }
 
@@ -201,13 +201,13 @@ void FCollectionToMeshDataflowNode::Evaluate(Dataflow::FContext& Context, const
 					ConverterToDynamicMesh.Convert(&MeshDescription, DynMesh);
 				}
 
-				SetValue<TObjectPtr<UDynamicMesh>>(Context, NewMesh, &Mesh);
+				SetValue(Context, NewMesh, &Mesh);
 
 				return;
 			}
 		}
 
-		SetValue<TObjectPtr<UDynamicMesh>>(Context, NewObject<UDynamicMesh>(), &Mesh);
+		SetValue(Context, TObjectPtr<UDynamicMesh>(NewObject<UDynamicMesh>()), &Mesh);
 	}
 #endif
 }
@@ -229,11 +229,11 @@ void FStaticMeshToMeshDataflowNode::Evaluate(Dataflow::FContext& Context, const
 				ConverterToDynamicMesh.Convert(MeshDescription, DynMesh);
 			}
 
-			SetValue<TObjectPtr<UDynamicMesh>>(Context, NewMesh, &Mesh);
+			SetValue(Context, NewMesh, &Mesh);
 		}
 		else
 		{
-			SetValue<TObjectPtr<UDynamicMesh>>(Context, NewObject<UDynamicMesh>(), &Mesh);
+			SetValue(Context, TObjectPtr<UDynamicMesh>(NewObject<UDynamicMesh>()), &Mesh);
 		}
 	}
 #endif
@@ -266,14 +266,14 @@ void FMeshAppendDataflowNode::Evaluate(Dataflow::FContext& Context, const FDataf
 					UE::Geometry::FMeshIndexMappings IndexMaps2;
 					MeshEditor.AppendMesh(&DynMesh2, IndexMaps2);
 
-					SetValue<TObjectPtr<UDynamicMesh>>(Context, NewMesh, &Mesh);
+					SetValue(Context, NewMesh, &Mesh);
 
 					return;
 				}
 			}
 		}
 
-		SetValue<TObjectPtr<UDynamicMesh>>(Context, NewObject<UDynamicMesh>(), &Mesh);
+		SetValue(Context, TObjectPtr<UDynamicMesh>(NewObject<UDynamicMesh>()), &Mesh);
 	}
 }
 
@@ -318,7 +318,7 @@ void FMeshBooleanDataflowNode::Evaluate(Dataflow::FContext& Context, const FData
 
 					if (Boolean.Compute())
 					{
-						SetValue<TObjectPtr<UDynamicMesh>>(Context, NewMesh, &Mesh);
+						SetValue(Context, NewMesh, &Mesh);
 
 						return;
 					}
@@ -326,7 +326,7 @@ void FMeshBooleanDataflowNode::Evaluate(Dataflow::FContext& Context, const FData
 			}
 		}
 
-		SetValue<TObjectPtr<UDynamicMesh>>(Context, NewObject<UDynamicMesh>(), &Mesh);
+		SetValue(Context, TObjectPtr<UDynamicMesh>(NewObject<UDynamicMesh>()), &Mesh);
 	}
 }
 
@@ -366,13 +366,13 @@ void FMeshCopyToPointsDataflowNode::Evaluate(Dataflow::FContext& Context, const
 					MeshEditor.AppendMesh(&DynMeshTemp, IndexMaps);
 				}
 
-				SetValue<TObjectPtr<UDynamicMesh>>(Context, NewMesh, &Mesh);
+				SetValue(Context, NewMesh, &Mesh);
 
 				return;
 			}
 		}
 
-		SetValue<TObjectPtr<UDynamicMesh>>(Context, NewObject<UDynamicMesh>(), &Mesh);
+		SetValue(Context, TObjectPtr<UDynamicMesh>(NewObject<UDynamicMesh>()), &Mesh);
 	}
 }
 
@@ -383,33 +383,33 @@ void FGetMeshDataDataflowNode::Evaluate(Dataflow::FContext& Context, const FData
 	{
 		if (const TObjectPtr<const UDynamicMesh> InMesh = GetValue<TObjectPtr<UDynamicMesh>>(Context, &Mesh))
 		{
-			SetValue<int32>(Context, InMesh->GetMeshRef().VertexCount(), &VertexCount);
+			SetValue(Context, InMesh->GetMeshRef().VertexCount(), &VertexCount);
 		}
 		else
 		{
-			SetValue<int32>(Context, 0, &VertexCount);
+			SetValue(Context, 0, &VertexCount);
 		}
 	}
 	else if (Out->IsA<int32>(&EdgeCount))
 	{
 		if (const TObjectPtr<const UDynamicMesh> InMesh = GetValue<TObjectPtr<UDynamicMesh>>(Context, &Mesh))
 		{
-			SetValue<int32>(Context, InMesh->GetMeshRef().EdgeCount(), &EdgeCount);
+			SetValue(Context, InMesh->GetMeshRef().EdgeCount(), &EdgeCount);
 		}
 		else
 		{
-			SetValue<int32>(Context, 0, &EdgeCount);
+			SetValue(Context, 0, &EdgeCount);
 		}
 	}
 	else if (Out->IsA<int32>(&TriangleCount))
 	{
 		if (const TObjectPtr<const UDynamicMesh> InMesh = GetValue<TObjectPtr<UDynamicMesh>>(Context, &Mesh))
 		{
-			SetValue<int32>(Context, InMesh->GetMeshRef().TriangleCount(), &TriangleCount);
+			SetValue(Context, InMesh->GetMeshRef().TriangleCount(), &TriangleCount);
 		}
 		else
 		{
-			SetValue<int32>(Context, 0, &TriangleCount);
+			SetValue(Context, 0, &TriangleCount);
 		}
 	}
 }
diff --git a/Engine/Plugins/Experimental/GeometryCollectionPlugin/Source/GeometryCollectionNodes/Private/Dataflow/GeometryCollectionNodes.cpp b/Engine/Plugins/Experimental/GeometryCollectionPlugin/Source/GeometryCollectionNodes/Private/Dataflow/GeometryCollectionNodes.cpp
--- a/Engine/Plugins/Experimental/GeometryCollectionPlugin/Source/GeometryCollectionNodes/Private/Dataflow/GeometryCollectionNodes.cpp
+++ b/Engine/Plugins/Experimental/GeometryCollectionPlugin/Source/GeometryCollectionNodes/Private/Dataflow/GeometryCollectionNodes.cpp
@@ -105,16 +105,16 @@ void FGetCollectionFromAssetDataflowNode::Evaluate(Dataflow::FContext& Context,
 		{
 			if (const TSharedPtr<FGeometryCollection, ESPMode::ThreadSafe> AssetCollection = CollectionAsset->GetGeometryCollection())
 			{
-				SetValue<FManagedArrayCollection>(Context, (const FManagedArrayCollection&)(*AssetCollection), &Collection);
+				SetValue(Context, (const FManagedArrayCollection&)(*AssetCollection), &Collection);
 			}
 			else
 			{
-				SetValue<FManagedArrayCollection>(Context, FManagedArrayCollection(), &Collection);
+				SetValue(Context, FManagedArrayCollection(), &Collection);
 			}
 		}
 		else
 		{
-			SetValue<FManagedArrayCollection>(Context, FManagedArrayCollection(), &Collection);
+			SetValue(Context, FManagedArrayCollection(), &Collection);
 		}
 	}
 }
@@ -132,13 +132,13 @@ void FAppendCollectionAssetsDataflowNode::Evaluate(Dataflow::FContext& Context,
 			GeometryGroupGuidsLocal1 = GuidArray1->GetConstArray();
 		}
 		InCollection1.Append(InCollection2);
-		SetValue<DataType>(Context, InCollection1, &Collection1);
+		SetValue(Context, MoveTemp(InCollection1), &Collection1);
 		if (const TManagedArray<FString>* GuidArray2 = InCollection2.FindAttribute<FString>("Guid", FGeometryCollection::GeometryGroup))
 		{
 			GeometryGroupGuidsLocal2 = GuidArray2->GetConstArray();
 		}
-		SetValue<TArray<FString>>(Context, GeometryGroupGuidsLocal1, &GeometryGroupGuidsOut1);
-		SetValue<TArray<FString>>(Context, GeometryGroupGuidsLocal2, &GeometryGroupGuidsOut2);
+		SetValue(Context, MoveTemp(GeometryGroupGuidsLocal1), &GeometryGroupGuidsOut1);
+		SetValue(Context, MoveTemp(GeometryGroupGuidsLocal2), &GeometryGroupGuidsOut2);
 	}
 }
 
@@ -177,7 +177,7 @@ void FBoundingBoxDataflowNode::Evaluate(Dataflow::FContext& Context, const FData
 		GeometryCollection::Facades::FBoundsFacade BoundsFacade(InCollection);
 		const FBox& BoundingBoxInCollectionSpace = BoundsFacade.GetBoundingBoxInCollectionSpace();
 
-		SetValue<FBox>(Context, BoundingBoxInCollectionSpace, &BoundingBox);
+		SetValue(Context, BoundingBoxInCollectionSpace, &BoundingBox);
 	}
 }
 
@@ -188,23 +188,23 @@ void FExpandBoundingBoxDataflowNode::Evaluate(Dataflow::FContext& Context, const
 
 	if (Out->IsA<FVector>(&Min))
 	{
-		SetValue<FVector>(Context, BBox.Min, &Min);
+		SetValue(Context, BBox.Min, &Min);
 	}
 	else if (Out->IsA<FVector>(&Max))
 	{
-		SetValue<FVector>(Context, BBox.Max, &Max);
+		SetValue(Context, BBox.Max, &Max);
 	}
 	else if (Out->IsA<FVector>(&Center))
 	{
-		SetValue<FVector>(Context, BBox.GetCenter(), &Center);
+		SetValue(Context, BBox.GetCenter(), &Center);
 	}
 	else if (Out->IsA<FVector>(&HalfExtents))
 	{
-		SetValue<FVector>(Context, BBox.GetExtent(), &HalfExtents);
+		SetValue(Context, BBox.GetExtent(), &HalfExtents);
 	}
 	else if (Out->IsA<float>(&Volume))
 	{
-		SetValue<float>(Context, BBox.GetVolume(), &Volume);
+		SetValue(Context, (float)BBox.GetVolume(), &Volume);
 	}
 }
 
@@ -215,15 +215,15 @@ void FExpandVectorDataflowNode::Evaluate(Dataflow::FContext& Context, const FDat
 
 	if (Out->IsA<float>(&X))
 	{
-		SetValue<float>(Context, VectorVal.X, &X);
+		SetValue(Context, (float)VectorVal.X, &X);
 	}
 	else if (Out->IsA<float>(&Y))
 	{
-		SetValue<float>(Context, VectorVal.Y, &Y);
+		SetValue(Context, (float)VectorVal.Y, &Y);
 	}
 	else if (Out->IsA<float>(&Z))
 	{
-		SetValue<float>(Context, VectorVal.Z, &Z);
+		SetValue(Context, (float)VectorVal.Z, &Z);
 	}
 }
 
@@ -231,8 +231,8 @@ void FStringAppendDataflowNode::Evaluate(Dataflow::FContext& Context, const FDat
 {
 	if (Out->IsA<FString>(&String))
 	{
-		const FString StringOut = GetValue<FString>(Context, &String1) + GetValue<FString>(Context, &String2);
-		SetValue<FString>(Context, StringOut, &String);
+		FString StringOut = GetValue<FString>(Context, &String1) + GetValue<FString>(Context, &String2);
+		SetValue(Context, MoveTemp(StringOut), &String);
 	}
 }
 
@@ -241,7 +241,7 @@ void FHashStringDataflowNode::Evaluate(Dataflow::FContext& Context, const FDataf
 {
 	if (Out->IsA<int32>(&Hash))
 	{
-		SetValue<int32>(Context, GetTypeHash(GetValue<FString>(Context, &String)), &Hash);
+		SetValue(Context, (int32)GetTypeHash(GetValue<FString>(Context, &String)), &Hash);
 	}
 }
 
@@ -249,7 +249,7 @@ void FHashVectorDataflowNode::Evaluate(Dataflow::FContext& Context, const FDataf
 {
 	if (Out->IsA<int32>(&Hash))
 	{
-		SetValue<int32>(Context, GetTypeHash(GetValue<FVector>(Context, &Vector)), &Hash);
+		SetValue(Context, (int32)GetTypeHash(GetValue<FVector>(Context, &Vector)), &Hash);
 	}
 }
 
@@ -289,7 +289,7 @@ void FGetBoundingBoxesFromCollectionDataflowNode::Evaluate(Dataflow::FContext& C
 
 		}
 
-		SetValue<TArray<FBox>>(Context, BoundingBoxesArr, &BoundingBoxes);
+		SetValue(Context, MoveTemp(BoundingBoxesArr), &BoundingBoxes);
 	}
 }
 
@@ -338,7 +338,7 @@ void FGetCentroidsFromCollectionDataflowNode::Evaluate(Dataflow::FContext& Conte
 			}
 		}
 
-		SetValue<TArray<FVector>>(Context, CentroidsArr, &Centroids);
+		SetValue(Context, MoveTemp(CentroidsArr), &Centroids);
 	}
 }
 
@@ -361,7 +361,7 @@ void FTransformCollectionDataflowNode::Evaluate(Dataflow::FContext& Context, con
 		GeometryCollection::Facades::FCollectionTransformFacade TransformFacade(InCollection);
 		TransformFacade.Transform(NewTransform);
 
-		SetValue<FManagedArrayCollection>(Context, InCollection, &Collection);
+		SetValue(Context, MoveTemp(InCollection), &Collection);
 	}
 }
 
@@ -385,7 +385,7 @@ void FBakeTransformsInCollectionDataflowNode::Evaluate(Dataflow::FContext& Conte
 			TransformFacade.SetBoneTransformToIdentity(TransformIdx);
 		}
 
-		SetValue<FManagedArrayCollection>(Context, InCollection, &Collection);
+		SetValue(Context, MoveTemp(InCollection), &Collection);
 	}
 }
 
@@ -413,11 +413,11 @@ void FTransformMeshDataflowNode::Evaluate(Dataflow::FContext& Context, const FDa
 
 			MeshTransforms::ApplyTransform(DynamicMesh, UE::Geometry::FTransformSRT3d(NewTransform), true);
 
-			SetValue<TObjectPtr<UDynamicMesh>>(Context, NewMesh, &Mesh);
+			SetValue(Context, NewMesh, &Mesh);
 		}
 		else
 		{
-			SetValue<TObjectPtr<UDynamicMesh>>(Context, NewObject<UDynamicMesh>(), &Mesh);
+			SetValue(Context, TObjectPtr<UDynamicMesh>(NewObject<UDynamicMesh>()), &Mesh);
 		}
 	}
 }
@@ -453,11 +453,11 @@ void FCompareIntDataflowNode::Evaluate(Dataflow::FContext& Context, const FDataf
 {
 	if (Out->IsA<bool>(&Result))
 	{
-		int32 IntAValue = GetValue<int32>(Context, &IntA);
-		int32 IntBValue = GetValue<int32>(Context, &IntB);
-		bool ResultValue = ApplyDataflowOperationComparison(IntAValue, IntBValue, Operation);
+		const int32 IntAValue = GetValue<int32>(Context, &IntA);
+		const int32 IntBValue = GetValue<int32>(Context, &IntB);
+		const bool ResultValue = ApplyDataflowOperationComparison(IntAValue, IntBValue, Operation);
 
-		SetValue<bool>(Context, ResultValue, &Result);
+		SetValue(Context, ResultValue, &Result);
 	}
 }
 
@@ -466,11 +466,11 @@ void FCompareFloatDataflowNode::Evaluate(Dataflow::FContext& Context, const FDat
 {
 	if (Out->IsA(&Result))
 	{
-		float AValue = GetValue(Context, &FloatA);
-		float BValue = GetValue(Context, &FloatB);
-		bool ResultValue = ApplyDataflowOperationComparison(AValue, BValue, Operation);
+		const float AValue = GetValue(Context, &FloatA);
+		const float BValue = GetValue(Context, &FloatB);
+		const bool ResultValue = ApplyDataflowOperationComparison(AValue, BValue, Operation);
 
-		SetValue<bool>(Context, ResultValue, &Result);
+		SetValue(Context, ResultValue, &Result);
 	}
 }
 
@@ -485,7 +485,7 @@ void FBranchDataflowNode::Evaluate(Dataflow::FContext& Context, const FDataflowO
 		{
 			if (TObjectPtr<UDynamicMesh> InMeshA = GetValue<TObjectPtr<UDynamicMesh>>(Context, &MeshA))
 			{
-				SetValue<TObjectPtr<UDynamicMesh>>(Context, InMeshA, &Mesh);
+				SetValue(Context, InMeshA, &Mesh);
 
 				return;
 			}
@@ -494,13 +494,13 @@ void FBranchDataflowNode::Evaluate(Dataflow::FContext& Context, const FDataflowO
 		{
 			if (TObjectPtr<UDynamicMesh> InMeshB = GetValue<TObjectPtr<UDynamicMesh>>(Context, &MeshB))
 			{
-				SetValue<TObjectPtr<UDynamicMesh>>(Context, InMeshB, &Mesh);
+				SetValue(Context, InMeshB, &Mesh);
 
 				return;
 			}
 		}
 
-		SetValue<TObjectPtr<UDynamicMesh>>(Context, NewObject<UDynamicMesh>(), &Mesh);
+		SetValue(Context, TObjectPtr<UDynamicMesh>(NewObject<UDynamicMesh>()), &Mesh);
 	}
 }
 
@@ -578,7 +578,7 @@ void FGetSchemaDataflowNode::Evaluate(Dataflow::FContext& Context, const FDatafl
 		}
 		OutputStr.Appendf(TEXT("----------------------------------------\n"));
 
-		SetValue<FString>(Context, OutputStr, &String);
+		SetValue(Context, MoveTemp(OutputStr), &String);
 	}
 }
 
@@ -618,7 +618,7 @@ void FRemoveOnBreakDataflowNode::Evaluate(Dataflow::FContext& Context, const FDa
 		}
 
 		// move the collection to the output to avoid making another copy
-		SetValue<FManagedArrayCollection>(Context, MoveTemp(InCollection), &Collection);
+		SetValue(Context, MoveTemp(InCollection), &Collection);
 	}
 }
 
@@ -649,7 +649,7 @@ void FSetAnchorStateDataflowNode::Evaluate(Dataflow::FContext& Context, const FD
 				AnchoringFacade.SetAnchored(BoneIndices, !bAnchored);
 			}
 
-			SetValue<FManagedArrayCollection>(Context, (const FManagedArrayCollection&)(*GeomCollection), &Collection);
+			SetValue(Context, MoveTemp((FManagedArrayCollection&)(*GeomCollection)), &Collection);
 		}
 	}
 }
@@ -679,7 +679,7 @@ void FProximityDataflowNode::Evaluate(Dataflow::FContext& Context, const FDatafl
 			ProximityUtility.InvalidateProximity();
 			ProximityUtility.UpdateProximity();
 
-			SetValue<FManagedArrayCollection>(Context, (const FManagedArrayCollection&)(*GeomCollection), &Collection);
+			SetValue(Context, MoveTemp((FManagedArrayCollection&)(*GeomCollection)), &Collection);
 		}
 	}
 }
@@ -695,7 +695,7 @@ void FCollectionSetPivotDataflowNode::Evaluate(Dataflow::FContext& Context, cons
 		GeometryCollection::Facades::FCollectionTransformFacade TransformFacade(InCollection);
 		TransformFacade.SetPivot(InTransform);
 
-		SetValue<FManagedArrayCollection>(Context, InCollection, &Collection);
+		SetValue(Context, MoveTemp(InCollection), &Collection);
 	}
 }
 
@@ -844,7 +844,7 @@ void FAddCustomCollectionAttributeDataflowNode::Evaluate(Dataflow::FContext& Con
 			}
 		}
 
-		SetValue<FManagedArrayCollection>(Context, MoveTemp(InCollection), &Collection);
+		SetValue(Context, MoveTemp(InCollection), &Collection);
 	}
 }
 
@@ -875,7 +875,7 @@ void FGetNumElementsInCollectionGroupDataflowNode::Evaluate(Dataflow::FContext&
 			}
 		}
 
-		SetValue<int32>(Context, OutNumElements, &NumElements);
+		SetValue(Context, OutNumElements, &NumElements);
 	}
 }
 
@@ -902,13 +902,13 @@ void FGetCollectionAttributeDataTypedDataflowNode::Evaluate(Dataflow::FContext&
 			GroupNameToUse = FName(*CustomGroupName);
 		}
 
-		SetValue<TArray<bool>>(Context, TArray<bool>(), &BoolAttributeData);
-		SetValue<TArray<float>>(Context, TArray<float>(), &FloatAttributeData);
-		SetValue<TArray<double>>(Context, TArray<double>(), &DoubleAttributeData);
-		SetValue<TArray<int32>>(Context, TArray<int32>(), &Int32AttributeData);
-		SetValue<TArray<FString>>(Context, TArray<FString>(), &StringAttributeData);
-		SetValue<TArray<FVector3f>>(Context, TArray<FVector3f>(), &Vector3fAttributeData);
-		SetValue<TArray<FVector3d>>(Context, TArray<FVector3d>(), &Vector3dAttributeData);
+		SetValue(Context, TArray<bool>(), &BoolAttributeData);
+		SetValue(Context, TArray<float>(), &FloatAttributeData);
+		SetValue(Context, TArray<double>(), &DoubleAttributeData);
+		SetValue(Context, TArray<int32>(), &Int32AttributeData);
+		SetValue(Context, TArray<FString>(), &StringAttributeData);
+		SetValue(Context, TArray<FVector3f>(), &Vector3fAttributeData);
+		SetValue(Context, TArray<FVector3d>(), &Vector3dAttributeData);
 
 		if (GroupNameToUse.GetStringLength() > 0 && AttrName.Len() > 0)
 		{
@@ -922,37 +922,37 @@ void FGetCollectionAttributeDataTypedDataflowNode::Evaluate(Dataflow::FContext&
 					{
 						const TManagedArray<bool>& AttributeArr = InCollection.GetAttribute<bool>(FName(*AttrName), GroupNameToUse);
 						TArray<bool> BoolArray = AttributeArr.GetAsBoolArray();
-						SetValue<TArray<bool>>(Context, MoveTemp(BoolArray), &BoolAttributeData);
+						SetValue(Context, MoveTemp(BoolArray), &BoolAttributeData);
 					}
 					else if (TypeStr == FString("Float"))
 					{
 						const TManagedArray<float>& AttributeArr = InCollection.GetAttribute<float>(FName(*AttrName), GroupNameToUse);
-						SetValue<TArray<float>>(Context, AttributeArr.GetConstArray(), &FloatAttributeData);
+						SetValue(Context, AttributeArr.GetConstArray(), &FloatAttributeData);
 					}
 					else if (TypeStr == FString("Double"))
 					{
 						const TManagedArray<double>& AttributeArr = InCollection.GetAttribute<double>(FName(*AttrName), GroupNameToUse);
-						SetValue<TArray<double>>(Context, AttributeArr.GetConstArray(), &DoubleAttributeData);
+						SetValue(Context, AttributeArr.GetConstArray(), &DoubleAttributeData);
 					}
 					else if (TypeStr == FString("Int32"))
 					{
 						const TManagedArray<int32>& AttributeArr = InCollection.GetAttribute<int32>(FName(*AttrName), GroupNameToUse);
-						SetValue<TArray<int32>>(Context, AttributeArr.GetConstArray(), &Int32AttributeData);
+						SetValue(Context, AttributeArr.GetConstArray(), &Int32AttributeData);
 					}
 					else if (TypeStr == FString("String"))
 					{
 						const TManagedArray<FString>& AttributeArr = InCollection.GetAttribute<FString>(FName(*AttrName), GroupNameToUse);
-						SetValue<TArray<FString>>(Context, AttributeArr.GetConstArray(), &StringAttributeData);
+						SetValue(Context, AttributeArr.GetConstArray(), &StringAttributeData);
 					}
 					else if (TypeStr == FString("Vector"))
 					{
 						const TManagedArray<FVector3f>& AttributeArr = InCollection.GetAttribute<FVector3f>(FName(*AttrName), GroupNameToUse);
-						SetValue<TArray<FVector3f>>(Context, AttributeArr.GetConstArray(), &Vector3fAttributeData);
+						SetValue(Context, AttributeArr.GetConstArray(), &Vector3fAttributeData);
 					}
 					else if (TypeStr == FString("Vector3d"))
 					{
 						const TManagedArray<FVector3d>& AttributeArr = InCollection.GetAttribute<FVector3d>(FName(*AttrName), GroupNameToUse);
-						SetValue<TArray<FVector3d>>(Context, AttributeArr.GetConstArray(), &Vector3dAttributeData);
+						SetValue(Context, AttributeArr.GetConstArray(), &Vector3dAttributeData);
 					}
 				}
 			}
@@ -1035,7 +1035,7 @@ void FSetCollectionAttributeDataTypedDataflowNode::Evaluate(Dataflow::FContext&
 			}
 		}
 
-		SetValue<FManagedArrayCollection>(Context, MoveTemp(InCollection), &Collection);
+		SetValue(Context, MoveTemp(InCollection), &Collection);
 	}
 }
 
@@ -1070,14 +1070,14 @@ void FSetVertexColorInCollectionFromVertexSelectionDataflowNode::Evaluate(Datafl
 			}
 		}
 
-		SetValue<FManagedArrayCollection>(Context, MoveTemp(InCollection), &Collection);
+		SetValue(Context, MoveTemp(InCollection), &Collection);
 	}
 }
 
 void FSelectionToVertexListDataflowNode::Evaluate(Dataflow::FContext& Context, const FDataflowOutput* Out) const
 {
 	const FDataflowVertexSelection& InVertexSelection = GetValue<FDataflowVertexSelection>(Context, &VertexSelection);
-	SetValue< TArray<int32> >(Context, InVertexSelection.AsArray(), &VertexList);
+	SetValue(Context, InVertexSelection.AsArray(), &VertexList);
 }
 
 void FSetVertexColorInCollectionFromFloatArrayDataflowNode::Evaluate(Dataflow::FContext& Context, const FDataflowOutput* Out) const
@@ -1100,7 +1100,7 @@ void FSetVertexColorInCollectionFromFloatArrayDataflowNode::Evaluate(Dataflow::F
 			}
 		}
 
-		SetValue<FManagedArrayCollection>(Context, MoveTemp(InCollection), &Collection);
+		SetValue(Context, MoveTemp(InCollection), &Collection);
 	}
 }
 
@@ -1110,7 +1110,7 @@ void FMultiplyTransformDataflowNode::Evaluate(Dataflow::FContext& Context, const
 {
 	if (Out->IsA<FTransform>(&OutTransform))
 	{
-		SetValue<FTransform>(Context,
+		SetValue(Context,
 			GetValue<FTransform>(Context, &InLeftTransform, FTransform::Identity)
 			*GetValue<FTransform>(Context, &InRightTransform, FTransform::Identity)
 			, &OutTransform);
@@ -1121,9 +1121,9 @@ void FInvertTransformDataflowNode::Evaluate(Dataflow::FContext& Context, const F
 {
 	if (Out->IsA<FTransform>(&OutTransform))
 	{
-		FTransform InXf = GetValue<FTransform>(Context, &InTransform, FTransform::Identity);
-		FTransform OutXf = InXf.Inverse();
-		SetValue<FTransform>(Context, OutXf, &OutTransform);
+		const FTransform InXf = GetValue<FTransform>(Context, &InTransform, FTransform::Identity);
+		const FTransform OutXf = InXf.Inverse();
+		SetValue(Context, OutXf, &OutTransform);
 	}
 }
 
diff --git a/Engine/Plugins/Experimental/GeometryCollectionPlugin/Source/GeometryCollectionNodes/Private/Dataflow/GeometryCollectionProcessingNodes.cpp b/Engine/Plugins/Experimental/GeometryCollectionPlugin/Source/GeometryCollectionNodes/Private/Dataflow/GeometryCollectionProcessingNodes.cpp
--- a/Engine/Plugins/Experimental/GeometryCollectionPlugin/Source/GeometryCollectionNodes/Private/Dataflow/GeometryCollectionProcessingNodes.cpp
+++ b/Engine/Plugins/Experimental/GeometryCollectionPlugin/Source/GeometryCollectionNodes/Private/Dataflow/GeometryCollectionProcessingNodes.cpp
@@ -24,7 +24,7 @@ void FCloseGeometryOnCollectionDataflowNode::Evaluate(Dataflow::FContext& Contex
 		// @todo(dataflow) : Implemention that closes the geometry on the collection
 		//
 
-		SetValue<DataType>(Context, InCollection, &Collection);
+		SetValue(Context, MoveTemp(InCollection), &Collection);
 	}
 }
 
diff --git a/Engine/Plugins/Experimental/GeometryCollectionPlugin/Source/GeometryCollectionNodes/Private/Dataflow/GeometryCollectionSelectionNodes.cpp b/Engine/Plugins/Experimental/GeometryCollectionPlugin/Source/GeometryCollectionNodes/Private/Dataflow/GeometryCollectionSelectionNodes.cpp
--- a/Engine/Plugins/Experimental/GeometryCollectionPlugin/Source/GeometryCollectionNodes/Private/Dataflow/GeometryCollectionSelectionNodes.cpp
+++ b/Engine/Plugins/Experimental/GeometryCollectionPlugin/Source/GeometryCollectionNodes/Private/Dataflow/GeometryCollectionSelectionNodes.cpp
@@ -87,12 +87,12 @@ void FCollectionTransformSelectionAllDataflowNode::Evaluate(Dataflow::FContext&
 		NewTransformSelection.Initialize(InCollection.NumElements(FGeometryCollection::TransformGroup), false);
 		NewTransformSelection.SetFromArray(SelectionArr);
 
-		SetValue<FDataflowTransformSelection>(Context, NewTransformSelection, &TransformSelection);
+		SetValue(Context, MoveTemp(NewTransformSelection), &TransformSelection);
 	}
 	else if (Out->IsA<FManagedArrayCollection>(&Collection))
 	{
 		const FManagedArrayCollection& InCollection = GetValue<FManagedArrayCollection>(Context, &Collection);
-		SetValue<FManagedArrayCollection>(Context, InCollection, &Collection);
+		SetValue(Context, InCollection, &Collection);
 	}
 }
 
@@ -128,7 +128,7 @@ void FCollectionTransformSelectionSetOperationDataflowNode::Evaluate(Dataflow::F
 			UE_LOG(LogTemp, Error, TEXT("[Dataflow ERROR] %s"), *ErrorStr);
 		}
 
-		SetValue<FDataflowTransformSelection>(Context, NewTransformSelection, &TransformSelection);
+		SetValue(Context, MoveTemp(NewTransformSelection), &TransformSelection);
 	}
 }
 
@@ -266,7 +266,7 @@ void FCollectionTransformSelectionInfoDataflowNode::Evaluate(Dataflow::FContext&
 
 		OutputStr.Appendf(TEXT("----------------------------------------\n"));
 
-		SetValue<FString>(Context, OutputStr, &String);
+		SetValue(Context, MoveTemp(OutputStr), &String);
 	}
 }
 
@@ -284,12 +284,12 @@ void FCollectionTransformSelectionNoneDataflowNode::Evaluate(Dataflow::FContext&
 		NewTransformSelection.Initialize(InCollection.NumElements(FGeometryCollection::TransformGroup), false);
 		NewTransformSelection.SetFromArray(SelectionArr);
 
-		SetValue<FDataflowTransformSelection>(Context, NewTransformSelection, &TransformSelection);
+		SetValue(Context, MoveTemp(NewTransformSelection), &TransformSelection);
 	}
 	else if (Out->IsA<FManagedArrayCollection>(&Collection))
 	{
 		const FManagedArrayCollection& InCollection = GetValue<FManagedArrayCollection>(Context, &Collection);
-		SetValue<FManagedArrayCollection>(Context, InCollection, &Collection);
+		SetValue(Context, InCollection, &Collection);
 	}
 }
 
@@ -302,7 +302,7 @@ void FCollectionTransformSelectionInvertDataflowNode::Evaluate(Dataflow::FContex
 
 		InTransformSelection.Invert();
 
-		SetValue<FDataflowTransformSelection>(Context, InTransformSelection, &TransformSelection);
+		SetValue(Context, MoveTemp(InTransformSelection), &TransformSelection);
 	}
 }
 
@@ -323,12 +323,12 @@ void FCollectionTransformSelectionRandomDataflowNode::Evaluate(Dataflow::FContex
 		NewTransformSelection.Initialize(InCollection.NumElements(FGeometryCollection::TransformGroup), false);
 		NewTransformSelection.SetFromArray(SelectionArr);
 
-		SetValue<FDataflowTransformSelection>(Context, NewTransformSelection, &TransformSelection);
+		SetValue(Context, MoveTemp(NewTransformSelection), &TransformSelection);
 	}
 	else if (Out->IsA<FManagedArrayCollection>(&Collection))
 	{
 		const FManagedArrayCollection& InCollection = GetValue<FManagedArrayCollection>(Context, &Collection);
-		SetValue<FManagedArrayCollection>(Context, InCollection, &Collection);
+		SetValue(Context, InCollection, &Collection);
 	}
 }
 
@@ -346,12 +346,12 @@ void FCollectionTransformSelectionRootDataflowNode::Evaluate(Dataflow::FContext&
 		NewTransformSelection.Initialize(InCollection.NumElements(FGeometryCollection::TransformGroup), false);
 		NewTransformSelection.SetFromArray(SelectionArr);
 
-		SetValue<FDataflowTransformSelection>(Context, NewTransformSelection, &TransformSelection);
+		SetValue(Context, MoveTemp(NewTransformSelection), &TransformSelection);
 	}
 	else if (Out->IsA<FManagedArrayCollection>(&Collection))
 	{
 		const FManagedArrayCollection& InCollection = GetValue<FManagedArrayCollection>(Context, &Collection);
-		SetValue<FManagedArrayCollection>(Context, InCollection, &Collection);
+		SetValue(Context, InCollection, &Collection);
 	}
 }
 
@@ -392,17 +392,17 @@ void FCollectionTransformSelectionCustomDataflowNode::Evaluate(Dataflow::FContex
 				}
 			}
 
-			SetValue<FDataflowTransformSelection>(Context, NewTransformSelection, &TransformSelection);
+			SetValue(Context, MoveTemp(NewTransformSelection), &TransformSelection);
 		}
 		else
 		{
-			SetValue<FDataflowTransformSelection>(Context, FDataflowTransformSelection(), &TransformSelection);
+			SetValue(Context, FDataflowTransformSelection(), &TransformSelection);
 		}
 	}
 	else if (Out->IsA<FManagedArrayCollection>(&Collection))
 	{
 		const FManagedArrayCollection& InCollection = GetValue<FManagedArrayCollection>(Context, &Collection);
-		SetValue<FManagedArrayCollection>(Context, InCollection, &Collection);
+		SetValue(Context, InCollection, &Collection);
 	}
 }
 
@@ -420,12 +420,12 @@ void FCollectionTransformSelectionParentDataflowNode::Evaluate(Dataflow::FContex
 
 		InTransformSelection.SetFromArray(SelectionArr);
 		
-		SetValue<FDataflowTransformSelection>(Context, InTransformSelection, &TransformSelection);
+		SetValue(Context, MoveTemp(InTransformSelection), &TransformSelection);
 	}
 	else if (Out->IsA<FManagedArrayCollection>(&Collection))
 	{
 		const FManagedArrayCollection& InCollection = GetValue<FManagedArrayCollection>(Context, &Collection);
-		SetValue<FManagedArrayCollection>(Context, InCollection, &Collection);
+		SetValue(Context, InCollection, &Collection);
 	}
 }
 
@@ -444,7 +444,7 @@ void FCollectionTransformSelectionByPercentageDataflowNode::Evaluate(Dataflow::F
 		GeometryCollection::Facades::FCollectionTransformSelectionFacade::SelectByPercentage(SelectionArr, InPercentage, bDeterministic, InRandomSeed);
 
 		InTransformSelection.SetFromArray(SelectionArr);
-		SetValue<FDataflowTransformSelection>(Context, InTransformSelection, &TransformSelection);
+		SetValue(Context, MoveTemp(InTransformSelection), &TransformSelection);
 	}
 }
 
@@ -462,12 +462,12 @@ void FCollectionTransformSelectionChildrenDataflowNode::Evaluate(Dataflow::FCont
 		TransformSelectionFacade.SelectChildren(SelectionArr);
 		InTransformSelection.SetFromArray(SelectionArr);
 
-		SetValue<FDataflowTransformSelection>(Context, InTransformSelection, &TransformSelection);
+		SetValue(Context, MoveTemp(InTransformSelection), &TransformSelection);
 	}
 	else if (Out->IsA<FManagedArrayCollection>(&Collection))
 	{
 		const FManagedArrayCollection& InCollection = GetValue<FManagedArrayCollection>(Context, &Collection);
-		SetValue<FManagedArrayCollection>(Context, InCollection, &Collection);
+		SetValue(Context, InCollection, &Collection);
 	}
 }
 
@@ -485,12 +485,12 @@ void FCollectionTransformSelectionSiblingsDataflowNode::Evaluate(Dataflow::FCont
 		TransformSelectionFacade.SelectSiblings(SelectionArr);
 		InTransformSelection.SetFromArray(SelectionArr);
 
-		SetValue<FDataflowTransformSelection>(Context, InTransformSelection, &TransformSelection);
+		SetValue(Context, MoveTemp(InTransformSelection), &TransformSelection);
 	}
 	else if (Out->IsA<FManagedArrayCollection>(&Collection))
 	{
 		const FManagedArrayCollection& InCollection = GetValue<FManagedArrayCollection>(Context, &Collection);
-		SetValue<FManagedArrayCollection>(Context, InCollection, &Collection);
+		SetValue(Context, InCollection, &Collection);
 	}
 }
 
@@ -508,12 +508,12 @@ void FCollectionTransformSelectionLevelDataflowNode::Evaluate(Dataflow::FContext
 		TransformSelectionFacade.SelectLevel(SelectionArr);
 		InTransformSelection.SetFromArray(SelectionArr);
 
-		SetValue<FDataflowTransformSelection>(Context, InTransformSelection, &TransformSelection);
+		SetValue(Context, MoveTemp(InTransformSelection), &TransformSelection);
 	}
 	else if (Out->IsA<FManagedArrayCollection>(&Collection))
 	{
 		const FManagedArrayCollection& InCollection = GetValue<FManagedArrayCollection>(Context, &Collection);
-		SetValue<FManagedArrayCollection>(Context, InCollection, &Collection);
+		SetValue(Context, InCollection, &Collection);
 	}
 }
 
@@ -556,12 +556,12 @@ void FCollectionTransformSelectionContactDataflowNode::Evaluate(Dataflow::FConte
 		TransformSelectionFacade.SelectLevel(SelectionArr);
 		InTransformSelection.SetFromArray(SelectionArr);
 
-		SetValue<FDataflowTransformSelection>(Context, InTransformSelection, &TransformSelection);
+		SetValue(Context, MoveTemp(InTransformSelection), &TransformSelection);
 	}
 	else if (Out->IsA<FManagedArrayCollection>(&Collection))
 	{
 		const FManagedArrayCollection& InCollection = GetValue<FManagedArrayCollection>(Context, &Collection);
-		SetValue<FManagedArrayCollection>(Context, InCollection, &Collection);
+		SetValue(Context, InCollection, &Collection);
 	}
 }
 
@@ -579,12 +579,12 @@ void FCollectionTransformSelectionLeafDataflowNode::Evaluate(Dataflow::FContext&
 		NewTransformSelection.Initialize(InCollection.NumElements(FGeometryCollection::TransformGroup), false);
 		NewTransformSelection.SetFromArray(SelectionArr);
 
-		SetValue<FDataflowTransformSelection>(Context, NewTransformSelection, &TransformSelection);
+		SetValue(Context, MoveTemp(NewTransformSelection), &TransformSelection);
 	}
 	else if (Out->IsA<FManagedArrayCollection>(&Collection))
 	{
 		const FManagedArrayCollection& InCollection = GetValue<FManagedArrayCollection>(Context, &Collection);
-		SetValue<FManagedArrayCollection>(Context, InCollection, &Collection);
+		SetValue(Context, InCollection, &Collection);
 	}
 }
 
@@ -602,12 +602,12 @@ void FCollectionTransformSelectionClusterDataflowNode::Evaluate(Dataflow::FConte
 		NewTransformSelection.Initialize(InCollection.NumElements(FGeometryCollection::TransformGroup), false);
 		NewTransformSelection.SetFromArray(SelectionArr);
 
-		SetValue<FDataflowTransformSelection>(Context, NewTransformSelection, &TransformSelection);
+		SetValue(Context, MoveTemp(NewTransformSelection), &TransformSelection);
 	}
 	else if (Out->IsA<FManagedArrayCollection>(&Collection))
 	{
 		const FManagedArrayCollection& InCollection = GetValue<FManagedArrayCollection>(Context, &Collection);
-		SetValue<FManagedArrayCollection>(Context, InCollection, &Collection);
+		SetValue(Context, InCollection, &Collection);
 	}
 }
 
@@ -628,12 +628,12 @@ void FCollectionTransformSelectionBySizeDataflowNode::Evaluate(Dataflow::FContex
 		NewTransformSelection.Initialize(InCollection.NumElements(FGeometryCollection::TransformGroup), false);
 		NewTransformSelection.SetFromArray(SelectionArr);
 
-		SetValue<FDataflowTransformSelection>(Context, NewTransformSelection, &TransformSelection);
+		SetValue(Context, MoveTemp(NewTransformSelection), &TransformSelection);
 	}
 	else if (Out->IsA<FManagedArrayCollection>(&Collection))
 	{
 		const FManagedArrayCollection& InCollection = GetValue<FManagedArrayCollection>(Context, &Collection);
-		SetValue<FManagedArrayCollection>(Context, InCollection, &Collection);
+		SetValue(Context, InCollection, &Collection);
 	}
 }
 
@@ -654,12 +654,12 @@ void FCollectionTransformSelectionByVolumeDataflowNode::Evaluate(Dataflow::FCont
 		NewTransformSelection.Initialize(InCollection.NumElements(FGeometryCollection::TransformGroup), false);
 		NewTransformSelection.SetFromArray(SelectionArr);
 
-		SetValue<FDataflowTransformSelection>(Context, NewTransformSelection, &TransformSelection);
+		SetValue(Context, MoveTemp(NewTransformSelection), &TransformSelection);
 	}
 	else if (Out->IsA<FManagedArrayCollection>(&Collection))
 	{
 		const FManagedArrayCollection& InCollection = GetValue<FManagedArrayCollection>(Context, &Collection);
-		SetValue<FManagedArrayCollection>(Context, InCollection, &Collection);
+		SetValue(Context, InCollection, &Collection);
 	}
 }
 
@@ -692,12 +692,12 @@ void FCollectionTransformSelectionInBoxDataflowNode::Evaluate(Dataflow::FContext
 		NewTransformSelection.Initialize(InCollection.NumElements(FGeometryCollection::TransformGroup), false);
 		NewTransformSelection.SetFromArray(SelectionArr);
 
-		SetValue<FDataflowTransformSelection>(Context, NewTransformSelection, &TransformSelection);
+		SetValue(Context, MoveTemp(NewTransformSelection), &TransformSelection);
 	}
 	else if (Out->IsA<FManagedArrayCollection>(&Collection))
 	{
 		const FManagedArrayCollection& InCollection = GetValue<FManagedArrayCollection>(Context, &Collection);
-		SetValue<FManagedArrayCollection>(Context, InCollection, &Collection);
+		SetValue(Context, InCollection, &Collection);
 	}
 }
 
@@ -730,12 +730,12 @@ void FCollectionTransformSelectionInSphereDataflowNode::Evaluate(Dataflow::FCont
 		NewTransformSelection.Initialize(InCollection.NumElements(FGeometryCollection::TransformGroup), false);
 		NewTransformSelection.SetFromArray(SelectionArr);
 
-		SetValue<FDataflowTransformSelection>(Context, NewTransformSelection, &TransformSelection);
+		SetValue(Context, MoveTemp(NewTransformSelection), &TransformSelection);
 	}
 	else if (Out->IsA<FManagedArrayCollection>(&Collection))
 	{
 		const FManagedArrayCollection& InCollection = GetValue<FManagedArrayCollection>(Context, &Collection);
-		SetValue<FManagedArrayCollection>(Context, InCollection, &Collection);
+		SetValue(Context, InCollection, &Collection);
 	}
 }
 
@@ -756,12 +756,12 @@ void FCollectionTransformSelectionByFloatAttrDataflowNode::Evaluate(Dataflow::FC
 		NewTransformSelection.Initialize(InCollection.NumElements(FGeometryCollection::TransformGroup), false);
 		NewTransformSelection.SetFromArray(SelectionArr);
 
-		SetValue<FDataflowTransformSelection>(Context, NewTransformSelection, &TransformSelection);
+		SetValue(Context, MoveTemp(NewTransformSelection), &TransformSelection);
 	}
 	else if (Out->IsA<FManagedArrayCollection>(&Collection))
 	{
 		const FManagedArrayCollection& InCollection = GetValue<FManagedArrayCollection>(Context, &Collection);
-		SetValue<FManagedArrayCollection>(Context, InCollection, &Collection);
+		SetValue(Context, InCollection, &Collection);
 	}
 }
 
@@ -782,12 +782,12 @@ void FCollectionTransformSelectionByIntAttrDataflowNode::Evaluate(Dataflow::FCon
 		NewTransformSelection.Initialize(InCollection.NumElements(FGeometryCollection::TransformGroup), false);
 		NewTransformSelection.SetFromArray(SelectionArr);
 
-		SetValue<FDataflowTransformSelection>(Context, NewTransformSelection, &TransformSelection);
+		SetValue(Context, MoveTemp(NewTransformSelection), &TransformSelection);
 	}
 	else if (Out->IsA<FManagedArrayCollection>(&Collection))
 	{
 		const FManagedArrayCollection& InCollection = GetValue<FManagedArrayCollection>(Context, &Collection);
-		SetValue<FManagedArrayCollection>(Context, InCollection, &Collection);
+		SetValue(Context, InCollection, &Collection);
 	}
 }
 
@@ -828,17 +828,17 @@ void FCollectionVertexSelectionCustomDataflowNode::Evaluate(Dataflow::FContext&
 				}
 			}
 
-			SetValue<FDataflowVertexSelection>(Context, NewVertexSelection, &VertexSelection);
+			SetValue(Context, MoveTemp(NewVertexSelection), &VertexSelection);
 		}
 		else
 		{
-			SetValue<FDataflowVertexSelection>(Context, FDataflowVertexSelection(), &VertexSelection);
+			SetValue(Context, FDataflowVertexSelection(), &VertexSelection);
 		}
 	}
 	else if (Out->IsA<FManagedArrayCollection>(&Collection))
 	{
 		const FManagedArrayCollection& InCollection = GetValue<FManagedArrayCollection>(Context, &Collection);
-		SetValue<FManagedArrayCollection>(Context, InCollection, &Collection);
+		SetValue(Context, InCollection, &Collection);
 	}
 }
 
@@ -879,17 +879,17 @@ void FCollectionFaceSelectionCustomDataflowNode::Evaluate(Dataflow::FContext& Co
 				}
 			}
 
-			SetValue<FDataflowFaceSelection>(Context, NewFaceSelection, &FaceSelection);
+			SetValue(Context, MoveTemp(NewFaceSelection), &FaceSelection);
 		}
 		else
 		{
-			SetValue<FDataflowFaceSelection>(Context, FDataflowFaceSelection(), &FaceSelection);
+			SetValue(Context, FDataflowFaceSelection(), &FaceSelection);
 		}
 	}
 	else if (Out->IsA<FManagedArrayCollection>(&Collection))
 	{
 		const FManagedArrayCollection& InCollection = GetValue<FManagedArrayCollection>(Context, &Collection);
-		SetValue<FManagedArrayCollection>(Context, InCollection, &Collection);
+		SetValue(Context, InCollection, &Collection);
 	}
 }
 
@@ -910,7 +910,7 @@ void FCollectionSelectionConvertDataflowNode::Evaluate(Dataflow::FContext& Conte
 			NewTransformSelection.Initialize(InCollection.NumElements(FGeometryCollection::TransformGroup), false);
 			NewTransformSelection.SetFromArray(SelectionArr);
 
-			SetValue<FDataflowTransformSelection>(Context, NewTransformSelection, &TransformSelection);
+			SetValue(Context, MoveTemp(NewTransformSelection), &TransformSelection);
 		}
 		else if (IsConnected<FDataflowFaceSelection>(&FaceSelection))
 		{
@@ -924,13 +924,13 @@ void FCollectionSelectionConvertDataflowNode::Evaluate(Dataflow::FContext& Conte
 			NewTransformSelection.Initialize(InCollection.NumElements(FGeometryCollection::TransformGroup), false);
 			NewTransformSelection.SetFromArray(SelectionArr);
 
-			SetValue<FDataflowTransformSelection>(Context, NewTransformSelection, &TransformSelection);
+			SetValue(Context, MoveTemp(NewTransformSelection), &TransformSelection);
 		}
 		else if (IsConnected<FDataflowTransformSelection>(&TransformSelection))
 		{
 			// Passthrough
 			const FDataflowTransformSelection& InTransformSelection = GetValue<FDataflowTransformSelection>(Context, &TransformSelection);
-			SetValue<FDataflowTransformSelection>(Context, InTransformSelection, &TransformSelection);
+			SetValue(Context, InTransformSelection, &TransformSelection);
 		}
 	}
 	else if (Out->IsA<FDataflowFaceSelection>(&FaceSelection))
@@ -947,13 +947,13 @@ void FCollectionSelectionConvertDataflowNode::Evaluate(Dataflow::FContext& Conte
 			NewFaceSelection.Initialize(InCollection.NumElements(FGeometryCollection::FacesGroup), false);
 			NewFaceSelection.SetFromArray(SelectionArr);
 
-			SetValue<FDataflowFaceSelection>(Context, NewFaceSelection, &FaceSelection);
+			SetValue(Context, MoveTemp(NewFaceSelection), &FaceSelection);
 		}
 		else if (IsConnected<FDataflowFaceSelection>(&FaceSelection))
 		{
 			// Passthrough
 			const FDataflowFaceSelection& InFaceSelection = GetValue<FDataflowFaceSelection>(Context, &FaceSelection);
-			SetValue<FDataflowFaceSelection>(Context, InFaceSelection, &FaceSelection);
+			SetValue(Context, InFaceSelection, &FaceSelection);
 		}
 		else if (IsConnected<FDataflowTransformSelection>(&TransformSelection))
 		{
@@ -967,7 +967,7 @@ void FCollectionSelectionConvertDataflowNode::Evaluate(Dataflow::FContext& Conte
 			NewFaceSelection.Initialize(InCollection.NumElements(FGeometryCollection::FacesGroup), false);
 			NewFaceSelection.SetFromArray(SelectionArr);
 
-			SetValue<FDataflowFaceSelection>(Context, NewFaceSelection, &FaceSelection);
+			SetValue(Context, MoveTemp(NewFaceSelection), &FaceSelection);
 		}
 	}
 	else if (Out->IsA<FDataflowVertexSelection>(&VertexSelection))
@@ -976,7 +976,7 @@ void FCollectionSelectionConvertDataflowNode::Evaluate(Dataflow::FContext& Conte
 		{
 			// Passthrough
 			const FDataflowVertexSelection& InVertexSelection = GetValue<FDataflowVertexSelection>(Context, &VertexSelection);
-			SetValue<FDataflowVertexSelection>(Context, InVertexSelection, &VertexSelection);
+			SetValue(Context, InVertexSelection, &VertexSelection);
 		}
 		else if (IsConnected<FDataflowFaceSelection>(&FaceSelection))
 		{
@@ -990,7 +990,7 @@ void FCollectionSelectionConvertDataflowNode::Evaluate(Dataflow::FContext& Conte
 			NewVertexSelection.Initialize(InCollection.NumElements(FGeometryCollection::VerticesGroup), false);
 			NewVertexSelection.SetFromArray(SelectionArr);
 
-			SetValue<FDataflowVertexSelection>(Context, NewVertexSelection, &VertexSelection);
+			SetValue(Context, MoveTemp(NewVertexSelection), &VertexSelection);
 		}
 		else if (IsConnected<FDataflowTransformSelection>(&TransformSelection))
 		{
@@ -1004,13 +1004,13 @@ void FCollectionSelectionConvertDataflowNode::Evaluate(Dataflow::FContext& Conte
 			NewVertexSelection.Initialize(InCollection.NumElements(FGeometryCollection::VerticesGroup), false);
 			NewVertexSelection.SetFromArray(SelectionArr);
 
-			SetValue<FDataflowVertexSelection>(Context, NewVertexSelection, &VertexSelection);
+			SetValue(Context, MoveTemp(NewVertexSelection), &VertexSelection);
 		}
 	}
 	else if (Out->IsA<FManagedArrayCollection>(&Collection))
 	{
 		const FManagedArrayCollection& InCollection = GetValue<FManagedArrayCollection>(Context, &Collection);
-		SetValue<FManagedArrayCollection>(Context, InCollection, &Collection);
+		SetValue(Context, InCollection, &Collection);
 	}
 }
 
@@ -1042,7 +1042,7 @@ void FCollectionVertexSelectionByPercentageDataflowNode::Evaluate(Dataflow::FCon
 		GeometryCollection::Facades::FCollectionTransformSelectionFacade::SelectByPercentage(SelectionArr, InPercentage, bDeterministic, InRandomSeed);
 
 		InVertexSelection.SetFromArray(SelectionArr);
-		SetValue<FDataflowVertexSelection>(Context, InVertexSelection, &VertexSelection);
+		SetValue(Context, MoveTemp(InVertexSelection), &VertexSelection);
 	}
 }
 
@@ -1078,7 +1078,7 @@ void FCollectionVertexSelectionSetOperationDataflowNode::Evaluate(Dataflow::FCon
 			UE_LOG(LogTemp, Error, TEXT("[Dataflow ERROR] %s"), *ErrorStr);
 		}
 
-		SetValue<FDataflowVertexSelection>(Context, NewVertexSelection, &VertexSelection);
+		SetValue(Context, MoveTemp(NewVertexSelection), &VertexSelection);
 	}
 }
 
diff --git a/Engine/Plugins/Experimental/GeometryCollectionPlugin/Source/GeometryCollectionNodes/Private/Dataflow/GeometryCollectionSkeletalMeshNodes.cpp b/Engine/Plugins/Experimental/GeometryCollectionPlugin/Source/GeometryCollectionNodes/Private/Dataflow/GeometryCollectionSkeletalMeshNodes.cpp
--- a/Engine/Plugins/Experimental/GeometryCollectionPlugin/Source/GeometryCollectionNodes/Private/Dataflow/GeometryCollectionSkeletalMeshNodes.cpp
+++ b/Engine/Plugins/Experimental/GeometryCollectionPlugin/Source/GeometryCollectionNodes/Private/Dataflow/GeometryCollectionSkeletalMeshNodes.cpp
@@ -27,8 +27,8 @@ void FSkeletonToCollectionDataflowNode::Evaluate(Dataflow::FContext& Context, co
 		{
 			FGeometryCollectionEngineConversion::AppendSkeleton(SkeletonValue.Get(), FTransform::Identity, &OutCollection);
 		}
-		SetValue<DataType>(Context, OutCollection, &Collection);
+		SetValue(Context, MoveTemp(OutCollection), &Collection);
 	}
 }
 
-									
+
diff --git a/Engine/Plugins/Experimental/GeometryCollectionPlugin/Source/GeometryCollectionNodes/Private/Dataflow/GeometryCollectionUtilityNodes.cpp b/Engine/Plugins/Experimental/GeometryCollectionPlugin/Source/GeometryCollectionNodes/Private/Dataflow/GeometryCollectionUtilityNodes.cpp
--- a/Engine/Plugins/Experimental/GeometryCollectionPlugin/Source/GeometryCollectionNodes/Private/Dataflow/GeometryCollectionUtilityNodes.cpp
+++ b/Engine/Plugins/Experimental/GeometryCollectionPlugin/Source/GeometryCollectionNodes/Private/Dataflow/GeometryCollectionUtilityNodes.cpp
@@ -80,7 +80,7 @@ void FCreateLeafConvexHullsDataflowNode::Evaluate(Dataflow::FContext& Context, c
 		const FManagedArrayCollection& InCollection = GetValue(Context, &Collection);
 		if (InCollection.NumElements(FGeometryCollection::TransformGroup) == 0)
 		{
-			SetValue<FManagedArrayCollection>(Context, InCollection, &Collection);
+			SetValue(Context, InCollection, &Collection);
 			return;
 		}
 
@@ -188,7 +188,7 @@ void FCreateNonOverlappingConvexHullsDataflowNode::Evaluate(Dataflow::FContext&
 				(EConvexOverlapRemoval)OverlapRemovalMethod,
 				InOverlapRemovalShrinkPercent);
 
-			SetValue<FManagedArrayCollection>(Context, (const FManagedArrayCollection&)(*GeomCollection), &Collection);
+			SetValue(Context, MoveTemp((FManagedArrayCollection&)(*GeomCollection)), &Collection);
 		}
 	}
 }
@@ -345,7 +345,7 @@ void FGenerateClusterConvexHullsFromChildrenHullsDataflowNode::Evaluate(Dataflow
 				);
 			}
 
-			SetValue<FManagedArrayCollection>(Context, static_cast<const FManagedArrayCollection>(*GeomCollection), &Collection);
+			SetValue(Context, MoveTemp((FManagedArrayCollection&)(*GeomCollection)), &Collection);
 			// Move the negative space to the output container at the end to be sure it is no longer needed
 			Spheres.Spheres = MoveTemp(NegativeSpace);
 		}
diff --git a/Engine/Plugins/Experimental/GeometryCollectionPlugin/Source/GeometryCollectionNodes/Private/Dataflow/GeometryCollectionVerticesNodes.cpp b/Engine/Plugins/Experimental/GeometryCollectionPlugin/Source/GeometryCollectionNodes/Private/Dataflow/GeometryCollectionVerticesNodes.cpp
--- a/Engine/Plugins/Experimental/GeometryCollectionPlugin/Source/GeometryCollectionNodes/Private/Dataflow/GeometryCollectionVerticesNodes.cpp
+++ b/Engine/Plugins/Experimental/GeometryCollectionPlugin/Source/GeometryCollectionNodes/Private/Dataflow/GeometryCollectionVerticesNodes.cpp
@@ -36,7 +36,7 @@ void FTransformCollectionAttributeDataflowNode::Evaluate(Dataflow::FContext& Con
 				Vertices[i] = UEVertf(FinalTransform.TransformPosition(UEVertd(Vertices[i])));
 			}
 		}
-		SetValue<DataType>(Context, CollectionValue, &Collection);
+		SetValue(Context, MoveTemp(CollectionValue), &Collection);
 
 	}
 }
diff --git a/Engine/Source/Runtime/Experimental/Dataflow/Core/Public/Dataflow/DataflowInputOutput.h b/Engine/Source/Runtime/Experimental/Dataflow/Core/Public/Dataflow/DataflowInputOutput.h
--- a/Engine/Source/Runtime/Experimental/Dataflow/Core/Public/Dataflow/DataflowInputOutput.h
+++ b/Engine/Source/Runtime/Experimental/Dataflow/Core/Public/Dataflow/DataflowInputOutput.h
@@ -126,11 +126,11 @@ public:
 	}
  
 	template<class T>
-	void SetValue(const T& InVal, Dataflow::FContext& Context) const
+	void SetValue(T&& InVal, Dataflow::FContext& Context) const
 	{
 		if (Property)
 		{
-			Context.SetData(CacheKey(), Property, InVal);
+			Context.SetData(CacheKey(), Property, Forward<T>(InVal));
 		}
 	}
 
@@ -201,8 +201,8 @@ bool FDataflowOutput::Evaluate(Dataflow::FContext& Context) const
 	else if(const FDataflowInput* PassthroughInput = OwningNode->FindInput(GetPassthroughRealAddress()))
 	{
 		// @todo(dataflow) would be nice if the passthrough does not overwrite the existing cache value.
-		T PassthroughData = PassthroughInput->GetValue<T>(Context, *reinterpret_cast<const T*>(PassthroughInput->RealAddress()));
-		SetValue<T>(PassthroughData, Context);
+		const T& PassthroughData = PassthroughInput->GetValue<T>(Context, *reinterpret_cast<const T*>(PassthroughInput->RealAddress()));
+		SetValue(PassthroughData, Context);
 		return true;
 	}
  
diff --git a/Engine/Source/Runtime/Experimental/Dataflow/Core/Public/Dataflow/DataflowNode.h b/Engine/Source/Runtime/Experimental/Dataflow/Core/Public/Dataflow/DataflowNode.h
--- a/Engine/Source/Runtime/Experimental/Dataflow/Core/Public/Dataflow/DataflowNode.h
+++ b/Engine/Source/Runtime/Experimental/Dataflow/Core/Public/Dataflow/DataflowNode.h
@@ -189,16 +189,22 @@ struct DATAFLOWCORE_API FDataflowNode
 	/**
 	*   SetValue(...)
 	*
-	*	Set the value of the Reference output.
+	*   Set the value of the Reference output.
+	* 
+	*   Note: If the compiler errors out with "You cannot bind an lvalue to an rvalue reference", then simply remove
+	*         the explicit template parameter from the function call to allow for a const reference type to be deducted.
+	*         const int32 Value = 0; SetValue<int32>(Context, Value, &Reference);  // Error: You cannot bind an lvalue to an rvalue reference
+	*         const int32 Value = 0; SetValue(Context, Value, &Reference);  // Fine
+	*         const int32 Value = 0; SetValue<const int32&>(Context, Value, &Reference);  // Fine
 	* 
 	*   @param Context : The evaluation context that holds the data store.
 	*   @param Value : The value to store in the contexts data store. 
 	*   @param Reference : Pointer to a member of this node that corresponds with the output to set. 
 	*/
-	template<class T> void SetValue(Dataflow::FContext& Context, const T& Value, const T* Reference) const
+	template<class T> void SetValue(Dataflow::FContext& Context, T&& Value, const typename TDecay<T>::Type* Reference) const
 	{
 		checkSlow(FindOutput(Reference));
-		FindOutput(Reference)->template SetValue<T>(Value, Context);
+		FindOutput(Reference)->SetValue(Forward<T>(Value), Context);
 	}
 
 	/**
diff --git a/Engine/Source/Runtime/Experimental/Dataflow/Core/Public/Dataflow/DataflowNodeParameters.h b/Engine/Source/Runtime/Experimental/Dataflow/Core/Public/Dataflow/DataflowNodeParameters.h
--- a/Engine/Source/Runtime/Experimental/Dataflow/Core/Public/Dataflow/DataflowNodeParameters.h
+++ b/Engine/Source/Runtime/Experimental/Dataflow/Core/Public/Dataflow/DataflowNodeParameters.h
@@ -60,17 +60,13 @@ namespace Dataflow
 	template<class T>
 	struct FContextCacheElement : public FContextCacheElementBase 
 	{
-		FContextCacheElement(FProperty* InProperty, const T& InData, FTimestamp Timestamp)
-			: FContextCacheElementBase(InProperty, Timestamp)
-			, Data(InData)
-		{}
-
 		FContextCacheElement(FProperty* InProperty, T&& InData, FTimestamp Timestamp)
 			: FContextCacheElementBase(InProperty, Timestamp)
-			, Data(InData)
+			, Data(Forward<T>(InData))
 		{}
 		
-		const T Data;
+		typedef typename TDecay<T>::Type FDataType;  // Using universal references here means T could be either const& or an rvalue reference
+		const FDataType Data;                        // Decaying T removes any reference and gets the correct underlying storage data type
 	};
 
 	template<class T>
@@ -124,15 +120,6 @@ namespace Dataflow
 
 		virtual void SetDataImpl(int64 Key, TUniquePtr<FContextCacheElementBase>&& DataStoreEntry) = 0;
 		
-		template<typename T>
-		void SetData(size_t Key, FProperty* Property, const T& Value)
-		{
-			int64 IntKey = (int64)Key;
-			TUniquePtr<FContextCacheElement<T>> DataStoreEntry = MakeUnique<FContextCacheElement<T>>(Property, Value, FTimestamp::Current());
-
-			SetDataImpl(IntKey, MoveTemp(DataStoreEntry));
-		}
-
 		template<typename T>
 		void SetData(size_t Key, FProperty* Property, T&& Value)
 		{
