Dataflow : MoveTemp crash fix.
- Fixed double frees from recent move semantics changes where Unique pointers were being moved into the evaluation context with out releasing the pointer first.

diff --git a/Engine/Plugins/Experimental/ChaosFlesh/Source/ChaosFleshNodes/Private/Dataflow/ChaosFleshCoreNodes.cpp b/Engine/Plugins/Experimental/ChaosFlesh/Source/ChaosFleshNodes/Private/Dataflow/ChaosFleshCoreNodes.cpp
--- a/Engine/Plugins/Experimental/ChaosFlesh/Source/ChaosFleshNodes/Private/Dataflow/ChaosFleshCoreNodes.cpp
+++ b/Engine/Plugins/Experimental/ChaosFlesh/Source/ChaosFleshNodes/Private/Dataflow/ChaosFleshCoreNodes.cpp
@@ -37,7 +37,7 @@ void FAppendTetrahedralCollectionDataflowNode::Evaluate(Dataflow::FContext& Cont
 		{
 			GeometryGroupGuidsLocal2 = GuidArray2->GetConstArray();
 		}
-		SetValue(Context, MoveTemp((FManagedArrayCollection&)(*InCollection1)), &Collection1);
+		SetValue(Context, MoveTemp((FManagedArrayCollection&)(*InCollection1.Release())), &Collection1);
 		SetValue(Context, MoveTemp(GeometryGroupGuidsLocal1), &GeometryGroupGuidsOut1);
 		SetValue(Context, MoveTemp(GeometryGroupGuidsLocal2), &GeometryGroupGuidsOut2);
 	}
diff --git a/Engine/Plugins/Experimental/ChaosFlesh/Source/ChaosFleshNodes/Private/Dataflow/ChaosFleshRadialTetrahedronNodes.cpp b/Engine/Plugins/Experimental/ChaosFlesh/Source/ChaosFleshNodes/Private/Dataflow/ChaosFleshRadialTetrahedronNodes.cpp
--- a/Engine/Plugins/Experimental/ChaosFlesh/Source/ChaosFleshNodes/Private/Dataflow/ChaosFleshRadialTetrahedronNodes.cpp
+++ b/Engine/Plugins/Experimental/ChaosFlesh/Source/ChaosFleshNodes/Private/Dataflow/ChaosFleshRadialTetrahedronNodes.cpp
@@ -48,7 +48,7 @@ void FRadialTetrahedronDataflowNodes::Evaluate(Dataflow::FContext& Context, cons
 		TUniquePtr<FTetrahedralCollection> RadialCollection(FTetrahedralCollection::NewTetrahedralCollection(TetVertices, SurfaceElements, TetElements));
 		InCollection->AppendGeometry(*RadialCollection.Get());
 
-		SetValue(Context, MoveTemp((FManagedArrayCollection&)(*InCollection)), &Collection);
+		SetValue(Context, MoveTemp((FManagedArrayCollection&)(*InCollection.Release())), &Collection);
 	}
 }
 
diff --git a/Engine/Plugins/Experimental/ChaosFlesh/Source/ChaosFleshNodes/Private/Dataflow/ChaosFleshTetrahedralNodes.cpp b/Engine/Plugins/Experimental/ChaosFlesh/Source/ChaosFleshNodes/Private/Dataflow/ChaosFleshTetrahedralNodes.cpp
--- a/Engine/Plugins/Experimental/ChaosFlesh/Source/ChaosFleshNodes/Private/Dataflow/ChaosFleshTetrahedralNodes.cpp
+++ b/Engine/Plugins/Experimental/ChaosFlesh/Source/ChaosFleshNodes/Private/Dataflow/ChaosFleshTetrahedralNodes.cpp
@@ -89,7 +89,6 @@ FCalculateTetMetrics::Evaluate(Dataflow::FContext& Context, const FDataflowOutpu
 {
 	if (Out->IsA<DataType>(&Collection))
 	{
-		//DataType OutCollection = GetValue<DataType>(Context, &Collection); // Deep copy
 		TUniquePtr<FFleshCollection> InCollection(GetValue<DataType>(Context, &Collection).NewCopy<FFleshCollection>());
 
 		TManagedArray<FIntVector4>* TetMesh =
@@ -168,8 +167,7 @@ FCalculateTetMetrics::Evaluate(Dataflow::FContext& Context, const FDataflowOutpu
 				MinAR, AvgAR, MaxAR);
 		}
 
-		//SetValue(Context, MoveTemp((FManagedArrayCollection&)(*OutCollection)), &Collection);
-		SetValue(Context, MoveTemp((FManagedArrayCollection&)(*InCollection)), &Collection);
+		SetValue(Context, MoveTemp((FManagedArrayCollection&)(*InCollection.Release())), &Collection);
 	}
 }
 
@@ -200,7 +198,7 @@ void FConstructTetGridNode::Evaluate(Dataflow::FContext& Context, const FDataflo
 			FTetrahedralCollection::NewTetrahedralCollection(X, Tris, Tets));
 		InCollection->AppendGeometry(*TetCollection.Get());
 
-		SetValue(Context, MoveTemp((FManagedArrayCollection&)(*InCollection)), &Collection);
+		SetValue(Context, MoveTemp((FManagedArrayCollection&)(*InCollection.Release())), &Collection);
 	}
 }
 
@@ -327,7 +325,7 @@ void FGenerateTetrahedralCollectionDataflowNodes::Evaluate(Dataflow::FContext& C
 			ensureMsgf(false, TEXT("FGenerateTetrahedralCollectionDataflowNodes is an editor only node."));
 #endif
 		} // end if InStaticMesh || InSkeletalMesh
-		SetValue(Context, MoveTemp((FManagedArrayCollection&)(*InCollection)), &Collection);
+		SetValue(Context, MoveTemp((FManagedArrayCollection&)(*InCollection.Release())), &Collection);
 	}
 }
 
diff --git a/Engine/Plugins/Experimental/GeometryCollectionPlugin/Source/GeometryCollectionNodes/Private/Dataflow/GeometryCollectionClusteringNodes.cpp b/Engine/Plugins/Experimental/GeometryCollectionPlugin/Source/GeometryCollectionNodes/Private/Dataflow/GeometryCollectionClusteringNodes.cpp
--- a/Engine/Plugins/Experimental/GeometryCollectionPlugin/Source/GeometryCollectionNodes/Private/Dataflow/GeometryCollectionClusteringNodes.cpp
+++ b/Engine/Plugins/Experimental/GeometryCollectionPlugin/Source/GeometryCollectionNodes/Private/Dataflow/GeometryCollectionClusteringNodes.cpp
@@ -87,7 +87,7 @@ void FAutoClusterDataflowNode::Evaluate(Dataflow::FContext& Context, const FData
 				InEnforceSiteParameters,
 				InGridX, InGridY, InGridZ, InMinimumClusterSize, InKMeansIterations);
 
-			SetValue(Context, MoveTemp((FManagedArrayCollection&)(*GeomCollection)), &Collection);
+			SetValue(Context, MoveTemp((FManagedArrayCollection&)(*GeomCollection.Release())), &Collection);
 		}
 	}
 }
@@ -138,16 +138,14 @@ void FClusterUnclusterDataflowNode::Evaluate(Dataflow::FContext& Context, const
 			GeometryCollection::Facades::FCollectionTransformSelectionFacade SelectionFacade(*GeomCollection);
 			SelectionFacade.ConvertSelectionToClusterNodes(Selection, false);
 			SelectionFacade.RemoveRootNodes(Selection);
-			if (Selection.IsEmpty())
+			if (!Selection.IsEmpty())
 			{
-				SetValue(Context, MoveTemp((FManagedArrayCollection&)(*GeomCollection)), &Collection);
-				return;
-			}
-			FGeometryCollectionClusteringUtility::CollapseHierarchyOneLevel(GeomCollection.Get(), Selection);
-			FGeometryCollectionClusteringUtility::RemoveDanglingClusters(GeomCollection.Get());
+				FGeometryCollectionClusteringUtility::CollapseHierarchyOneLevel(GeomCollection.Get(), Selection);
+				FGeometryCollectionClusteringUtility::RemoveDanglingClusters(GeomCollection.Get());
 
-			HierarchyFacade.GenerateLevelAttribute();
-			SetValue(Context, MoveTemp((FManagedArrayCollection&)(*GeomCollection)), &Collection);
+				HierarchyFacade.GenerateLevelAttribute();
+			}
+			SetValue(Context, MoveTemp((FManagedArrayCollection&)(*GeomCollection.Release())), &Collection);
 		}
 	}
 }
diff --git a/Engine/Plugins/Experimental/GeometryCollectionPlugin/Source/GeometryCollectionNodes/Private/Dataflow/GeometryCollectionEditNodes.cpp b/Engine/Plugins/Experimental/GeometryCollectionPlugin/Source/GeometryCollectionNodes/Private/Dataflow/GeometryCollectionEditNodes.cpp
--- a/Engine/Plugins/Experimental/GeometryCollectionPlugin/Source/GeometryCollectionNodes/Private/Dataflow/GeometryCollectionEditNodes.cpp
+++ b/Engine/Plugins/Experimental/GeometryCollectionPlugin/Source/GeometryCollectionNodes/Private/Dataflow/GeometryCollectionEditNodes.cpp
@@ -66,7 +66,7 @@ void FPruneInCollectionDataflowNode::Evaluate(Dataflow::FContext& Context, const
 
 					FFractureEngineEdit::DeleteBranch(*GeomCollection, BoneIndices);
 
-					SetValue(Context, MoveTemp((FManagedArrayCollection&)(*GeomCollection)), &Collection);
+					SetValue(Context, MoveTemp((FManagedArrayCollection&)(*GeomCollection.Release())), &Collection);
 					return;
 				}
 			}
@@ -172,7 +172,7 @@ void FMergeInCollectionDataflowNode::Evaluate(Dataflow::FContext& Context, const
 
 					FFractureEngineEdit::Merge(*GeomCollection, BoneIndices);
 
-					SetValue(Context, MoveTemp((FManagedArrayCollection&)(*GeomCollection)), &Collection);
+					SetValue(Context, MoveTemp((FManagedArrayCollection&)(*GeomCollection.Release())), &Collection);
 					return;
 				}
 			}
diff --git a/Engine/Plugins/Experimental/GeometryCollectionPlugin/Source/GeometryCollectionNodes/Private/Dataflow/GeometryCollectionUtilityNodes.cpp b/Engine/Plugins/Experimental/GeometryCollectionPlugin/Source/GeometryCollectionNodes/Private/Dataflow/GeometryCollectionUtilityNodes.cpp
--- a/Engine/Plugins/Experimental/GeometryCollectionPlugin/Source/GeometryCollectionNodes/Private/Dataflow/GeometryCollectionUtilityNodes.cpp
+++ b/Engine/Plugins/Experimental/GeometryCollectionPlugin/Source/GeometryCollectionNodes/Private/Dataflow/GeometryCollectionUtilityNodes.cpp
@@ -188,7 +188,7 @@ void FCreateNonOverlappingConvexHullsDataflowNode::Evaluate(Dataflow::FContext&
 				(EConvexOverlapRemoval)OverlapRemovalMethod,
 				InOverlapRemovalShrinkPercent);
 
-			SetValue(Context, MoveTemp((FManagedArrayCollection&)(*GeomCollection)), &Collection);
+			SetValue(Context, MoveTemp((FManagedArrayCollection&)(*GeomCollection.Release())), &Collection);
 		}
 	}
 }
@@ -345,7 +345,7 @@ void FGenerateClusterConvexHullsFromChildrenHullsDataflowNode::Evaluate(Dataflow
 				);
 			}
 
-			SetValue(Context, MoveTemp((FManagedArrayCollection&)(*GeomCollection)), &Collection);
+			SetValue(Context, MoveTemp((FManagedArrayCollection&)(*GeomCollection.Release())), &Collection);
 			// Move the negative space to the output container at the end to be sure it is no longer needed
 			Spheres.Spheres = MoveTemp(NegativeSpace);
 		}
