Dataflow - Fixed crashes caused by SetValue with move semantic on cast FleshCollection and GeometryCollection.

#rnx

diff --git a/Engine/Plugins/Experimental/ChaosFlesh/Source/ChaosFleshNodes/Private/Dataflow/ChaosFleshCoreNodes.cpp b/Engine/Plugins/Experimental/ChaosFlesh/Source/ChaosFleshNodes/Private/Dataflow/ChaosFleshCoreNodes.cpp
--- a/Engine/Plugins/Experimental/ChaosFlesh/Source/ChaosFleshNodes/Private/Dataflow/ChaosFleshCoreNodes.cpp
+++ b/Engine/Plugins/Experimental/ChaosFlesh/Source/ChaosFleshNodes/Private/Dataflow/ChaosFleshCoreNodes.cpp
@@ -37,7 +37,7 @@ void FAppendTetrahedralCollectionDataflowNode::Evaluate(Dataflow::FContext& Cont
 		{
 			GeometryGroupGuidsLocal2 = GuidArray2->GetConstArray();
 		}
-		SetValue(Context, MoveTemp((FManagedArrayCollection&)(*InCollection1.Release())), &Collection1);
+		SetValue<const FManagedArrayCollection&>(Context, *InCollection1, &Collection1);
 		SetValue(Context, MoveTemp(GeometryGroupGuidsLocal1), &GeometryGroupGuidsOut1);
 		SetValue(Context, MoveTemp(GeometryGroupGuidsLocal2), &GeometryGroupGuidsOut2);
 	}
diff --git a/Engine/Plugins/Experimental/ChaosFlesh/Source/ChaosFleshNodes/Private/Dataflow/ChaosFleshRadialTetrahedronNodes.cpp b/Engine/Plugins/Experimental/ChaosFlesh/Source/ChaosFleshNodes/Private/Dataflow/ChaosFleshRadialTetrahedronNodes.cpp
--- a/Engine/Plugins/Experimental/ChaosFlesh/Source/ChaosFleshNodes/Private/Dataflow/ChaosFleshRadialTetrahedronNodes.cpp
+++ b/Engine/Plugins/Experimental/ChaosFlesh/Source/ChaosFleshNodes/Private/Dataflow/ChaosFleshRadialTetrahedronNodes.cpp
@@ -48,7 +48,7 @@ void FRadialTetrahedronDataflowNodes::Evaluate(Dataflow::FContext& Context, cons
 		TUniquePtr<FTetrahedralCollection> RadialCollection(FTetrahedralCollection::NewTetrahedralCollection(TetVertices, SurfaceElements, TetElements));
 		InCollection->AppendGeometry(*RadialCollection.Get());
 
-		SetValue(Context, MoveTemp((FManagedArrayCollection&)(*InCollection.Release())), &Collection);
+		SetValue<const DataType&>(Context, *InCollection, &Collection);
 	}
 }
 
diff --git a/Engine/Plugins/Experimental/ChaosFlesh/Source/ChaosFleshNodes/Private/Dataflow/ChaosFleshTetrahedralNodes.cpp b/Engine/Plugins/Experimental/ChaosFlesh/Source/ChaosFleshNodes/Private/Dataflow/ChaosFleshTetrahedralNodes.cpp
--- a/Engine/Plugins/Experimental/ChaosFlesh/Source/ChaosFleshNodes/Private/Dataflow/ChaosFleshTetrahedralNodes.cpp
+++ b/Engine/Plugins/Experimental/ChaosFlesh/Source/ChaosFleshNodes/Private/Dataflow/ChaosFleshTetrahedralNodes.cpp
@@ -167,7 +167,7 @@ FCalculateTetMetrics::Evaluate(Dataflow::FContext& Context, const FDataflowOutpu
 				MinAR, AvgAR, MaxAR);
 		}
 
-		SetValue(Context, MoveTemp((FManagedArrayCollection&)(*InCollection.Release())), &Collection);
+		SetValue<const DataType&>(Context, *InCollection, &Collection);
 	}
 }
 
@@ -198,7 +198,7 @@ void FConstructTetGridNode::Evaluate(Dataflow::FContext& Context, const FDataflo
 			FTetrahedralCollection::NewTetrahedralCollection(X, Tris, Tets));
 		InCollection->AppendGeometry(*TetCollection.Get());
 
-		SetValue(Context, MoveTemp((FManagedArrayCollection&)(*InCollection.Release())), &Collection);
+		SetValue<const DataType&>(Context, *InCollection, &Collection);
 	}
 }
 
@@ -325,7 +325,7 @@ void FGenerateTetrahedralCollectionDataflowNodes::Evaluate(Dataflow::FContext& C
 			ensureMsgf(false, TEXT("FGenerateTetrahedralCollectionDataflowNodes is an editor only node."));
 #endif
 		} // end if InStaticMesh || InSkeletalMesh
-		SetValue(Context, MoveTemp((FManagedArrayCollection&)(*InCollection.Release())), &Collection);
+		SetValue<const DataType&>(Context, *InCollection, &Collection);
 	}
 }
 
diff --git a/Engine/Plugins/Experimental/GeometryCollectionPlugin/Source/GeometryCollectionNodes/Private/Dataflow/GeometryCollectionClusteringNodes.cpp b/Engine/Plugins/Experimental/GeometryCollectionPlugin/Source/GeometryCollectionNodes/Private/Dataflow/GeometryCollectionClusteringNodes.cpp
--- a/Engine/Plugins/Experimental/GeometryCollectionPlugin/Source/GeometryCollectionNodes/Private/Dataflow/GeometryCollectionClusteringNodes.cpp
+++ b/Engine/Plugins/Experimental/GeometryCollectionPlugin/Source/GeometryCollectionNodes/Private/Dataflow/GeometryCollectionClusteringNodes.cpp
@@ -87,7 +87,7 @@ void FAutoClusterDataflowNode::Evaluate(Dataflow::FContext& Context, const FData
 				InEnforceSiteParameters,
 				InGridX, InGridY, InGridZ, InMinimumClusterSize, InKMeansIterations);
 
-			SetValue(Context, MoveTemp((FManagedArrayCollection&)(*GeomCollection.Release())), &Collection);
+			SetValue<const FManagedArrayCollection&>(Context, *GeomCollection, &Collection);
 		}
 	}
 }
@@ -145,7 +145,7 @@ void FClusterUnclusterDataflowNode::Evaluate(Dataflow::FContext& Context, const
 
 				HierarchyFacade.GenerateLevelAttribute();
 			}
-			SetValue(Context, MoveTemp((FManagedArrayCollection&)(*GeomCollection.Release())), &Collection);
+			SetValue<const FManagedArrayCollection&>(Context, *GeomCollection, &Collection);
 		}
 	}
 }
diff --git a/Engine/Plugins/Experimental/GeometryCollectionPlugin/Source/GeometryCollectionNodes/Private/Dataflow/GeometryCollectionEditNodes.cpp b/Engine/Plugins/Experimental/GeometryCollectionPlugin/Source/GeometryCollectionNodes/Private/Dataflow/GeometryCollectionEditNodes.cpp
--- a/Engine/Plugins/Experimental/GeometryCollectionPlugin/Source/GeometryCollectionNodes/Private/Dataflow/GeometryCollectionEditNodes.cpp
+++ b/Engine/Plugins/Experimental/GeometryCollectionPlugin/Source/GeometryCollectionNodes/Private/Dataflow/GeometryCollectionEditNodes.cpp
@@ -66,7 +66,7 @@ void FPruneInCollectionDataflowNode::Evaluate(Dataflow::FContext& Context, const
 
 					FFractureEngineEdit::DeleteBranch(*GeomCollection, BoneIndices);
 
-					SetValue(Context, MoveTemp((FManagedArrayCollection&)(*GeomCollection.Release())), &Collection);
+					SetValue<const FManagedArrayCollection&>(Context, *GeomCollection, &Collection);
 					return;
 				}
 			}
@@ -172,7 +172,7 @@ void FMergeInCollectionDataflowNode::Evaluate(Dataflow::FContext& Context, const
 
 					FFractureEngineEdit::Merge(*GeomCollection, BoneIndices);
 
-					SetValue(Context, MoveTemp((FManagedArrayCollection&)(*GeomCollection.Release())), &Collection);
+					SetValue<const FManagedArrayCollection&>(Context, *GeomCollection, &Collection);
 					return;
 				}
 			}
diff --git a/Engine/Plugins/Experimental/GeometryCollectionPlugin/Source/GeometryCollectionNodes/Private/Dataflow/GeometryCollectionFracturingNodes.cpp b/Engine/Plugins/Experimental/GeometryCollectionPlugin/Source/GeometryCollectionNodes/Private/Dataflow/GeometryCollectionFracturingNodes.cpp
--- a/Engine/Plugins/Experimental/GeometryCollectionPlugin/Source/GeometryCollectionNodes/Private/Dataflow/GeometryCollectionFracturingNodes.cpp
+++ b/Engine/Plugins/Experimental/GeometryCollectionPlugin/Source/GeometryCollectionNodes/Private/Dataflow/GeometryCollectionFracturingNodes.cpp
@@ -196,7 +196,7 @@ void FExplodedViewDataflowNode::Evaluate(Dataflow::FContext& Context, const FDat
 
 		FFractureEngineFracturing::GenerateExplodedViewAttribute(InCollection, GetValue<FVector>(Context, &Scale), GetValue<float>(Context, &UniformScale));
 
-		SetValue<FManagedArrayCollection>(Context, MoveTemp(InCollection), &Collection);
+		SetValue(Context, MoveTemp(InCollection), &Collection);
 	}
 }
 
diff --git a/Engine/Plugins/Experimental/GeometryCollectionPlugin/Source/GeometryCollectionNodes/Private/Dataflow/GeometryCollectionNodes.cpp b/Engine/Plugins/Experimental/GeometryCollectionPlugin/Source/GeometryCollectionNodes/Private/Dataflow/GeometryCollectionNodes.cpp
--- a/Engine/Plugins/Experimental/GeometryCollectionPlugin/Source/GeometryCollectionNodes/Private/Dataflow/GeometryCollectionNodes.cpp
+++ b/Engine/Plugins/Experimental/GeometryCollectionPlugin/Source/GeometryCollectionNodes/Private/Dataflow/GeometryCollectionNodes.cpp
@@ -649,7 +649,7 @@ void FSetAnchorStateDataflowNode::Evaluate(Dataflow::FContext& Context, const FD
 				AnchoringFacade.SetAnchored(BoneIndices, !bAnchored);
 			}
 
-			SetValue(Context, MoveTemp((FManagedArrayCollection&)(*GeomCollection)), &Collection);
+			SetValue<const FManagedArrayCollection&>(Context, *GeomCollection, &Collection);
 		}
 	}
 }
@@ -679,7 +679,7 @@ void FProximityDataflowNode::Evaluate(Dataflow::FContext& Context, const FDatafl
 			ProximityUtility.InvalidateProximity();
 			ProximityUtility.UpdateProximity();
 
-			SetValue(Context, MoveTemp((FManagedArrayCollection&)(*GeomCollection)), &Collection);
+			SetValue<const FManagedArrayCollection&>(Context, *GeomCollection, &Collection);
 		}
 	}
 }
diff --git a/Engine/Plugins/Experimental/GeometryCollectionPlugin/Source/GeometryCollectionNodes/Private/Dataflow/GeometryCollectionUtilityNodes.cpp b/Engine/Plugins/Experimental/GeometryCollectionPlugin/Source/GeometryCollectionNodes/Private/Dataflow/GeometryCollectionUtilityNodes.cpp
--- a/Engine/Plugins/Experimental/GeometryCollectionPlugin/Source/GeometryCollectionNodes/Private/Dataflow/GeometryCollectionUtilityNodes.cpp
+++ b/Engine/Plugins/Experimental/GeometryCollectionPlugin/Source/GeometryCollectionNodes/Private/Dataflow/GeometryCollectionUtilityNodes.cpp
@@ -131,7 +131,7 @@ void FSimplifyConvexHullsDataflowNode::Evaluate(Dataflow::FContext& Context, con
 		FManagedArrayCollection InCollection = GetValue(Context, &Collection);
 		if (InCollection.NumElements(FGeometryCollection::TransformGroup) == 0)
 		{
-			SetValue<FManagedArrayCollection>(Context, MoveTemp(InCollection), &Collection);
+			SetValue(Context, MoveTemp(InCollection), &Collection);
 			return;
 		}
 
@@ -153,7 +153,7 @@ void FSimplifyConvexHullsDataflowNode::Evaluate(Dataflow::FContext& Context, con
 		Settings.bUseExistingVertexPositions = bUseExistingVertices;
 		Settings.TargetTriangleCount = GetValue(Context, &MinTargetTriangleCount);
 		UE::FractureEngine::Convex::SimplifyConvexHulls(InCollection, Settings, bRestrictToSelection, SelectedBones);
-		SetValue<FManagedArrayCollection>(Context, MoveTemp(InCollection), &Collection);
+		SetValue(Context, MoveTemp(InCollection), &Collection);
 	}
 }
 
@@ -188,7 +188,7 @@ void FCreateNonOverlappingConvexHullsDataflowNode::Evaluate(Dataflow::FContext&
 				(EConvexOverlapRemoval)OverlapRemovalMethod,
 				InOverlapRemovalShrinkPercent);
 
-			SetValue(Context, MoveTemp((FManagedArrayCollection&)(*GeomCollection.Release())), &Collection);
+			SetValue<const FManagedArrayCollection&>(Context, *GeomCollection, &Collection);
 		}
 	}
 }
@@ -345,7 +345,7 @@ void FGenerateClusterConvexHullsFromChildrenHullsDataflowNode::Evaluate(Dataflow
 				);
 			}
 
-			SetValue(Context, MoveTemp((FManagedArrayCollection&)(*GeomCollection.Release())), &Collection);
+			SetValue<const FManagedArrayCollection&>(Context, *GeomCollection, &Collection);
 			// Move the negative space to the output container at the end to be sure it is no longer needed
 			Spheres.Spheres = MoveTemp(NegativeSpace);
 		}
@@ -413,7 +413,7 @@ void FMergeConvexHullsDataflowNode::Evaluate(Dataflow::FContext& Context, const
 
 		FGeometryCollectionConvexUtility::MergeHullsOnTransforms(InCollection, HullMergeSettings, bHasSelectionFilter, SelectionArray);
 
-		SetValue<FManagedArrayCollection>(Context, MoveTemp(InCollection), &Collection);
+		SetValue(Context, MoveTemp(InCollection), &Collection);
 
 		// Move the negative space to the output container at the end, after it is no longer needed
 		Spheres.Spheres = MoveTemp(NegativeSpace);
