add dataflow nodes to support selecting geometry collection transform indices based on lengths (in chosen axis, min/max axis, or on diagonal) of their bounding boxes

diff --git a/Engine/Plugins/Experimental/GeometryCollectionPlugin/Source/GeometryCollectionNodes/Private/Dataflow/GeometryCollectionNodes.cpp b/Engine/Plugins/Experimental/GeometryCollectionPlugin/Source/GeometryCollectionNodes/Private/Dataflow/GeometryCollectionNodes.cpp
--- a/Engine/Plugins/Experimental/GeometryCollectionPlugin/Source/GeometryCollectionNodes/Private/Dataflow/GeometryCollectionNodes.cpp
+++ b/Engine/Plugins/Experimental/GeometryCollectionPlugin/Source/GeometryCollectionNodes/Private/Dataflow/GeometryCollectionNodes.cpp
@@ -53,6 +53,7 @@ namespace Dataflow
 		DATAFLOW_NODE_REGISTER_CREATION_FACTORY(FLogStringDataflowNode);
 		DATAFLOW_NODE_REGISTER_CREATION_FACTORY(FBoundingBoxDataflowNode);
 		DATAFLOW_NODE_REGISTER_CREATION_FACTORY(FExpandBoundingBoxDataflowNode);
+		DATAFLOW_NODE_REGISTER_CREATION_FACTORY(FGetBoxLengthsDataflowNode);
 		DATAFLOW_NODE_REGISTER_CREATION_FACTORY(FExpandVectorDataflowNode);
 		DATAFLOW_NODE_REGISTER_CREATION_FACTORY(FStringAppendDataflowNode);
 		DATAFLOW_NODE_REGISTER_CREATION_FACTORY(FHashStringDataflowNode);
@@ -181,6 +182,24 @@ void FBoundingBoxDataflowNode::Evaluate(Dataflow::FContext& Context, const FData
 	}
 }
 
+void FGetBoxLengthsDataflowNode::Evaluate(Dataflow::FContext& Context, const FDataflowOutput* Out) const
+{
+	if (Out->IsA(&Lengths))
+	{
+		const TArray<FBox>& InBoxes = GetValue(Context, &Boxes);
+
+		TArray<float> OutLengths;
+		OutLengths.SetNumUninitialized(InBoxes.Num());
+		for (int32 Idx = 0; Idx < InBoxes.Num(); ++Idx)
+		{
+			const FBox& Box = InBoxes[Idx];
+			OutLengths[Idx] = BoxToMeasurement(Box);
+		}
+
+		SetValue(Context, MoveTemp(OutLengths), &Lengths);
+	}
+}
+
 
 void FExpandBoundingBoxDataflowNode::Evaluate(Dataflow::FContext& Context, const FDataflowOutput* Out) const
 {
diff --git a/Engine/Plugins/Experimental/GeometryCollectionPlugin/Source/GeometryCollectionNodes/Private/Dataflow/GeometryCollectionSelectionNodes.cpp b/Engine/Plugins/Experimental/GeometryCollectionPlugin/Source/GeometryCollectionNodes/Private/Dataflow/GeometryCollectionSelectionNodes.cpp
--- a/Engine/Plugins/Experimental/GeometryCollectionPlugin/Source/GeometryCollectionNodes/Private/Dataflow/GeometryCollectionSelectionNodes.cpp
+++ b/Engine/Plugins/Experimental/GeometryCollectionPlugin/Source/GeometryCollectionNodes/Private/Dataflow/GeometryCollectionSelectionNodes.cpp
@@ -46,6 +46,7 @@ namespace Dataflow
 		DATAFLOW_NODE_REGISTER_CREATION_FACTORY(FCollectionTransformSelectionRandomDataflowNode);
 		DATAFLOW_NODE_REGISTER_CREATION_FACTORY(FCollectionTransformSelectionRootDataflowNode);
 		DATAFLOW_NODE_REGISTER_CREATION_FACTORY(FCollectionTransformSelectionCustomDataflowNode);
+		DATAFLOW_NODE_REGISTER_CREATION_FACTORY(FCollectionTransformSelectionFromIndexArrayDataflowNode);
 		DATAFLOW_NODE_REGISTER_CREATION_FACTORY(FCollectionTransformSelectionParentDataflowNode);
 		DATAFLOW_NODE_REGISTER_CREATION_FACTORY(FCollectionTransformSelectionByPercentageDataflowNode);
 		DATAFLOW_NODE_REGISTER_CREATION_FACTORY(FCollectionTransformSelectionChildrenDataflowNode);
@@ -60,6 +61,7 @@ namespace Dataflow
 		DATAFLOW_NODE_REGISTER_CREATION_FACTORY(FCollectionTransformSelectionInBoxDataflowNode);
 		DATAFLOW_NODE_REGISTER_CREATION_FACTORY(FCollectionTransformSelectionInSphereDataflowNode);
 		DATAFLOW_NODE_REGISTER_CREATION_FACTORY(FCollectionTransformSelectionByFloatAttrDataflowNode);
+		DATAFLOW_NODE_REGISTER_CREATION_FACTORY(FSelectFloatArrayIndicesInRangeDataflowNode);
 		DATAFLOW_NODE_REGISTER_CREATION_FACTORY(FCollectionTransformSelectionByIntAttrDataflowNode);
 		DATAFLOW_NODE_REGISTER_CREATION_FACTORY(FCollectionVertexSelectionCustomDataflowNode);
 		DATAFLOW_NODE_REGISTER_CREATION_FACTORY(FCollectionFaceSelectionCustomDataflowNode);
@@ -402,6 +404,46 @@ void FCollectionTransformSelectionCustomDataflowNode::Evaluate(Dataflow::FContex
 	else if (Out->IsA<FManagedArrayCollection>(&Collection))
 	{
 		const FManagedArrayCollection& InCollection = GetValue<FManagedArrayCollection>(Context, &Collection);
+	}
+}
+
+
+void FCollectionTransformSelectionFromIndexArrayDataflowNode::Evaluate(Dataflow::FContext& Context, const FDataflowOutput* Out) const
+{
+	if (Out->IsA(&TransformSelection))
+	{
+		const FManagedArrayCollection& InCollection = GetValue(Context, &Collection);
+
+		if (InCollection.HasGroup(FGeometryCollection::TransformGroup))
+		{
+			const int32 NumTransforms = InCollection.NumElements(FGeometryCollection::TransformGroup);
+
+			const TArray<int32>& InBoneIndices = GetValue(Context, &BoneIndices);
+
+			FDataflowTransformSelection NewTransformSelection;
+			NewTransformSelection.Initialize(NumTransforms, false);
+			for (int32 SelectedIdx : InBoneIndices)
+			{
+				if (SelectedIdx >= 0 && SelectedIdx < NumTransforms)
+				{
+					NewTransformSelection.SetSelected(SelectedIdx);
+				}
+				else
+				{
+					UE_LOG(LogChaos, Error, TEXT("[Dataflow ERROR] Invalid selection index %d is outside valid bone index range [0, %d)"), SelectedIdx, NumTransforms);
+				}
+			}
+
+			SetValue(Context, MoveTemp(NewTransformSelection), &TransformSelection);
+		}
+		else
+		{
+			SetValue(Context, FDataflowTransformSelection(), &TransformSelection);
+		}
+	}
+	else if (Out->IsA(&Collection))
+	{
+		const FManagedArrayCollection& InCollection = GetValue(Context, &Collection);
 		SetValue(Context, InCollection, &Collection);
 	}
 }
@@ -765,6 +807,37 @@ void FCollectionTransformSelectionByFloatAttrDataflowNode::Evaluate(Dataflow::FC
 	}
 }
 
+void FSelectFloatArrayIndicesInRangeDataflowNode::Evaluate(Dataflow::FContext& Context, const FDataflowOutput* Out) const
+{
+	if (Out->IsA(&Indices))
+	{
+		const TArray<float>& InValues = GetValue(Context, &Values);
+		float InMin = GetValue(Context, &Min);
+		float InMax = GetValue(Context, &Max);
+		bool bInsideRange = RangeSetting == ERangeSettingEnum::Dataflow_RangeSetting_InsideRange;
+
+		TArray<int32> OutIndices;
+		for (int32 Idx = 0; Idx < InValues.Num(); ++Idx)
+		{
+			const float FloatValue = InValues[Idx];
+
+			if (bInsideRange && FloatValue > Min && FloatValue < Max)
+			{
+				OutIndices.Add(Idx);
+			}
+			else if (!bInsideRange && (FloatValue < Min || FloatValue > Max))
+			{
+				OutIndices.Add(Idx);
+			}
+			else if (bInclusive && (FloatValue == Min || FloatValue == Max))
+			{
+				OutIndices.Add(Idx);
+			}
+		}
+
+		SetValue(Context, MoveTemp(OutIndices), &Indices);
+	}
+}
 
 void FCollectionTransformSelectionByIntAttrDataflowNode::Evaluate(Dataflow::FContext& Context, const FDataflowOutput* Out) const
 {
diff --git a/Engine/Plugins/Experimental/GeometryCollectionPlugin/Source/GeometryCollectionNodes/Public/Dataflow/GeometryCollectionNodes.h b/Engine/Plugins/Experimental/GeometryCollectionPlugin/Source/GeometryCollectionNodes/Public/Dataflow/GeometryCollectionNodes.h
--- a/Engine/Plugins/Experimental/GeometryCollectionPlugin/Source/GeometryCollectionNodes/Public/Dataflow/GeometryCollectionNodes.h
+++ b/Engine/Plugins/Experimental/GeometryCollectionPlugin/Source/GeometryCollectionNodes/Public/Dataflow/GeometryCollectionNodes.h
@@ -191,6 +191,73 @@ public:
 
 };
 
+UENUM()
+enum class EBoxLengthMeasurementMethod : uint8
+{
+	XAxis,
+	YAxis,
+	ZAxis,
+	ShortestAxis,
+	LongestAxis,
+	Diagonal
+};
+
+/**
+ *
+ * Create an array of lengths of bounding boxes (measured along an axis, diagonal, or the max/min axes) from an array of bounding boxes
+ *
+ */
+USTRUCT(meta = (DataflowGeometryCollection))
+struct FGetBoxLengthsDataflowNode : public FDataflowNode
+{
+	GENERATED_USTRUCT_BODY()
+	DATAFLOW_NODE_DEFINE_INTERNAL(FGetBoxLengthsDataflowNode, "GetBoxLengths", "Utilities|Box", "")
+
+public:
+	
+	UPROPERTY(meta = (DataflowInput))
+	TArray<FBox> Boxes;
+
+	UPROPERTY(meta = (DataflowOutput))
+	TArray<float> Lengths;
+
+	UPROPERTY(EditAnywhere, Category = Options)
+	EBoxLengthMeasurementMethod MeasurementMethod = EBoxLengthMeasurementMethod::Diagonal;
+
+	inline double BoxToMeasurement(const FBox& Box) const
+	{
+		FVector Size = Box.GetSize();
+		switch (MeasurementMethod)
+		{
+		case EBoxLengthMeasurementMethod::XAxis:
+			return Size.X;
+		case EBoxLengthMeasurementMethod::YAxis:
+			return Size.Y;
+		case EBoxLengthMeasurementMethod::ZAxis:
+			return Size.Z;
+		case EBoxLengthMeasurementMethod::ShortestAxis:
+			return Size.GetMin();
+		case EBoxLengthMeasurementMethod::LongestAxis:
+			return Size.GetMax();
+		case EBoxLengthMeasurementMethod::Diagonal:
+			return Size.Length();
+		}
+		checkNoEntry(); // switch above should handle all cases
+		return Size.X;
+	}
+
+	FGetBoxLengthsDataflowNode(const Dataflow::FNodeParameters& InParam, FGuid InGuid = FGuid::NewGuid())
+		: FDataflowNode(InParam, InGuid)
+	{
+		RegisterInputConnection(&Boxes);
+		RegisterOutputConnection(&Lengths);
+	}
+
+	virtual void Evaluate(Dataflow::FContext& Context, const FDataflowOutput* Out) const override;
+
+};
+
+
 
 /**
  *
diff --git a/Engine/Plugins/Experimental/GeometryCollectionPlugin/Source/GeometryCollectionNodes/Public/Dataflow/GeometryCollectionSelectionNodes.h b/Engine/Plugins/Experimental/GeometryCollectionPlugin/Source/GeometryCollectionNodes/Public/Dataflow/GeometryCollectionSelectionNodes.h
--- a/Engine/Plugins/Experimental/GeometryCollectionPlugin/Source/GeometryCollectionNodes/Public/Dataflow/GeometryCollectionSelectionNodes.h
+++ b/Engine/Plugins/Experimental/GeometryCollectionPlugin/Source/GeometryCollectionNodes/Public/Dataflow/GeometryCollectionSelectionNodes.h
@@ -298,11 +298,11 @@ public:
 	UPROPERTY(meta = (DataflowInput, DataflowOutput, DataflowPassthrough = "Collection", DataflowIntrinsic))
 	FManagedArrayCollection Collection;
 
-	/** Space separated list of bone indicies to specify the selection */
-	UPROPERTY(EditAnywhere, Category = "Selection")
+	/** Space separated list of bone indices to specify the selection */
+	UPROPERTY(EditAnywhere, Category = "Selection", meta=(DisplayName="Bone Indices"))
 	FString BoneIndicies = FString();
 
-	/** Array of the selected bone indicies */
+	/** Array of the selected bone indices */
 	UPROPERTY(meta = (DataflowOutput, DisplayName = "TransformSelection"))
 	FDataflowTransformSelection TransformSelection;
 
@@ -319,6 +319,44 @@ public:
 
 };
 
+/**
+ *
+ * Convert index array to a transform selection
+ *
+ */
+USTRUCT(meta = (DataflowGeometryCollection))
+struct FCollectionTransformSelectionFromIndexArrayDataflowNode : public FDataflowNode
+{
+	GENERATED_USTRUCT_BODY()
+	DATAFLOW_NODE_DEFINE_INTERNAL(FCollectionTransformSelectionFromIndexArrayDataflowNode, "CollectionTransformSelectionFromIndexArray", "GeometryCollection|Selection|Array", "")
+
+public:
+
+	/** Collection to use for the selection. Note only valid bone indices for the collection will be included in the output selection. */
+	UPROPERTY(meta = (DataflowInput, DataflowOutput, DataflowPassthrough = "Collection", DataflowIntrinsic))
+	FManagedArrayCollection Collection;
+
+	/** Array of bone indices to convert to a trannsform selection */
+	UPROPERTY(EditAnywhere, Category = "Selection", meta = (DataflowInput))
+	TArray<int32> BoneIndices;
+
+	/** Array of the selected bone indices */
+	UPROPERTY(meta = (DataflowOutput, DisplayName = "TransformSelection"))
+	FDataflowTransformSelection TransformSelection;
+
+	FCollectionTransformSelectionFromIndexArrayDataflowNode(const Dataflow::FNodeParameters& InParam, FGuid InGuid = FGuid::NewGuid())
+		: FDataflowNode(InParam, InGuid)
+	{
+		RegisterInputConnection(&Collection);
+		RegisterInputConnection(&BoneIndices);
+		RegisterOutputConnection(&TransformSelection);
+		RegisterOutputConnection(&Collection, &Collection);
+	}
+
+	virtual void Evaluate(Dataflow::FContext& Context, const FDataflowOutput* Out) const override;
+
+};
+
 
 /**
  *
@@ -654,6 +692,57 @@ enum class ERangeSettingEnum : uint8
 	Dataflow_Max                UMETA(Hidden)
 };
 
+
+/**
+ *
+ * Selects indices of a float array by range
+ *
+ */
+USTRUCT(meta = (DataflowGeometryCollection))
+struct FSelectFloatArrayIndicesInRangeDataflowNode : public FDataflowNode
+{
+	GENERATED_USTRUCT_BODY()
+	DATAFLOW_NODE_DEFINE_INTERNAL(FSelectFloatArrayIndicesInRangeDataflowNode, "SelectFloatArrayIndicesInRange", "GeometryCollection|Selection|Array", "")
+
+public:
+	/** GeometryCollection for the selection */
+	UPROPERTY(meta = (DataflowInput))
+	TArray<float> Values;
+
+	/** Minimum value for the selection */
+	UPROPERTY(EditAnywhere, Category = "Attribute", meta = (UIMin = 0.f, UIMax = 1000000000.f))
+	float Min = 0.f;
+
+	/** Maximum value for the selection */
+	UPROPERTY(EditAnywhere, Category = "Attribute", meta = (UIMin = 0.f, UIMax = 1000000000.f))
+	float Max = 1000.f;
+
+	/** Values for the selection has to be inside or outside [Min, Max] range */
+	UPROPERTY(EditAnywhere, Category = "Attribute")
+	ERangeSettingEnum RangeSetting = ERangeSettingEnum::Dataflow_RangeSetting_InsideRange;
+
+	/** If true then range includes Min and Max values */
+	UPROPERTY(EditAnywhere, Category = "Attribute")
+	bool bInclusive = true;
+
+	/** Indices of float Values matching the specified range */
+	UPROPERTY(meta = (DataflowOutput))
+	TArray<int> Indices;
+
+	FSelectFloatArrayIndicesInRangeDataflowNode(const Dataflow::FNodeParameters& InParam, FGuid InGuid = FGuid::NewGuid())
+		: FDataflowNode(InParam, InGuid)
+	{
+		RegisterInputConnection(&Values);
+		RegisterInputConnection(&Min);
+		RegisterInputConnection(&Max);
+		RegisterOutputConnection(&Indices);
+	}
+
+	virtual void Evaluate(Dataflow::FContext& Context, const FDataflowOutput* Out) const override;
+
+};
+
+
 /**
  *
  * Selects pieces based on their size
