Add option to drive adding more negative space samples when the current samples don't cover the space identified by [REDACTED_PERSON] (and are farther apart than the target Sample Spacing), to make the negative space coverage algorithm better able to protect space in challenging cases where a smaller fixed number of samples would fail.

diff --git a/Engine/Plugins/Experimental/GeometryCollectionPlugin/Source/GeometryCollectionNodes/Private/Dataflow/GeometryCollectionUtilityNodes.cpp b/Engine/Plugins/Experimental/GeometryCollectionPlugin/Source/GeometryCollectionNodes/Private/Dataflow/GeometryCollectionUtilityNodes.cpp
--- a/Engine/Plugins/Experimental/GeometryCollectionPlugin/Source/GeometryCollectionNodes/Private/Dataflow/GeometryCollectionUtilityNodes.cpp
+++ b/Engine/Plugins/Experimental/GeometryCollectionPlugin/Source/GeometryCollectionNodes/Private/Dataflow/GeometryCollectionUtilityNodes.cpp
@@ -251,6 +251,7 @@ void FGenerateClusterConvexHullsFromLeafHullsDataflowNode::Evaluate(Dataflow::FC
 				NegativeSpaceSettings.ReduceRadiusMargin = GetValue(Context, &NegativeSpaceTolerance);
 				NegativeSpaceSettings.MinSpacing = GetValue(Context, &MinSampleSpacing);
 				NegativeSpaceSettings.SampleMethod = ConvertNegativeSpaceSampleMethodDataflowEnum(SampleMethod);
+				NegativeSpaceSettings.bRequireSearchSampleCoverage = bRequireSearchSampleCoverage;
 				NegativeSpaceSettings.Sanitize();
 				bHasNegativeSpace = UE::FractureEngine::Convex::ComputeConvexHullsNegativeSpace(*GeomCollection, NegativeSpace, NegativeSpaceSettings, bHasSelectionFilter, SelectionArray);
 			}
@@ -335,6 +336,7 @@ void FGenerateClusterConvexHullsFromChildrenHullsDataflowNode::Evaluate(Dataflow
 				NegativeSpaceSettings.ReduceRadiusMargin = GetValue(Context, &NegativeSpaceTolerance);
 				NegativeSpaceSettings.MinSpacing = GetValue(Context, &MinSampleSpacing);
 				NegativeSpaceSettings.SampleMethod = ConvertNegativeSpaceSampleMethodDataflowEnum(SampleMethod);
+				NegativeSpaceSettings.bRequireSearchSampleCoverage = bRequireSearchSampleCoverage;
 				NegativeSpaceSettings.Sanitize();
 				bHasNegativeSpace = UE::FractureEngine::Convex::ComputeConvexHullsNegativeSpace(*GeomCollection, NegativeSpace, NegativeSpaceSettings, bHasSelectionFilter, SelectionArray);
 			}
@@ -418,6 +420,7 @@ void FMergeConvexHullsDataflowNode::Evaluate(Dataflow::FContext& Context, const
 			NegativeSpaceSettings.ReduceRadiusMargin = GetValue(Context, &NegativeSpaceTolerance);
 			NegativeSpaceSettings.MinSpacing = GetValue(Context, &MinSampleSpacing);
 			NegativeSpaceSettings.SampleMethod = ConvertNegativeSpaceSampleMethodDataflowEnum(SampleMethod);
+			NegativeSpaceSettings.bRequireSearchSampleCoverage = bRequireSearchSampleCoverage;
 			NegativeSpaceSettings.Sanitize();
 		}
 		if (bInProtectNegativeSpace && !bComputeNegativeSpacePerBone)
diff --git a/Engine/Plugins/Experimental/GeometryCollectionPlugin/Source/GeometryCollectionNodes/Public/Dataflow/GeometryCollectionUtilityNodes.h b/Engine/Plugins/Experimental/GeometryCollectionPlugin/Source/GeometryCollectionNodes/Public/Dataflow/GeometryCollectionUtilityNodes.h
--- a/Engine/Plugins/Experimental/GeometryCollectionPlugin/Source/GeometryCollectionNodes/Public/Dataflow/GeometryCollectionUtilityNodes.h
+++ b/Engine/Plugins/Experimental/GeometryCollectionPlugin/Source/GeometryCollectionNodes/Public/Dataflow/GeometryCollectionUtilityNodes.h
@@ -298,6 +298,10 @@ public:
 	UPROPERTY(EditAnywhere, Category = NegativeSpace, meta = (EditCondition = "bProtectNegativeSpace", EditConditionHides))
 	ENegativeSpaceSampleMethodDataflowEnum SampleMethod = ENegativeSpaceSampleMethodDataflowEnum::Uniform;
 
+	/** Whether to require that all candidate locations identified by Voxel Search are covered by negative space samples, up to the specified Min Sample Spacing. Only applies to Voxel Search. */
+	UPROPERTY(EditAnywhere, Category = NegativeSpace, meta = (EditCondition = "bProtectNegativeSpace && SampleMethod == ENegativeSpaceSampleMethodDataflowEnum::VoxelSearch", EditConditionHides))
+	bool bRequireSearchSampleCoverage = false;
+
 	/** Approximate number of spheres to consider when covering negative space */
 	UPROPERTY(EditAnywhere, Category = NegativeSpace, meta = (DataflowInput, ClampMin = 1, EditCondition = "bProtectNegativeSpace", EditConditionHides))
 	int32 TargetNumSamples = 50;
@@ -367,6 +371,10 @@ public:
 	UPROPERTY(EditAnywhere, Category = NegativeSpace, meta = (EditCondition = "bProtectNegativeSpace", EditConditionHides))
 	ENegativeSpaceSampleMethodDataflowEnum SampleMethod = ENegativeSpaceSampleMethodDataflowEnum::Uniform;
 
+	/** Whether to require that all candidate locations identified by Voxel Search are covered by negative space samples, up to the specified Min Sample Spacing. Only applies to Voxel Search. */
+	UPROPERTY(EditAnywhere, Category = NegativeSpace, meta = (EditCondition = "bProtectNegativeSpace && SampleMethod == ENegativeSpaceSampleMethodDataflowEnum::VoxelSearch", EditConditionHides))
+	bool bRequireSearchSampleCoverage = false;
+
 	/** Approximate number of spheres to consider when covering negative space */
 	UPROPERTY(EditAnywhere, Category = NegativeSpace, meta = (DataflowInput, ClampMin = 1, EditCondition = "bProtectNegativeSpace", EditConditionHides))
 	int32 TargetNumSamples = 50;
@@ -433,6 +441,10 @@ public:
 	UPROPERTY(EditAnywhere, Category = NegativeSpace, meta = (EditCondition = "bProtectNegativeSpace", EditConditionHides))
 	ENegativeSpaceSampleMethodDataflowEnum SampleMethod = ENegativeSpaceSampleMethodDataflowEnum::Uniform;
 
+	/** Whether to require that all candidate locations identified by Voxel Search are covered by negative space samples, up to the specified Min Sample Spacing. Only applies to Voxel Search. */
+	UPROPERTY(EditAnywhere, Category = NegativeSpace, meta = (EditCondition = "bProtectNegativeSpace && SampleMethod == ENegativeSpaceSampleMethodDataflowEnum::VoxelSearch", EditConditionHides))
+	bool bRequireSearchSampleCoverage = false;
+
 	/** Approximate number of spheres to consider when covering negative space */
 	UPROPERTY(EditAnywhere, Category = NegativeSpace, meta = (DataflowInput, ClampMin = 1, EditCondition = "bProtectNegativeSpace", EditConditionHides))
 	int32 TargetNumSamples = 50;
diff --git a/Engine/Source/Runtime/GeometryCore/Private/CompGeom/ConvexDecomposition3.cpp b/Engine/Source/Runtime/GeometryCore/Private/CompGeom/ConvexDecomposition3.cpp
--- a/Engine/Source/Runtime/GeometryCore/Private/CompGeom/ConvexDecomposition3.cpp
+++ b/Engine/Source/Runtime/GeometryCore/Private/CompGeom/ConvexDecomposition3.cpp
@@ -265,6 +265,31 @@ bool FSphereCovering::AddNegativeSpace(const TFastWindingTree<FDynamicMesh3>& Sp
 				int32 VID = Ordering.Order[SampleIdx];
 				AddSample(MorphologyMesh.GetVertex(VID));
 			}
+
+			if (SampleSettings.bRequireSearchSampleCoverage)
+			{
+				double SpacingThresholdSq = SampleSettings.MinSpacing * SampleSettings.MinSpacing;
+				for (int32 SampleIdx = NumSamples; SampleIdx < Ordering.Order.Num(); ++SampleIdx)
+				{
+					int32 VID = Ordering.Order[SampleIdx];
+					FVector3d Pos = MorphologyMesh.GetVertex(VID);
+					// TODO: Consider accelerating this coverage search w/ e.g. a sparse dynamic octree representation of the sphere covering
+					bool bFoundCover = false;
+					for (int32 SphereIdx = 0; SphereIdx < Position.Num(); ++SphereIdx)
+					{
+						double ThreshSq = FMath::Max(SpacingThresholdSq, Radius[SphereIdx] * Radius[SphereIdx]);
+						if (FVector3d::DistSquared(Position[SphereIdx], Pos) < ThreshSq)
+						{
+							bFoundCover = true;
+							break;
+						}
+					}
+					if (!bFoundCover)
+					{
+						AddSample(Pos);
+					}
+				}
+			}
 		}
 		else
 		{
diff --git a/Engine/Source/Runtime/GeometryCore/Public/CompGeom/ConvexDecomposition3.h b/Engine/Source/Runtime/GeometryCore/Public/CompGeom/ConvexDecomposition3.h
--- a/Engine/Source/Runtime/GeometryCore/Public/CompGeom/ConvexDecomposition3.h
+++ b/Engine/Source/Runtime/GeometryCore/Public/CompGeom/ConvexDecomposition3.h
@@ -63,6 +63,11 @@ struct FNegativeSpaceSampleSettings
 	// Ignore spheres with smaller radius than this
 	double MinRadius = 10.0;
 
+	// Currently only applies to VoxelSearch.
+	// Whether to require that all candidate sample locations identified by Voxel Search are covered by negative space samples, up to the specified Min Sample Spacing.
+	// Note: This takes priority over TargetNumSamples if the TargetNumSamples did not achieve the required coverage.
+	bool bRequireSearchSampleCoverage = false;
+
 	// Make sure the settings values are in valid ranges
 	void Sanitize()
 	{
