Add dataflow node to get num spheres in a sphere covering

diff --git a/Engine/Plugins/Experimental/GeometryCollectionPlugin/Source/GeometryCollectionNodes/Private/Dataflow/GeometryCollectionDebugNodes.cpp b/Engine/Plugins/Experimental/GeometryCollectionPlugin/Source/GeometryCollectionNodes/Private/Dataflow/GeometryCollectionDebugNodes.cpp
--- a/Engine/Plugins/Experimental/GeometryCollectionPlugin/Source/GeometryCollectionNodes/Private/Dataflow/GeometryCollectionDebugNodes.cpp
+++ b/Engine/Plugins/Experimental/GeometryCollectionPlugin/Source/GeometryCollectionNodes/Private/Dataflow/GeometryCollectionDebugNodes.cpp
@@ -23,10 +23,13 @@ namespace Dataflow
 		DATAFLOW_NODE_REGISTER_CREATION_FACTORY(FConvexHullToMeshDataflowNode);
 		DATAFLOW_NODE_REGISTER_CREATION_FACTORY(FSphereCoveringToMeshDataflowNode);
 		DATAFLOW_NODE_REGISTER_CREATION_FACTORY(FMeshToOBJStringDebugDataflowNode);
+		DATAFLOW_NODE_REGISTER_CREATION_FACTORY(FSphereCoveringCountSpheresNode);
 		DATAFLOW_NODE_REGISTER_CREATION_FACTORY(FWriteStringToFile);
 
+
 		DATAFLOW_NODE_REGISTER_CREATION_FACTORY_NODE_COLORS_BY_CATEGORY("GeometryCollection|Development", FLinearColor(1.f, 0.f, 0.f), CDefaultNodeBodyTintColor);
 		DATAFLOW_NODE_REGISTER_CREATION_FACTORY_NODE_COLORS_BY_CATEGORY("GeometryCollection|Mesh", FLinearColor(1.f, 0.16f, 0.05f), CDefaultNodeBodyTintColor);
+		DATAFLOW_NODE_REGISTER_CREATION_FACTORY_NODE_COLORS_BY_CATEGORY("SphereCovering", FLinearColor(1.f, 0.16f, 0.05f), CDefaultNodeBodyTintColor);
 	}
 }
 
diff --git a/Engine/Plugins/Experimental/GeometryCollectionPlugin/Source/GeometryCollectionNodes/Public/Dataflow/GeometryCollectionDebugNodes.h b/Engine/Plugins/Experimental/GeometryCollectionPlugin/Source/GeometryCollectionNodes/Public/Dataflow/GeometryCollectionDebugNodes.h
--- a/Engine/Plugins/Experimental/GeometryCollectionPlugin/Source/GeometryCollectionNodes/Public/Dataflow/GeometryCollectionDebugNodes.h
+++ b/Engine/Plugins/Experimental/GeometryCollectionPlugin/Source/GeometryCollectionNodes/Public/Dataflow/GeometryCollectionDebugNodes.h
@@ -72,6 +72,39 @@ public:
 	virtual void Evaluate(Dataflow::FContext& Context, const FDataflowOutput* Out) const override;
 };
 
+// Report the number of spheres in a sphere covering
+USTRUCT(meta = (DataflowGeometryCollection))
+struct FSphereCoveringCountSpheresNode : public FDataflowNode
+{
+	GENERATED_USTRUCT_BODY()
+	DATAFLOW_NODE_DEFINE_INTERNAL(FSphereCoveringCountSpheresNode, "Get Sphere Covering Sphere Count", "SphereCovering", "")
+
+public:
+	// The sphere covering to evaluate
+	UPROPERTY(meta = (DataflowInput, DataflowIntrinsic))
+	FDataflowSphereCovering SphereCovering;
+
+	// Number of spheres in the sphere covering
+	UPROPERTY(meta = (DataflowOutput))
+	int32 NumSpheres = 0;
+
+	FSphereCoveringCountSpheresNode(const Dataflow::FNodeParameters& InParam, FGuid InGuid = FGuid::NewGuid())
+		: FDataflowNode(InParam, InGuid)
+	{
+		RegisterInputConnection(&SphereCovering);
+		RegisterOutputConnection(&NumSpheres);
+	}
+
+	virtual void Evaluate(Dataflow::FContext& Context, const FDataflowOutput* Out) const override
+	{
+		if (Out->IsA(&NumSpheres))
+		{
+			const FDataflowSphereCovering& InSphereCovering = GetValue(Context, &SphereCovering);
+			SetValue(Context, InSphereCovering.Spheres.Num(), &NumSpheres);
+		}
+	}
+};
+
 // Convert a mesh to a string formatted as an OBJ file, which can be read by many external mesh viewers
 USTRUCT(meta = (DataflowGeometryCollection))
 struct FMeshToOBJStringDebugDataflowNode: public FDataflowNode
