Dataflow:
- Added OneMinus, BranchFloat, BranchInt nodes
- Renamed Branch node to BranchMesh
- bug fixes

diff --git a/Engine/Plugins/Experimental/GeometryCollectionPlugin/Source/GeometryCollectionNodes/Private/Dataflow/GeometryCollectionMathNodes.cpp b/Engine/Plugins/Experimental/GeometryCollectionPlugin/Source/GeometryCollectionNodes/Private/Dataflow/GeometryCollectionMathNodes.cpp
--- a/Engine/Plugins/Experimental/GeometryCollectionPlugin/Source/GeometryCollectionNodes/Private/Dataflow/GeometryCollectionMathNodes.cpp
+++ b/Engine/Plugins/Experimental/GeometryCollectionPlugin/Source/GeometryCollectionNodes/Private/Dataflow/GeometryCollectionMathNodes.cpp
@@ -67,6 +67,7 @@ namespace Dataflow
 		DATAFLOW_NODE_REGISTER_CREATION_FACTORY(FRadiansToDegreesDataflowNode);
 		DATAFLOW_NODE_REGISTER_CREATION_FACTORY(FDegreesToRadiansDataflowNode);
 		DATAFLOW_NODE_REGISTER_CREATION_FACTORY(FMathConstantsDataflowNode);
+		DATAFLOW_NODE_REGISTER_CREATION_FACTORY(FOneMinusDataflowNode);
 
 		// Math
 		DATAFLOW_NODE_REGISTER_CREATION_FACTORY_NODE_COLORS_BY_CATEGORY("Math", FLinearColor(0.f, 0.4f, 0.8f), CDefaultNodeBodyTintColor);
@@ -846,6 +847,16 @@ void FMathConstantsDataflowNode::Evaluate(Dataflow::FContext& Context, const FDa
 	}
 }
 
+void FOneMinusDataflowNode::Evaluate(Dataflow::FContext& Context, const FDataflowOutput* Out) const
+{
+	if (Out->IsA<float>(&ReturnValue))
+	{
+		const float InA = GetValue<float>(Context, &A, A);
+
+		const float Result = 1.f - InA;
+		SetValue(Context, Result, &ReturnValue);
+	}
+}
 
 
 
diff --git a/Engine/Plugins/Experimental/GeometryCollectionPlugin/Source/GeometryCollectionNodes/Private/Dataflow/GeometryCollectionNodes.cpp b/Engine/Plugins/Experimental/GeometryCollectionPlugin/Source/GeometryCollectionNodes/Private/Dataflow/GeometryCollectionNodes.cpp
--- a/Engine/Plugins/Experimental/GeometryCollectionPlugin/Source/GeometryCollectionNodes/Private/Dataflow/GeometryCollectionNodes.cpp
+++ b/Engine/Plugins/Experimental/GeometryCollectionPlugin/Source/GeometryCollectionNodes/Private/Dataflow/GeometryCollectionNodes.cpp
@@ -66,7 +66,7 @@ namespace Dataflow
 		DATAFLOW_NODE_REGISTER_CREATION_FACTORY(FTransformMeshDataflowNode);
 		DATAFLOW_NODE_REGISTER_CREATION_FACTORY(FCompareIntDataflowNode);
 		DATAFLOW_NODE_REGISTER_CREATION_FACTORY(FCompareFloatDataflowNode);
-		DATAFLOW_NODE_REGISTER_CREATION_FACTORY(FBranchDataflowNode);
+		DATAFLOW_NODE_REGISTER_CREATION_FACTORY(FBranchMeshDataflowNode);
 		DATAFLOW_NODE_REGISTER_CREATION_FACTORY(FBranchCollectionDataflowNode);
 		DATAFLOW_NODE_REGISTER_CREATION_FACTORY(FGetSchemaDataflowNode);
 		DATAFLOW_NODE_REGISTER_CREATION_FACTORY(FRemoveOnBreakDataflowNode);
@@ -82,6 +82,8 @@ namespace Dataflow
 		DATAFLOW_NODE_REGISTER_CREATION_FACTORY(FMultiplyTransformDataflowNode);
 		DATAFLOW_NODE_REGISTER_CREATION_FACTORY(FInvertTransformDataflowNode);
 		DATAFLOW_NODE_REGISTER_CREATION_FACTORY(FSelectionToVertexListDataflowNode);
+		DATAFLOW_NODE_REGISTER_CREATION_FACTORY(FBranchFloatDataflowNode);
+		DATAFLOW_NODE_REGISTER_CREATION_FACTORY(FBranchIntDataflowNode);
 
 		// GeometryCollection
 		DATAFLOW_NODE_REGISTER_CREATION_FACTORY_NODE_COLORS_BY_CATEGORY("GeometryCollection", FLinearColor(0.55f, 0.45f, 1.0f), CDefaultNodeBodyTintColor);
@@ -494,7 +496,7 @@ void FCompareFloatDataflowNode::Evaluate(Dataflow::FContext& Context, const FDat
 }
 
 
-void FBranchDataflowNode::Evaluate(Dataflow::FContext& Context, const FDataflowOutput* Out) const
+void FBranchMeshDataflowNode::Evaluate(Dataflow::FContext& Context, const FDataflowOutput* Out) const
 {
 	if (Out->IsA<TObjectPtr<UDynamicMesh>>(&Mesh))
 	{
@@ -1146,3 +1148,45 @@ void FInvertTransformDataflowNode::Evaluate(Dataflow::FContext& Context, const F
 	}
 }
 
+void FBranchFloatDataflowNode::Evaluate(Dataflow::FContext& Context, const FDataflowOutput* Out) const
+{
+	if (Out->IsA<float>(&ReturnValue))
+	{
+		bool InCondition = GetValue<bool>(Context, &bCondition);
+
+		if (InCondition)
+		{
+			const float InA = GetValue<float>(Context, &A, A);
+
+			SetValue(Context, InA, &ReturnValue);
+		}
+		else
+		{
+			const float InB = GetValue<float>(Context, &B, B);
+
+			SetValue(Context, InB, &ReturnValue);
+		}
+	}
+}
+
+void FBranchIntDataflowNode::Evaluate(Dataflow::FContext& Context, const FDataflowOutput* Out) const
+{
+	if (Out->IsA<int32>(&ReturnValue))
+	{
+		bool InCondition = GetValue<bool>(Context, &bCondition);
+
+		if (InCondition)
+		{
+			const int32 InA = GetValue<int32>(Context, &A, A);
+
+			SetValue(Context, InA, &ReturnValue);
+		}
+		else
+		{
+			const int32 InB = GetValue<int32>(Context, &B, B);
+
+			SetValue(Context, InB, &ReturnValue);
+		}
+	}
+}
+
diff --git a/Engine/Plugins/Experimental/GeometryCollectionPlugin/Source/GeometryCollectionNodes/Public/Dataflow/GeometryCollectionMathNodes.h b/Engine/Plugins/Experimental/GeometryCollectionPlugin/Source/GeometryCollectionNodes/Public/Dataflow/GeometryCollectionMathNodes.h
--- a/Engine/Plugins/Experimental/GeometryCollectionPlugin/Source/GeometryCollectionNodes/Public/Dataflow/GeometryCollectionMathNodes.h
+++ b/Engine/Plugins/Experimental/GeometryCollectionPlugin/Source/GeometryCollectionNodes/Public/Dataflow/GeometryCollectionMathNodes.h
@@ -903,7 +903,7 @@ public:
 
 /**
 *
-* Linearly interpolates between A and B based on Alpha. (100% of A when Alpha = 0 and 100% of B when Alpha = 100)
+* Linearly interpolates between A and B based on Alpha. (100% of A when Alpha = 0.0 and 100% of B when Alpha = 1.0)
 *
 */
 USTRUCT(meta = (DataflowGeometryCollection))
@@ -1717,61 +1717,33 @@ public:
 
 };
 
+/**
+*
+* Returns 1.0 - A
+*
+*/
+USTRUCT(meta = (DataflowGeometryCollection))
+struct FOneMinusDataflowNode : public FDataflowNode
+{
+	GENERATED_USTRUCT_BODY()
+	DATAFLOW_NODE_DEFINE_INTERNAL(FOneMinusDataflowNode, "OneMinus", "Math|Float", "")
 
+public:
+	UPROPERTY(EditAnywhere, Category = "OneMinus", meta = (DataflowInput));
+	float A = 0.f;
 
+	UPROPERTY(meta = (DataflowOutput))
+	float ReturnValue = 0.f;
 
+	FOneMinusDataflowNode(const Dataflow::FNodeParameters& InParam, FGuid InGuid = FGuid::NewGuid())
+		: FDataflowNode(InParam, InGuid)
+	{
+		RegisterInputConnection(&A);
+		RegisterOutputConnection(&ReturnValue);
+	}
 
-
-
-
-
-
-
-
-
-
-
-
-
-
-
-
-
-
-
-
-
-
-
-
-
-
-
-
-
-
-
-
-
-
-
-
-
-
-
-
-
-
-
-
-
-
-
-
-
-
-
-
+	virtual void Evaluate(Dataflow::FContext& Context, const FDataflowOutput* Out) const override;
+};
 
 
 namespace Dataflow
diff --git a/Engine/Plugins/Experimental/GeometryCollectionPlugin/Source/GeometryCollectionNodes/Public/Dataflow/GeometryCollectionNodes.h b/Engine/Plugins/Experimental/GeometryCollectionPlugin/Source/GeometryCollectionNodes/Public/Dataflow/GeometryCollectionNodes.h
--- a/Engine/Plugins/Experimental/GeometryCollectionPlugin/Source/GeometryCollectionNodes/Public/Dataflow/GeometryCollectionNodes.h
+++ b/Engine/Plugins/Experimental/GeometryCollectionPlugin/Source/GeometryCollectionNodes/Public/Dataflow/GeometryCollectionNodes.h
@@ -819,14 +819,14 @@ public:
 
 /**
  *
- * Branch between two inputs based on boolean condition
+ * Branch between two mesh inputs based on boolean condition
  *
  */
 USTRUCT()
-struct FBranchDataflowNode : public FDataflowNode
+struct FBranchMeshDataflowNode : public FDataflowNode
 {
 	GENERATED_USTRUCT_BODY()
-	DATAFLOW_NODE_DEFINE_INTERNAL(FBranchDataflowNode, "Branch", "Utilities|FlowControl", "")
+	DATAFLOW_NODE_DEFINE_INTERNAL(FBranchMeshDataflowNode, "BranchMesh", "Utilities|FlowControl", "")
 
 public:
 	/** Mesh input */
@@ -845,7 +845,7 @@ public:
 	UPROPERTY(meta = (DataflowOutput))
 	TObjectPtr<UDynamicMesh> Mesh;
 
-	FBranchDataflowNode(const Dataflow::FNodeParameters& InParam, FGuid InGuid = FGuid::NewGuid())
+	FBranchMeshDataflowNode(const Dataflow::FNodeParameters& InParam, FGuid InGuid = FGuid::NewGuid())
 		: FDataflowNode(InParam, InGuid)
 	{
 		RegisterInputConnection(&MeshA);
@@ -1615,6 +1615,88 @@ public:
 	virtual void Evaluate(Dataflow::FContext& Context, const FDataflowOutput* Out) const override;
 };
 
+/**
+ *
+ * Branch between two float inputs based on boolean condition
+ *
+ */
+USTRUCT()
+struct FBranchFloatDataflowNode : public FDataflowNode
+{
+	GENERATED_USTRUCT_BODY()
+	DATAFLOW_NODE_DEFINE_INTERNAL(FBranchFloatDataflowNode, "BranchFloat", "Utilities|FlowControl", "")
+
+public:
+	/** Float input */
+	UPROPERTY(EditAnywhere, Category = "Branch", meta = (DataflowInput));
+	float A = 0.f;
+
+	/** Float input */
+	UPROPERTY(EditAnywhere, Category = "Branch", meta = (DataflowInput));
+	float B = 0.f;
+
+	/** If true, Output = A, otherwise Output = B */
+	UPROPERTY(EditAnywhere, Category = "Branch");
+	bool bCondition = false;
+
+	/** Output */
+	UPROPERTY(meta = (DataflowOutput))
+	float ReturnValue = 0.f;
+
+	FBranchFloatDataflowNode(const Dataflow::FNodeParameters& InParam, FGuid InGuid = FGuid::NewGuid())
+		: FDataflowNode(InParam, InGuid)
+	{
+		RegisterInputConnection(&A);
+		RegisterInputConnection(&B);
+		RegisterInputConnection(&bCondition);
+		RegisterOutputConnection(&ReturnValue);
+	}
+
+	virtual void Evaluate(Dataflow::FContext& Context, const FDataflowOutput* Out) const override;
+
+};
+
+/**
+ *
+ * Branch between two int inputs based on boolean condition
+ *
+ */
+USTRUCT()
+struct FBranchIntDataflowNode : public FDataflowNode
+{
+	GENERATED_USTRUCT_BODY()
+	DATAFLOW_NODE_DEFINE_INTERNAL(FBranchIntDataflowNode, "BranchInt", "Utilities|FlowControl", "")
+
+public:
+	/** Int input */
+	UPROPERTY(EditAnywhere, Category = "Branch", meta = (DataflowInput));
+	int32 A = 0;
+
+	/** Int input */
+	UPROPERTY(EditAnywhere, Category = "Branch", meta = (DataflowInput));
+	int32 B = 0;
+
+	/** If true, Output = A, otherwise Output = B */
+	UPROPERTY(EditAnywhere, Category = "Branch");
+	bool bCondition = false;
+
+	/** Output */
+	UPROPERTY(meta = (DataflowOutput))
+	int32 ReturnValue = 0;
+
+	FBranchIntDataflowNode(const Dataflow::FNodeParameters& InParam, FGuid InGuid = FGuid::NewGuid())
+		: FDataflowNode(InParam, InGuid)
+	{
+		RegisterInputConnection(&A);
+		RegisterInputConnection(&B);
+		RegisterInputConnection(&bCondition);
+		RegisterOutputConnection(&ReturnValue);
+	}
+
+	virtual void Evaluate(Dataflow::FContext& Context, const FDataflowOutput* Out) const override;
+
+};
+
 
 namespace Dataflow
 {
