Add a Geometry Collection dataflow node to create containing cluster nodes for collections w/ isolated root nodes w/ no other transforms.

diff --git a/Engine/Plugins/Experimental/GeometryCollectionPlugin/Source/GeometryCollectionNodes/Private/Dataflow/GeometryCollectionClusteringNodes.cpp b/Engine/Plugins/Experimental/GeometryCollectionPlugin/Source/GeometryCollectionNodes/Private/Dataflow/GeometryCollectionClusteringNodes.cpp
--- a/Engine/Plugins/Experimental/GeometryCollectionPlugin/Source/GeometryCollectionNodes/Private/Dataflow/GeometryCollectionClusteringNodes.cpp
+++ b/Engine/Plugins/Experimental/GeometryCollectionPlugin/Source/GeometryCollectionNodes/Private/Dataflow/GeometryCollectionClusteringNodes.cpp
@@ -44,6 +44,7 @@ namespace Dataflow
 		DATAFLOW_NODE_REGISTER_CREATION_FACTORY(FClusterUnclusterDataflowNode);
 		DATAFLOW_NODE_REGISTER_CREATION_FACTORY(FClusterDataflowNode);
 		DATAFLOW_NODE_REGISTER_CREATION_FACTORY(FClusterMergeDataflowNode);
+		DATAFLOW_NODE_REGISTER_CREATION_FACTORY(FClusterIsolatedRootsDataflowNode);
 
 		// GeometryCollection|Cluster
 		DATAFLOW_NODE_REGISTER_CREATION_FACTORY_NODE_COLORS_BY_CATEGORY("GeometryCollection|Cluster", FLinearColor(.25f, 0.45f, 0.8f), CDefaultNodeBodyTintColor);
@@ -181,3 +182,24 @@ void FClusterMergeDataflowNode::Evaluate(Dataflow::FContext& Context, const FDat
 	}
 }
 
+void FClusterIsolatedRootsDataflowNode::Evaluate(Dataflow::FContext& Context, const FDataflowOutput* Out) const
+{
+	if (Out->IsA(&Collection))
+	{
+		const FManagedArrayCollection& InCollection = GetValue(Context, &Collection);
+		int32 NumTransforms = InCollection.NumElements(FGeometryCollection::TransformGroup);
+		// Only if there is a single transform, re-parent it under a new transform
+		if (NumTransforms == 1)
+		{
+			if (TUniquePtr<FGeometryCollection> GeomCollection = TUniquePtr<FGeometryCollection>(InCollection.NewCopy<FGeometryCollection>()))
+			{
+				FGeometryCollectionClusteringUtility::ClusterAllBonesUnderNewRoot(GeomCollection.Get());
+				SetValue(Context, (const FManagedArrayCollection&)(*GeomCollection), &Collection);
+				return;
+			}
+		}
+		SetValue(Context, InCollection, &Collection);
+	}
+}
+
+
diff --git a/Engine/Plugins/Experimental/GeometryCollectionPlugin/Source/GeometryCollectionNodes/Public/Dataflow/GeometryCollectionClusteringNodes.h b/Engine/Plugins/Experimental/GeometryCollectionPlugin/Source/GeometryCollectionNodes/Public/Dataflow/GeometryCollectionClusteringNodes.h
--- a/Engine/Plugins/Experimental/GeometryCollectionPlugin/Source/GeometryCollectionNodes/Public/Dataflow/GeometryCollectionClusteringNodes.h
+++ b/Engine/Plugins/Experimental/GeometryCollectionPlugin/Source/GeometryCollectionNodes/Public/Dataflow/GeometryCollectionClusteringNodes.h
@@ -237,6 +237,31 @@ public:
 };
 
 
+/**
+ * Add a single cluster to the Geometry Collection if it only has a single transform with no clusters
+ */
+USTRUCT(meta = (DataflowGeometryCollection))
+struct FClusterIsolatedRootsDataflowNode : public FDataflowNode
+{
+	GENERATED_USTRUCT_BODY()
+	DATAFLOW_NODE_DEFINE_INTERNAL(FClusterIsolatedRootsDataflowNode, "ClusterIsolatedRoots", "GeometryCollection|Cluster", "")
+
+public:
+
+	/** Collection to modify */
+	UPROPERTY(meta = (DataflowInput, DataflowOutput, DataflowPassthrough = "Collection", DataflowIntrinsic))
+	FManagedArrayCollection Collection;
+
+	FClusterIsolatedRootsDataflowNode(const Dataflow::FNodeParameters& InParam, FGuid InGuid = FGuid::NewGuid())
+		: FDataflowNode(InParam, InGuid)
+	{
+		RegisterInputConnection(&Collection);
+		RegisterOutputConnection(&Collection, &Collection);
+	}
+
+	virtual void Evaluate(Dataflow::FContext& Context, const FDataflowOutput* Out) const override;
+
+};
 
 namespace Dataflow
 {
