Expose distance and contact thresholds of Proximity node as Dataflow connections

diff --git a/Engine/Plugins/Experimental/GeometryCollectionPlugin/Source/GeometryCollectionNodes/Private/Dataflow/GeometryCollectionNodes.cpp b/Engine/Plugins/Experimental/GeometryCollectionPlugin/Source/GeometryCollectionNodes/Private/Dataflow/GeometryCollectionNodes.cpp
--- a/Engine/Plugins/Experimental/GeometryCollectionPlugin/Source/GeometryCollectionNodes/Private/Dataflow/GeometryCollectionNodes.cpp
+++ b/Engine/Plugins/Experimental/GeometryCollectionPlugin/Source/GeometryCollectionNodes/Private/Dataflow/GeometryCollectionNodes.cpp
@@ -688,10 +688,10 @@ void FProximityDataflowNode::Evaluate(Dataflow::FContext& Context, const FDatafl
 
 			Properties.Method = (EProximityMethod)ProximityMethod;
 			Properties.ContactMethod = (EProximityContactMethod)FilterContactMethod;
-			Properties.DistanceThreshold = DistanceThreshold;
+			Properties.DistanceThreshold = GetValue(Context, &DistanceThreshold);
 			Properties.bUseAsConnectionGraph = bUseAsConnectionGraph;
 			Properties.ContactAreaMethod = (EConnectionContactMethod)ContactAreaMethod;
-			Properties.RequireContactAmount = ContactThreshold;
+			Properties.RequireContactAmount = GetValue(Context, &ContactThreshold);
 
 			GeomCollection->SetProximityProperties(Properties);
 
diff --git a/Engine/Plugins/Experimental/GeometryCollectionPlugin/Source/GeometryCollectionNodes/Public/Dataflow/GeometryCollectionNodes.h b/Engine/Plugins/Experimental/GeometryCollectionPlugin/Source/GeometryCollectionNodes/Public/Dataflow/GeometryCollectionNodes.h
--- a/Engine/Plugins/Experimental/GeometryCollectionPlugin/Source/GeometryCollectionNodes/Public/Dataflow/GeometryCollectionNodes.h
+++ b/Engine/Plugins/Experimental/GeometryCollectionPlugin/Source/GeometryCollectionNodes/Public/Dataflow/GeometryCollectionNodes.h
@@ -1089,12 +1089,12 @@ public:
 	EProximityMethodEnum ProximityMethod = EProximityMethodEnum::Dataflow_ProximityMethod_Precise;
 
 	/** If hull-based proximity detection is enabled, amount to expand hulls when searching for overlapping neighbors */
-	UPROPERTY(EditAnywhere, Category = "Proximity", meta = (ClampMin = "0", 
+	UPROPERTY(EditAnywhere, Category = "Proximity", meta = (DataflowInput, ClampMin = "0", Units = cm,
 		EditCondition = "ProximityMethod == EProximityMethodEnum::Dataflow_ProximityMethod_ConvexHull || FilterContactMethod == EProximityContactFilteringMethodEnum::Dataflow_ProximityContactFilteringMethod_ConvexHullSharp || FilterContactMethod == EProximityContactFilteringMethodEnum::Dataflow_ProximityContactFilteringMethod_ConvexHullArea || ContactAreaMethod = EConnectionContactAreaMethodEnum::Dataflow_ProximityContactFilteringMethod_ConvexHullArea"))
 	float DistanceThreshold = 1;
 
 	// If greater than zero, proximity will be additionally filtered by a 'contact' threshold, in cm, to exclude grazing / corner proximity
-	UPROPERTY(EditAnywhere, Category = "Proximity", meta = (ClampMin = "0"))
+	UPROPERTY(EditAnywhere, Category = "Proximity", meta = (DataflowInput, ClampMin = "0", Units = cm))
 	float ContactThreshold = 0;
 
 	/** How to use the Contact Threshold (if > 0) to filter out unwanted small or corner contacts from the proximity graph. If contact threshold is zero, no filtering is applied. */
@@ -1122,6 +1122,8 @@ public:
 		: FDataflowNode(InParam, InGuid)
 	{
 		RegisterInputConnection(&Collection);
+		RegisterInputConnection(&DistanceThreshold);
+		RegisterInputConnection(&ContactThreshold);
 		RegisterOutputConnection(&Collection, &Collection);
 	}
 
