Add a new Cluster Merge to Neighbors dataflow node, allowing selected bones to be independently merged to their nearest or largest neighboring bones/clusters

diff --git a/Engine/Plugins/Experimental/GeometryCollectionPlugin/Source/GeometryCollectionNodes/Private/Dataflow/GeometryCollectionClusteringNodes.cpp b/Engine/Plugins/Experimental/GeometryCollectionPlugin/Source/GeometryCollectionNodes/Private/Dataflow/GeometryCollectionClusteringNodes.cpp
--- a/Engine/Plugins/Experimental/GeometryCollectionPlugin/Source/GeometryCollectionNodes/Private/Dataflow/GeometryCollectionClusteringNodes.cpp
+++ b/Engine/Plugins/Experimental/GeometryCollectionPlugin/Source/GeometryCollectionNodes/Private/Dataflow/GeometryCollectionClusteringNodes.cpp
@@ -43,6 +43,7 @@ namespace Dataflow
 		DATAFLOW_NODE_REGISTER_CREATION_FACTORY(FClusterFlattenDataflowNode);
 		DATAFLOW_NODE_REGISTER_CREATION_FACTORY(FClusterUnclusterDataflowNode);
 		DATAFLOW_NODE_REGISTER_CREATION_FACTORY(FClusterDataflowNode);
+		DATAFLOW_NODE_REGISTER_CREATION_FACTORY(FClusterMergeToNeighborsDataflowNode);
 		DATAFLOW_NODE_REGISTER_CREATION_FACTORY(FClusterMergeDataflowNode);
 		DATAFLOW_NODE_REGISTER_CREATION_FACTORY(FClusterIsolatedRootsDataflowNode);
 
@@ -167,6 +168,50 @@ void FClusterDataflowNode::Evaluate(Dataflow::FContext& Context, const FDataflow
 	}
 }
 
+void FClusterMergeToNeighborsDataflowNode::Evaluate(Dataflow::FContext& Context, const FDataflowOutput* Out) const
+{
+	if (Out->IsA(&Collection))
+	{
+		const FManagedArrayCollection& InCollection = GetValue(Context, &Collection);
+		const FDataflowTransformSelection& InTransformSelection = GetValue(Context, &TransformSelection);
+		if (InTransformSelection.NumSelected() == 0)
+		{
+			SetValue(Context, InCollection, &Collection);
+			return;
+		}
+		if (TUniquePtr<FGeometryCollection> GeomCollection = TUniquePtr<FGeometryCollection>(InCollection.NewCopy<FGeometryCollection>()))
+		{
+			double InMinVolumeCubeRoot = (double)GetValue(Context, &MinVolumeCubeRoot);
+			double InMinVolume = InMinVolumeCubeRoot * InMinVolumeCubeRoot * InMinVolumeCubeRoot;
+			bool bInOnlyToConnected = GetValue(Context, &bOnlyToConnected);
+			bool bInOnlySameParent = GetValue(Context, &bOnlySameParent);
+			UE::PlanarCut::ENeighborSelectionMethod InNeighborSelectionMethod =
+				NeighborSelectionMethod == EClusterNeighborSelectionMethodEnum::Dataflow_ClusterNeighborSelectionMethod_LargestNeighbor ?
+				UE::PlanarCut::ENeighborSelectionMethod::LargestNeighbor : UE::PlanarCut::ENeighborSelectionMethod::NearestCenter;
+
+			TArray<int32> Selection = InTransformSelection.AsArray();
+			TArray<double> Volumes;
+			FindBoneVolumes(
+				*GeomCollection,
+				TArrayView<const int32>(),
+				Volumes, 1.0, true);
+			MergeClusters(
+				*GeomCollection,
+				Volumes,
+				InMinVolume,
+				Selection,
+				InNeighborSelectionMethod,
+				bInOnlyToConnected,
+				bInOnlySameParent
+			);
+			SetValue(Context, (const FManagedArrayCollection&)(*GeomCollection), &Collection);
+			return;
+		}
+
+		SetValue(Context, InCollection, &Collection);
+	}
+}
+
 void FClusterMergeDataflowNode::Evaluate(Dataflow::FContext& Context, const FDataflowOutput* Out) const
 {
 	if (Out->IsA(&Collection))
diff --git a/Engine/Plugins/Experimental/GeometryCollectionPlugin/Source/GeometryCollectionNodes/Public/Dataflow/GeometryCollectionClusteringNodes.h b/Engine/Plugins/Experimental/GeometryCollectionPlugin/Source/GeometryCollectionNodes/Public/Dataflow/GeometryCollectionClusteringNodes.h
--- a/Engine/Plugins/Experimental/GeometryCollectionPlugin/Source/GeometryCollectionNodes/Public/Dataflow/GeometryCollectionClusteringNodes.h
+++ b/Engine/Plugins/Experimental/GeometryCollectionPlugin/Source/GeometryCollectionNodes/Public/Dataflow/GeometryCollectionClusteringNodes.h
@@ -205,6 +205,66 @@ public:
 };
 
 
+
+UENUM(BlueprintType)
+enum class EClusterNeighborSelectionMethodEnum : uint8
+{
+	Dataflow_ClusterNeighborSelectionMethod_LargestNeighbor UMETA(DisplayName = "Largest Neighbor"),
+	Dataflow_ClusterNeighborSelectionMethod_NearestCenter UMETA(DisplayName = "Nearest Center")
+};
+
+/**
+ * Merge selected bones to their neighbors
+ */
+USTRUCT(meta = (DataflowGeometryCollection))
+struct FClusterMergeToNeighborsDataflowNode : public FDataflowNode
+{
+	GENERATED_USTRUCT_BODY()
+	DATAFLOW_NODE_DEFINE_INTERNAL(FClusterMergeToNeighborsDataflowNode, "ClusterMergeToNeighbors", "GeometryCollection|Cluster", "")
+
+public:
+
+	/** Collection on which to merge bones into a neighboring cluster */
+	UPROPERTY(meta = (DataflowInput, DataflowOutput, DataflowPassthrough = "Collection", DataflowIntrinsic))
+	FManagedArrayCollection Collection;
+
+	/** Bone selection */
+	UPROPERTY(meta = (DataflowInput, DisplayName = "TransformSelection"))
+	FDataflowTransformSelection TransformSelection;
+
+	/** Method to choose which neighbor to merge */
+	UPROPERTY(EditAnywhere, Category = Options, meta = (DisplayName = "Merge To"))
+	EClusterNeighborSelectionMethodEnum NeighborSelectionMethod = EClusterNeighborSelectionMethodEnum::Dataflow_ClusterNeighborSelectionMethod_LargestNeighbor;
+
+	/** Size (cube root of volume) of minimum desired post-merge clusters; if > 0, selected clusters may be merged multiple times until the cluster size is above this value */
+	UPROPERTY(EditAnywhere, Category = Options, meta = (DataflowInput, DisplayName = "Min Cluster Size", ClampMin = "0"))
+	float MinVolumeCubeRoot = 0.0f;
+
+	/** Whether to only allow clusters to merge if their bones are connected in the proximity graph */
+	UPROPERTY(EditAnywhere, Category = Options, meta = (DataflowInput))
+	bool bOnlyToConnected = true;
+
+	/** Whether to only allow clusters to merge if they have the same parent bone */
+	UPROPERTY(EditAnywhere, Category = Options, meta = (DataflowInput))
+	bool bOnlySameParent = true;
+
+	FClusterMergeToNeighborsDataflowNode(const Dataflow::FNodeParameters& InParam, FGuid InGuid = FGuid::NewGuid())
+		: FDataflowNode(InParam, InGuid)
+	{
+		RegisterInputConnection(&Collection);
+		RegisterInputConnection(&TransformSelection);
+		RegisterInputConnection(&MinVolumeCubeRoot);
+		RegisterInputConnection(&bOnlyToConnected);
+		RegisterInputConnection(&bOnlySameParent);
+
+		RegisterOutputConnection(&Collection, &Collection);
+	}
+
+	virtual void Evaluate(Dataflow::FContext& Context, const FDataflowOutput* Out) const override;
+
+};
+
+
 /**
  * Merge selected bones under a new parent cluster
  */
