Clarify dataflow selection set operations and add a non-symmetric difference op

diff --git a/Engine/Plugins/Experimental/GeometryCollectionPlugin/Source/GeometryCollectionNodes/Private/Dataflow/GeometryCollectionSelectionNodes.cpp b/Engine/Plugins/Experimental/GeometryCollectionPlugin/Source/GeometryCollectionNodes/Private/Dataflow/GeometryCollectionSelectionNodes.cpp
--- a/Engine/Plugins/Experimental/GeometryCollectionPlugin/Source/GeometryCollectionNodes/Private/Dataflow/GeometryCollectionSelectionNodes.cpp
+++ b/Engine/Plugins/Experimental/GeometryCollectionPlugin/Source/GeometryCollectionNodes/Private/Dataflow/GeometryCollectionSelectionNodes.cpp
@@ -122,6 +122,10 @@ void FCollectionTransformSelectionSetOperationDataflowNode::Evaluate(Dataflow::F
 			{
 				InTransformSelectionA.XOR(InTransformSelectionB, NewTransformSelection);
 			}
+			else if (Operation == ESetOperationEnum::Dataflow_SetOperation_Subtract)
+			{
+				InTransformSelectionA.Subtract(InTransformSelectionB, NewTransformSelection);
+			}
 		}
 		else
 		{
@@ -1143,6 +1147,10 @@ void FCollectionVertexSelectionSetOperationDataflowNode::Evaluate(Dataflow::FCon
 			{
 				InVertexSelectionA.XOR(InVertexSelectionB, NewVertexSelection);
 			}
+			else if (Operation == ESetOperationEnum::Dataflow_SetOperation_Subtract)
+			{
+				InVertexSelectionA.Subtract(InVertexSelectionB, NewVertexSelection);
+			}
 		}
 		else
 		{
diff --git a/Engine/Plugins/Experimental/GeometryCollectionPlugin/Source/GeometryCollectionNodes/Public/Dataflow/GeometryCollectionSelectionNodes.h b/Engine/Plugins/Experimental/GeometryCollectionPlugin/Source/GeometryCollectionNodes/Public/Dataflow/GeometryCollectionSelectionNodes.h
--- a/Engine/Plugins/Experimental/GeometryCollectionPlugin/Source/GeometryCollectionNodes/Public/Dataflow/GeometryCollectionSelectionNodes.h
+++ b/Engine/Plugins/Experimental/GeometryCollectionPlugin/Source/GeometryCollectionNodes/Public/Dataflow/GeometryCollectionSelectionNodes.h
@@ -50,12 +50,14 @@ public:
 UENUM(BlueprintType)
 enum class ESetOperationEnum : uint8
 {
-	/** Bitwise AND of incoming TransformSelections */
+	/** Select elements that are selected in both incoming selections (Bitwise AND) */
 	Dataflow_SetOperation_AND UMETA(DisplayName = "Intersect"),
-	/** Bitwise OR of incoming TransformSelections */
+	/** Select elements that are selected in either incoming selections (Bitwise OR) */
 	Dataflow_SetOperation_OR UMETA(DisplayName = "Union"),
-	/** Bitwise XOR of incoming TransformSelections */
-	Dataflow_SetOperation_XOR UMETA(DisplayName = "Symmetric Difference"),
+	/** Select elements that are selected in exactly one incoming selection (Bitwise XOR) */
+	Dataflow_SetOperation_XOR UMETA(DisplayName = "Symmetric Difference (XOR)"),
+	/** Select elements that are selected in only the first of the incoming selections (Bitwise A AND (NOT B)) */
+	Dataflow_SetOperation_Subtract UMETA(DisplayName = "Difference"),
 	//~~~
 	//256th entry
 	Dataflow_Max                UMETA(Hidden)
diff --git a/Engine/Source/Runtime/Experimental/Dataflow/Core/Private/Dataflow/DataflowSelection.cpp b/Engine/Source/Runtime/Experimental/Dataflow/Core/Private/Dataflow/DataflowSelection.cpp
--- a/Engine/Source/Runtime/Experimental/Dataflow/Core/Private/Dataflow/DataflowSelection.cpp
+++ b/Engine/Source/Runtime/Experimental/Dataflow/Core/Private/Dataflow/DataflowSelection.cpp
@@ -91,6 +91,13 @@ void FDataflowSelection::XOR(const FDataflowSelection& Other, FDataflowSelection
 	Result.SelectionArray = TBitArray<>::BitwiseXOR(SelectionArray, Other.SelectionArray, EBitwiseOperatorFlags::MaxSize);
 }
 
+void FDataflowSelection::Subtract(const FDataflowSelection& Other, FDataflowSelection& Result) const
+{
+	TBitArray<> InvSelectionArray = Other.SelectionArray;
+	InvSelectionArray.BitwiseNOT();
+	Result.SelectionArray = TBitArray<>::BitwiseAND(SelectionArray, InvSelectionArray, EBitwiseOperatorFlags::MaxSize);
+}
+
 int32 FDataflowSelection::NumSelected() const
 {
 	return SelectionArray.CountSetBits(0, SelectionArray.Num());
diff --git a/Engine/Source/Runtime/Experimental/Dataflow/Core/Public/Dataflow/DataflowSelection.h b/Engine/Source/Runtime/Experimental/Dataflow/Core/Public/Dataflow/DataflowSelection.h
--- a/Engine/Source/Runtime/Experimental/Dataflow/Core/Public/Dataflow/DataflowSelection.h
+++ b/Engine/Source/Runtime/Experimental/Dataflow/Core/Public/Dataflow/DataflowSelection.h
@@ -33,6 +33,8 @@ struct FDataflowSelection
 	DATAFLOWCORE_API void AND(const FDataflowSelection& Other, FDataflowSelection& Result) const;
 	DATAFLOWCORE_API void OR(const FDataflowSelection& Other, FDataflowSelection& Result) const;
 	DATAFLOWCORE_API void XOR(const FDataflowSelection& Other, FDataflowSelection& Result) const;
+	// subtract the selected elements of 'Other' from this selection
+	DATAFLOWCORE_API void Subtract(const FDataflowSelection& Other, FDataflowSelection& Result) const;
 	void Invert() { SelectionArray.BitwiseNOT(); }
 	DATAFLOWCORE_API void SetWithMask(const bool Value, const FDataflowSelection& Mask);
 	const TBitArray<>& GetBitArray() const { return SelectionArray; };
