Add an optional transform selection to the geometrycollection dataflow flatten node

diff --git a/Engine/Plugins/Experimental/GeometryCollectionPlugin/Source/GeometryCollectionNodes/Private/Dataflow/GeometryCollectionClusteringNodes.cpp b/Engine/Plugins/Experimental/GeometryCollectionPlugin/Source/GeometryCollectionNodes/Private/Dataflow/GeometryCollectionClusteringNodes.cpp
--- a/Engine/Plugins/Experimental/GeometryCollectionPlugin/Source/GeometryCollectionNodes/Private/Dataflow/GeometryCollectionClusteringNodes.cpp
+++ b/Engine/Plugins/Experimental/GeometryCollectionPlugin/Source/GeometryCollectionNodes/Private/Dataflow/GeometryCollectionClusteringNodes.cpp
@@ -108,15 +108,28 @@ void FClusterFlattenDataflowNode::Evaluate(Dataflow::FContext& Context, const FD
 				Chaos::Facades::FCollectionHierarchyFacade HierarchyFacade(*GeomCollection);
 				HierarchyFacade.GenerateLevelAttribute();
 
-				// Populate Selected Bones in an Array
-				// @todo(harsha) Implement with Selection
-				// For every bone in selected array: [RootClusterIndex]
-				const int32 RootClusterIndex = HierarchyFacade.GetRootIndex();
-				TArray<int32> LeafBones;
-				FGeometryCollectionClusteringUtility::GetLeafBones(GeomCollection.Get(), RootClusterIndex, true, LeafBones);
-				FGeometryCollectionClusteringUtility::ClusterBonesUnderExistingNode(GeomCollection.Get(), RootClusterIndex, LeafBones);
+				TArray<int32> ToFlatten;
+				if (IsConnected(&OptionalTransformSelection))
+				{
+					const FDataflowTransformSelection& InTransformSelection = GetValue(Context, &OptionalTransformSelection);
+					ToFlatten = InTransformSelection.AsArray();
+					GeometryCollection::Facades::FCollectionTransformSelectionFacade SelectionFacade(*GeomCollection);
+					SelectionFacade.ConvertSelectionToClusterNodes(ToFlatten, false);
+				}
+				else
+				{
+					const int32 RootClusterIndex = HierarchyFacade.GetRootIndex();
+					ToFlatten.Add(RootClusterIndex);
+				}
+
+				for (int32 ToFlattenIdx : ToFlatten)
+				{
+					TArray<int32> LeafBones;
+					FGeometryCollectionClusteringUtility::GetLeafBones(GeomCollection.Get(), ToFlattenIdx, true, LeafBones);
+					FGeometryCollectionClusteringUtility::ClusterBonesUnderExistingNode(GeomCollection.Get(), ToFlattenIdx, LeafBones);
+				}
+
 				FGeometryCollectionClusteringUtility::RemoveDanglingClusters(GeomCollection.Get());
-				// End for
 
 				HierarchyFacade.GenerateLevelAttribute();
 				SetValue(Context, (const FManagedArrayCollection&)(*GeomCollection), &Collection);
diff --git a/Engine/Plugins/Experimental/GeometryCollectionPlugin/Source/GeometryCollectionNodes/Public/Dataflow/GeometryCollectionClusteringNodes.h b/Engine/Plugins/Experimental/GeometryCollectionPlugin/Source/GeometryCollectionNodes/Public/Dataflow/GeometryCollectionClusteringNodes.h
--- a/Engine/Plugins/Experimental/GeometryCollectionPlugin/Source/GeometryCollectionNodes/Public/Dataflow/GeometryCollectionClusteringNodes.h
+++ b/Engine/Plugins/Experimental/GeometryCollectionPlugin/Source/GeometryCollectionNodes/Public/Dataflow/GeometryCollectionClusteringNodes.h
@@ -115,7 +115,7 @@ public:
 
 /**
  *
- * Flattens all bones to level 1
+ * Flattens selected bones. If no selection is provided, flattens all bones to level 1
  *
  */
 USTRUCT(meta = (DataflowGeometryCollection))
@@ -125,16 +125,21 @@ struct FClusterFlattenDataflowNode : public FDataflowNode
 	DATAFLOW_NODE_DEFINE_INTERNAL(FClusterFlattenDataflowNode, "Flatten", "GeometryCollection|Cluster", "")
 
 public:
-	// @todo(harsha) Support Selections
 
 	/** Fractured GeometryCollection to flatten */
 	UPROPERTY(meta = (DataflowInput, DataflowOutput, DataflowPassthrough = "Collection", DataflowIntrinsic))
 	FManagedArrayCollection Collection;
 
+	/** If connected, clusters under the selected bones will be flattened. If no selection is provided, all bones will be flattened to level 1. */
+	UPROPERTY(meta = (DataflowInput))
+	FDataflowTransformSelection OptionalTransformSelection;
+
 	FClusterFlattenDataflowNode(const Dataflow::FNodeParameters& InParam, FGuid InGuid = FGuid::NewGuid())
 		: FDataflowNode(InParam, InGuid)
 	{
 		RegisterInputConnection(&Collection);
+		RegisterInputConnection(&OptionalTransformSelection);
+		
 		RegisterOutputConnection(&Collection, &Collection);
 	}
 
