add geometrycollection dataflow node to clear convex hulls from a selection of bones
+ remove dataflowintrinsic tag from some optional inputs on some geometry collection nodes

diff --git a/Engine/Plugins/Experimental/GeometryCollectionPlugin/Source/GeometryCollectionNodes/Private/Dataflow/GeometryCollectionUtilityNodes.cpp b/Engine/Plugins/Experimental/GeometryCollectionPlugin/Source/GeometryCollectionNodes/Private/Dataflow/GeometryCollectionUtilityNodes.cpp
--- a/Engine/Plugins/Experimental/GeometryCollectionPlugin/Source/GeometryCollectionNodes/Private/Dataflow/GeometryCollectionUtilityNodes.cpp
+++ b/Engine/Plugins/Experimental/GeometryCollectionPlugin/Source/GeometryCollectionNodes/Private/Dataflow/GeometryCollectionUtilityNodes.cpp
@@ -30,6 +30,7 @@ namespace Dataflow
 		DATAFLOW_NODE_REGISTER_CREATION_FACTORY(FCreateNonOverlappingConvexHullsDataflowNode);
 		DATAFLOW_NODE_REGISTER_CREATION_FACTORY(FGenerateClusterConvexHullsFromLeafHullsDataflowNode);
 		DATAFLOW_NODE_REGISTER_CREATION_FACTORY(FGenerateClusterConvexHullsFromChildrenHullsDataflowNode);
+		DATAFLOW_NODE_REGISTER_CREATION_FACTORY(FClearConvexHullsDataflowNode);
 		DATAFLOW_NODE_REGISTER_CREATION_FACTORY(FMergeConvexHullsDataflowNode);
 		DATAFLOW_NODE_REGISTER_CREATION_FACTORY(FUpdateVolumeAttributesDataflowNode);
 		DATAFLOW_NODE_REGISTER_CREATION_FACTORY(FGetConvexHullVolumeDataflowNode);
@@ -394,6 +395,32 @@ FMergeConvexHullsDataflowNode::FMergeConvexHullsDataflowNode(const Dataflow::FNo
 	RegisterOutputConnection(&SphereCovering);
 }
 
+void FClearConvexHullsDataflowNode::Evaluate(Dataflow::FContext& Context, const FDataflowOutput* Out) const
+{
+	if (Out->IsA(&Collection))
+	{
+		FManagedArrayCollection InCollection = GetValue(Context, &Collection);
+
+		if (!IsConnected(&Collection) || !IsConnected(&TransformSelection) || !FGeometryCollectionConvexUtility::HasConvexHullData(&InCollection))
+		{
+			SetValue(Context, MoveTemp(InCollection), &Collection);
+			return;
+		}
+
+		const FDataflowTransformSelection& InSelection = GetValue(Context, &TransformSelection);
+
+		const FDataflowTransformSelection& InTransformSelection = GetValue<FDataflowTransformSelection>(Context, &TransformSelection);
+		TArray<int32> Selection = InTransformSelection.AsArray();
+
+		TArray<int32> ToClear = InTransformSelection.AsArray();
+		GeometryCollection::Facades::FCollectionTransformSelectionFacade SelectionFacade(InCollection);
+		SelectionFacade.Sanitize(ToClear);
+
+		FGeometryCollectionConvexUtility::RemoveConvexHulls(&InCollection, ToClear);
+		SetValue(Context, MoveTemp(InCollection), &Collection);
+	}
+}
+
 void FMergeConvexHullsDataflowNode::Evaluate(Dataflow::FContext& Context, const FDataflowOutput* Out) const
 {
 	if (Out->IsA(&Collection) || Out->IsA(&SphereCovering))
diff --git a/Engine/Plugins/Experimental/GeometryCollectionPlugin/Source/GeometryCollectionNodes/Public/Dataflow/GeometryCollectionUtilityNodes.h b/Engine/Plugins/Experimental/GeometryCollectionPlugin/Source/GeometryCollectionNodes/Public/Dataflow/GeometryCollectionUtilityNodes.h
--- a/Engine/Plugins/Experimental/GeometryCollectionPlugin/Source/GeometryCollectionNodes/Public/Dataflow/GeometryCollectionUtilityNodes.h
+++ b/Engine/Plugins/Experimental/GeometryCollectionPlugin/Source/GeometryCollectionNodes/Public/Dataflow/GeometryCollectionUtilityNodes.h
@@ -114,7 +114,7 @@ public:
 	FManagedArrayCollection Collection;
 
 	/** Optional transform selection to compute leaf hulls on -- if not provided, all leaf hulls will be computed. */
-	UPROPERTY(meta = (DataflowInput, DataflowIntrinsic))
+	UPROPERTY(meta = (DataflowInput))
 	FDataflowTransformSelection OptionalSelectionFilter;
 
 	/** How convex hulls are generated -- computed from geometry, imported from external collision shapes, or an intersection of both options. */
@@ -156,7 +156,7 @@ public:
 	FManagedArrayCollection Collection;
 
 	/** Optional transform selection to compute leaf hulls on -- if not provided, all leaf hulls will be computed. */
-	UPROPERTY(meta = (DataflowInput, DataflowIntrinsic))
+	UPROPERTY(meta = (DataflowInput))
 	FDataflowTransformSelection OptionalSelectionFilter;
 
 	UPROPERTY(EditAnywhere, Category = "Convex")
@@ -287,7 +287,7 @@ public:
 	EAllowConvexMergeMethod AllowMerges = EAllowConvexMergeMethod::ByProximity;
 
 	/** Optional transform selection to compute cluster hulls on -- if not provided, all cluster hulls will be computed. */
-	UPROPERTY(meta = (DataflowInput, DataflowIntrinsic))
+	UPROPERTY(meta = (DataflowInput))
 	FDataflowTransformSelection OptionalSelectionFilter;
 
 	/** Whether to use a sphere cover to define negative space that should not be covered by convex hulls */
@@ -360,7 +360,7 @@ public:
 	bool bPreferExternalCollisionShapes = true;
 
 	/** Optional transform selection to compute cluster hulls on -- if not provided, all cluster hulls will be computed. */
-	UPROPERTY(meta = (DataflowInput, DataflowIntrinsic))
+	UPROPERTY(meta = (DataflowInput))
 	FDataflowTransformSelection OptionalSelectionFilter;
 
 	/** Whether to use a sphere cover to define negative space that should not be covered by convex hulls */
@@ -398,6 +398,34 @@ public:
 };
 
 
+/** Clear convex hulls from the selected transforms */
+USTRUCT(meta = (DataflowGeometryCollection))
+struct FClearConvexHullsDataflowNode : public FDataflowNode
+{
+	GENERATED_USTRUCT_BODY()
+	DATAFLOW_NODE_DEFINE_INTERNAL(FClearConvexHullsDataflowNode, "ClearConvexHulls", "GeometryCollection|Utilities", "")
+
+public:
+	UPROPERTY(meta = (DataflowInput, DataflowOutput, DataflowPassthrough = "Collection", DataflowIntrinsic))
+	FManagedArrayCollection Collection;
+
+	/** Convex hulls will be cleared from these transforms */
+	UPROPERTY(meta = (DataflowInput, DataflowIntrinsic))
+	FDataflowTransformSelection TransformSelection;
+
+	FClearConvexHullsDataflowNode(const Dataflow::FNodeParameters& InParam, FGuid InGuid = FGuid::NewGuid())
+		: FDataflowNode(InParam, InGuid)
+	{
+		RegisterInputConnection(&Collection);
+		RegisterInputConnection(&TransformSelection);
+
+		RegisterOutputConnection(&Collection, &Collection);
+	}
+
+	virtual void Evaluate(Dataflow::FContext& Context, const FDataflowOutput* Out) const override;
+};
+
+
 /** Merge convex hulls on transforms with multiple hulls */
 USTRUCT(meta = (DataflowGeometryCollection))
 struct FMergeConvexHullsDataflowNode : public FDataflowNode
@@ -426,7 +454,7 @@ public:
 	double ErrorTolerance = 0.0;
 
 	/** Optional transform selection to compute cluster hulls on -- if not provided, all cluster hulls will be computed. */
-	UPROPERTY(meta = (DataflowInput, DataflowIntrinsic))
+	UPROPERTY(meta = (DataflowInput))
 	FDataflowTransformSelection OptionalSelectionFilter;
 
 	/** Whether to use a sphere cover to define negative space that should not be covered by convex hulls */
diff --git a/Engine/Source/Runtime/Experimental/Chaos/Private/GeometryCollection/GeometryCollectionConvexUtility.cpp b/Engine/Source/Runtime/Experimental/Chaos/Private/GeometryCollection/GeometryCollectionConvexUtility.cpp
--- a/Engine/Source/Runtime/Experimental/Chaos/Private/GeometryCollection/GeometryCollectionConvexUtility.cpp
+++ b/Engine/Source/Runtime/Experimental/Chaos/Private/GeometryCollection/GeometryCollectionConvexUtility.cpp
@@ -2412,13 +2412,13 @@ Chaos::FConvexPtr FGeometryCollectionConvexUtility::GetConvexHull(const FGeometr
 }
 
 
-void FGeometryCollectionConvexUtility::RemoveConvexHulls(FGeometryCollection* GeometryCollection, const TArray<int32>& SortedTransformDeletes)
+void FGeometryCollectionConvexUtility::RemoveConvexHulls(FManagedArrayCollection* GeometryCollection, const TArray<int32>& TransformsToClearHullsFrom)
 {
 	if (GeometryCollection->HasGroup(FGeometryCollection::ConvexGroup) && GeometryCollection->HasAttribute("TransformToConvexIndices", FTransformCollection::TransformGroup))
 	{
 		TManagedArray<TSet<int32>>& TransformToConvexIndices = GeometryCollection->ModifyAttribute<TSet<int32>>("TransformToConvexIndices", FTransformCollection::TransformGroup);
 		TArray<int32> ConvexIndices;
-		for (int32 TransformIdx : SortedTransformDeletes)
+		for (int32 TransformIdx : TransformsToClearHullsFrom)
 		{
 			if (TransformToConvexIndices[TransformIdx].Num() > 0)
 			{
diff --git a/Engine/Source/Runtime/Experimental/Chaos/Public/GeometryCollection/GeometryCollectionConvexUtility.h b/Engine/Source/Runtime/Experimental/Chaos/Public/GeometryCollection/GeometryCollectionConvexUtility.h
--- a/Engine/Source/Runtime/Experimental/Chaos/Public/GeometryCollection/GeometryCollectionConvexUtility.h
+++ b/Engine/Source/Runtime/Experimental/Chaos/Public/GeometryCollection/GeometryCollectionConvexUtility.h
@@ -236,7 +236,7 @@ public:
 	}
 
 	/** Delete the convex hulls pointed at by the transform indices provided. */
-	static CHAOS_API void RemoveConvexHulls(FGeometryCollection* GeometryCollection, const TArray<int32>& SortedTransformDeletes);
+	static CHAOS_API void RemoveConvexHulls(FManagedArrayCollection* GeometryCollection, const TArray<int32>& TransformsToClearHullsFrom);
 
 	/** Delete the convex hulls that are null */
 	static CHAOS_API void RemoveEmptyConvexHulls(FManagedArrayCollection& GeometryCollection);
