add outer-only negative space sampling via continuation marching cubes

diff --git a/Engine/Plugins/Experimental/GeometryCollectionPlugin/Source/GeometryCollectionNodes/Private/Dataflow/GeometryCollectionUtilityNodes.cpp b/Engine/Plugins/Experimental/GeometryCollectionPlugin/Source/GeometryCollectionNodes/Private/Dataflow/GeometryCollectionUtilityNodes.cpp
--- a/Engine/Plugins/Experimental/GeometryCollectionPlugin/Source/GeometryCollectionNodes/Private/Dataflow/GeometryCollectionUtilityNodes.cpp
+++ b/Engine/Plugins/Experimental/GeometryCollectionPlugin/Source/GeometryCollectionNodes/Private/Dataflow/GeometryCollectionUtilityNodes.cpp
@@ -253,6 +253,7 @@ void FGenerateClusterConvexHullsFromLeafHullsDataflowNode::Evaluate(Dataflow::FC
 				NegativeSpaceSettings.MinSpacing = GetValue(Context, &MinSampleSpacing);
 				NegativeSpaceSettings.SampleMethod = ConvertNegativeSpaceSampleMethodDataflowEnum(SampleMethod);
 				NegativeSpaceSettings.bRequireSearchSampleCoverage = bRequireSearchSampleCoverage;
+				NegativeSpaceSettings.bOnlyConnectedToHull = bOnlyConnectedToHull;
 				NegativeSpaceSettings.Sanitize();
 				bHasNegativeSpace = UE::FractureEngine::Convex::ComputeConvexHullsNegativeSpace(*GeomCollection, NegativeSpace, NegativeSpaceSettings, bHasSelectionFilter, SelectionArray);
 			}
@@ -338,6 +339,7 @@ void FGenerateClusterConvexHullsFromChildrenHullsDataflowNode::Evaluate(Dataflow
 				NegativeSpaceSettings.MinSpacing = GetValue(Context, &MinSampleSpacing);
 				NegativeSpaceSettings.SampleMethod = ConvertNegativeSpaceSampleMethodDataflowEnum(SampleMethod);
 				NegativeSpaceSettings.bRequireSearchSampleCoverage = bRequireSearchSampleCoverage;
+				NegativeSpaceSettings.bOnlyConnectedToHull = bOnlyConnectedToHull;
 				NegativeSpaceSettings.Sanitize();
 				bHasNegativeSpace = UE::FractureEngine::Convex::ComputeConvexHullsNegativeSpace(*GeomCollection, NegativeSpace, NegativeSpaceSettings, bHasSelectionFilter, SelectionArray);
 			}
@@ -448,6 +450,7 @@ void FMergeConvexHullsDataflowNode::Evaluate(Dataflow::FContext& Context, const
 			NegativeSpaceSettings.MinSpacing = GetValue(Context, &MinSampleSpacing);
 			NegativeSpaceSettings.SampleMethod = ConvertNegativeSpaceSampleMethodDataflowEnum(SampleMethod);
 			NegativeSpaceSettings.bRequireSearchSampleCoverage = bRequireSearchSampleCoverage;
+			NegativeSpaceSettings.bOnlyConnectedToHull = bOnlyConnectedToHull;
 			NegativeSpaceSettings.Sanitize();
 		}
 		if (bInProtectNegativeSpace && !bComputeNegativeSpacePerBone)
diff --git a/Engine/Plugins/Experimental/GeometryCollectionPlugin/Source/GeometryCollectionNodes/Public/Dataflow/GeometryCollectionUtilityNodes.h b/Engine/Plugins/Experimental/GeometryCollectionPlugin/Source/GeometryCollectionNodes/Public/Dataflow/GeometryCollectionUtilityNodes.h
--- a/Engine/Plugins/Experimental/GeometryCollectionPlugin/Source/GeometryCollectionNodes/Public/Dataflow/GeometryCollectionUtilityNodes.h
+++ b/Engine/Plugins/Experimental/GeometryCollectionPlugin/Source/GeometryCollectionNodes/Public/Dataflow/GeometryCollectionUtilityNodes.h
@@ -302,6 +302,10 @@ public:
 	UPROPERTY(EditAnywhere, Category = NegativeSpace, meta = (EditCondition = "bProtectNegativeSpace && SampleMethod == ENegativeSpaceSampleMethodDataflowEnum::VoxelSearch", EditConditionHides))
 	bool bRequireSearchSampleCoverage = false;
 
+	/** When performing Voxel Search, only look for negative space that is connected out to the convex hull. This removes inaccessable internal negative space from consideration. Only applies to Voxel Search. */
+	UPROPERTY(EditAnywhere, Category = NegativeSpace, meta = (EditCondition = "bProtectNegativeSpace && SampleMethod == ENegativeSpaceSampleMethodDataflowEnum::VoxelSearch", EditConditionHides))
+	bool bOnlyConnectedToHull = false;
+
 	/** Approximate number of spheres to consider when covering negative space */
 	UPROPERTY(EditAnywhere, Category = NegativeSpace, meta = (DataflowInput, ClampMin = 1, EditCondition = "bProtectNegativeSpace", EditConditionHides))
 	int32 TargetNumSamples = 50;
@@ -375,6 +379,10 @@ public:
 	UPROPERTY(EditAnywhere, Category = NegativeSpace, meta = (EditCondition = "bProtectNegativeSpace && SampleMethod == ENegativeSpaceSampleMethodDataflowEnum::VoxelSearch", EditConditionHides))
 	bool bRequireSearchSampleCoverage = false;
 
+	/** When performing Voxel Search, only look for negative space that is connected out to the convex hull. This removes inaccessable internal negative space from consideration. Only applies to Voxel Search. */
+	UPROPERTY(EditAnywhere, Category = NegativeSpace, meta = (EditCondition = "bProtectNegativeSpace && SampleMethod == ENegativeSpaceSampleMethodDataflowEnum::VoxelSearch", EditConditionHides))
+	bool bOnlyConnectedToHull = false;
+
 	/** Approximate number of spheres to consider when covering negative space */
 	UPROPERTY(EditAnywhere, Category = NegativeSpace, meta = (DataflowInput, ClampMin = 1, EditCondition = "bProtectNegativeSpace", EditConditionHides))
 	int32 TargetNumSamples = 50;
@@ -473,6 +481,10 @@ public:
 	UPROPERTY(EditAnywhere, Category = NegativeSpace, meta = (EditCondition = "bProtectNegativeSpace && SampleMethod == ENegativeSpaceSampleMethodDataflowEnum::VoxelSearch", EditConditionHides))
 	bool bRequireSearchSampleCoverage = false;
 
+	/** When performing Voxel Search, only look for negative space that is connected out to the convex hull. This removes inaccessable internal negative space from consideration. Only applies to Voxel Search. */
+	UPROPERTY(EditAnywhere, Category = NegativeSpace, meta = (EditCondition = "bProtectNegativeSpace && SampleMethod == ENegativeSpaceSampleMethodDataflowEnum::VoxelSearch", EditConditionHides))
+	bool bOnlyConnectedToHull = false;
+
 	/** Approximate number of spheres to consider when covering negative space */
 	UPROPERTY(EditAnywhere, Category = NegativeSpace, meta = (DataflowInput, ClampMin = 1, EditCondition = "bProtectNegativeSpace", EditConditionHides))
 	int32 TargetNumSamples = 50;
diff --git a/Engine/Plugins/Experimental/GeometryScripting/Source/GeometryScriptingCore/Private/CollisionFunctions.cpp b/Engine/Plugins/Experimental/GeometryScripting/Source/GeometryScriptingCore/Private/CollisionFunctions.cpp
--- a/Engine/Plugins/Experimental/GeometryScripting/Source/GeometryScriptingCore/Private/CollisionFunctions.cpp
+++ b/Engine/Plugins/Experimental/GeometryScripting/Source/GeometryScriptingCore/Private/CollisionFunctions.cpp
@@ -250,6 +250,7 @@ static FNegativeSpaceSampleSettings ConvertNegativeSpaceOptions(const FComputeNe
 	NegativeSpaceSettings.MinSpacing = NegativeSpaceOptions.MinSampleSpacing;
 	NegativeSpaceSettings.SampleMethod = ConvertNegativeSpaceSampleMethodEnum(NegativeSpaceOptions.SampleMethod);
 	NegativeSpaceSettings.bRequireSearchSampleCoverage = NegativeSpaceOptions.bRequireSearchSampleCoverage;
+	NegativeSpaceSettings.bOnlyConnectedToHull = NegativeSpaceOptions.bOnlyConnectedToHull;
 	NegativeSpaceSettings.bReferenceMeshHasNegativeWinding = false;
 	NegativeSpaceSettings.Sanitize();
 	return NegativeSpaceSettings;
diff --git a/Engine/Plugins/Experimental/GeometryScripting/Source/GeometryScriptingCore/Public/GeometryScript/CollisionFunctions.h b/Engine/Plugins/Experimental/GeometryScripting/Source/GeometryScriptingCore/Public/GeometryScript/CollisionFunctions.h
--- a/Engine/Plugins/Experimental/GeometryScripting/Source/GeometryScriptingCore/Public/GeometryScript/CollisionFunctions.h
+++ b/Engine/Plugins/Experimental/GeometryScripting/Source/GeometryScriptingCore/Public/GeometryScript/CollisionFunctions.h
@@ -140,6 +140,10 @@ public:
 	UPROPERTY(EditAnywhere, BlueprintReadWrite, Category = NegativeSpace, meta = (EditCondition = "SampleMethod == EGeometryScriptNegativeSpaceSampleMethod::VoxelSearch", EditConditionHides))
 	bool bRequireSearchSampleCoverage = false;
 
+	/** When performing Voxel Search, only look for negative space that is connected out to the convex hull. This removes inaccessable internal negative space from consideration. Only applies to Voxel Search. */
+	UPROPERTY(EditAnywhere, BlueprintReadWrite, Category = NegativeSpace, meta = (EditCondition = "SampleMethod == EGeometryScriptNegativeSpaceSampleMethod::VoxelSearch", EditConditionHides))
+	bool bOnlyConnectedToHull = false;
+
 	/** Approximate number of spheres to consider when covering negative space */
 	UPROPERTY(EditAnywhere, BlueprintReadWrite, Category = NegativeSpace, meta = (ClampMin = 1))
 	int32 TargetNumSamples = 50;
diff --git a/Engine/Source/Runtime/GeometryCore/Private/CompGeom/ConvexDecomposition3.cpp b/Engine/Source/Runtime/GeometryCore/Private/CompGeom/ConvexDecomposition3.cpp
--- a/Engine/Source/Runtime/GeometryCore/Private/CompGeom/ConvexDecomposition3.cpp
+++ b/Engine/Source/Runtime/GeometryCore/Private/CompGeom/ConvexDecomposition3.cpp
@@ -167,11 +167,64 @@ bool FSphereCovering::AddNegativeSpace(const TFastWindingTree<FDynamicMesh3>& Sp
 			return -1.0;
 		};
 
-		// Note: we could try to sample the hull away from vertices (on edges > 2*ReduceRadiusMargin long, and similarly on large triangles)
-		// to generate seeds to allow us to call the (generally faster) MarchingCubes.GenerateContinuation(Seeds) function below.
-		// However, naively calling GenerateContinuation with the mesh vertices can give an empty result (due to the radius margin).
-		// For now, we just call the slower (but reliable) Generate() function.
-		FDynamicMesh3 NegativeSpaceMesh(&MarchingCubes.Generate());
+		// If we can ignore internal negative space, we can use a seeded/continuation marching cubes here, starting from the convex hull.
+		// Since our negative space surface is away from the source mesh, we need to sample the hull away from its vertices as well;
+		// we do so with a subdivision scheme that can generate a large number of seed points, but still far fewer than there are voxel cells.
+		if (SampleSettings.bOnlyConnectedToHull)
+		{
+			TArray<FVector> Seeds;
+			double MinSeedSpacing = FMath::Max(MarchingCubes.CubeSize, SampleSettings.ReduceRadiusMargin);
+			double AddPtLenSq = 4 * MinSeedSpacing * MinSeedSpacing;
+			double SubDivLenSq = 16 * MinSeedSpacing * MinSeedSpacing;
+
+			// Make a manual stack of triangles to process --
+			// adding midpoints on long edges and subdividing when the resulting triangles could have long enough edges
+			TArray<FTriangle3d> ProcessTriStack;
+			for (FIndex3i TriInds : MeshHull.GetTriangles())
+			{
+				FTriangle3d HullTri(
+					Mesh->GetVertex(TriInds.A),
+					Mesh->GetVertex(TriInds.B),
+					Mesh->GetVertex(TriInds.C));
+				ProcessTriStack.Reset();
+				ProcessTriStack.Push(HullTri);
+				int32 SeedsAdded = 0;
+				while (!ProcessTriStack.IsEmpty())
+				{
+					FTriangle3d Tri = ProcessTriStack.Pop();
+					FVector3d EdgeLensSq(
+						FVector::DistSquared(Tri.V[2], Tri.V[0]),
+						FVector::DistSquared(Tri.V[0], Tri.V[1]),
+						FVector::DistSquared(Tri.V[1], Tri.V[2]));
+					int32 NumLongEdges = 0;
+					FTriangle3d Mids;
+					for (int32 SubIdx = 0, LastIdx = 2; SubIdx < 3; LastIdx = SubIdx++)
+					{
+						Mids.V[SubIdx] = (Tri.V[SubIdx] + Tri.V[LastIdx]) * .5;
+						if (EdgeLensSq[SubIdx] > AddPtLenSq)
+						{
+							Seeds.Add(Mids.V[SubIdx]);
+						}
+						NumLongEdges += (int32)(EdgeLensSq[SubIdx] > SubDivLenSq);
+					}
+					// Had any edge long enough to subdivide; go ahead and add all candidate tris
+					if (NumLongEdges > 0)
+					{
+						ProcessTriStack.Add(Mids);
+						ProcessTriStack.Emplace(Tri.V[0], Mids.V[0], Mids.V[1]);
+						ProcessTriStack.Emplace(Tri.V[1], Mids.V[1], Mids.V[2]);
+						ProcessTriStack.Emplace(Tri.V[2], Mids.V[2], Mids.V[0]);
+					}
+				}
+			}
+
+			MarchingCubes.GenerateContinuation(Seeds);
+		}
+		else
+		{
+			MarchingCubes.Generate();
+		}
+		FDynamicMesh3 NegativeSpaceMesh(&MarchingCubes);
 		NegativeSpaceMesh.DiscardAttributes();
 
 		// Make sure the mesh is compact to simplify downsampling below
diff --git a/Engine/Source/Runtime/GeometryCore/Public/CompGeom/ConvexDecomposition3.h b/Engine/Source/Runtime/GeometryCore/Public/CompGeom/ConvexDecomposition3.h
--- a/Engine/Source/Runtime/GeometryCore/Public/CompGeom/ConvexDecomposition3.h
+++ b/Engine/Source/Runtime/GeometryCore/Public/CompGeom/ConvexDecomposition3.h
@@ -63,13 +63,16 @@ struct FNegativeSpaceSampleSettings
 	// Ignore spheres with smaller radius than this
 	double MinRadius = 10.0;
 
-	// Currently only applies to VoxelSearch.
+	// Whether the reference mesh used has flipped orientation, so winding < -.5 is 'inside'
+	bool bReferenceMeshHasNegativeWinding = true;
+
+	// Below options currently only apply to VoxelSearch.
+	
 	// Whether to require that all candidate sample locations identified by Voxel Search are covered by negative space samples, up to the specified Min Sample Spacing.
 	// Note: This takes priority over TargetNumSamples if the TargetNumSamples did not achieve the required coverage.
 	bool bRequireSearchSampleCoverage = false;
-
-	// Whether the reference mesh used has flipped orientation, so winding < -.5 is 'inside'
-	bool bReferenceMeshHasNegativeWinding = true;
+	// Whether to only consider negative space that is connected to the bounding convex hull, i.e., to ignore hollow inner pockets of negative space that cannot be reached from the outside via a larger-than-ReduceRadiusMargin-wide path
+	bool bOnlyConnectedToHull = false;
 
 	// Make sure the settings values are in valid ranges
 	void Sanitize()
