Dataflow:
- Fixed bug in FGetGeometryCollectionAssetDataflowNode node when EngineContext is invalid

diff --git a/Engine/Plugins/Experimental/GeometryCollectionPlugin/Source/GeometryCollectionNodes/Private/Dataflow/Nodes/GeometryCollectionAssetNodes.cpp b/Engine/Plugins/Experimental/GeometryCollectionPlugin/Source/GeometryCollectionNodes/Private/Dataflow/Nodes/GeometryCollectionAssetNodes.cpp
--- a/Engine/Plugins/Experimental/GeometryCollectionPlugin/Source/GeometryCollectionNodes/Private/Dataflow/Nodes/GeometryCollectionAssetNodes.cpp
+++ b/Engine/Plugins/Experimental/GeometryCollectionPlugin/Source/GeometryCollectionNodes/Private/Dataflow/Nodes/GeometryCollectionAssetNodes.cpp
@@ -109,6 +109,8 @@ void FGetGeometryCollectionAssetDataflowNode::Evaluate(Dataflow::FContext& Conte
 			SetValue(Context, CollectionAsset, &Asset);
 		}
 	}
+
+	SetValue(Context, TObjectPtr<UGeometryCollection>(), &Asset);
 }
 
 // ===========================================================================================================================
