Dataflow : GeometryCollection
- Fix for GetGeometryCollectionAsset Node

diff --git a/Engine/Plugins/Experimental/GeometryCollectionPlugin/Source/GeometryCollectionNodes/Private/Dataflow/Nodes/GeometryCollectionAssetNodes.cpp b/Engine/Plugins/Experimental/GeometryCollectionPlugin/Source/GeometryCollectionNodes/Private/Dataflow/Nodes/GeometryCollectionAssetNodes.cpp
--- a/Engine/Plugins/Experimental/GeometryCollectionPlugin/Source/GeometryCollectionNodes/Private/Dataflow/Nodes/GeometryCollectionAssetNodes.cpp
+++ b/Engine/Plugins/Experimental/GeometryCollectionPlugin/Source/GeometryCollectionNodes/Private/Dataflow/Nodes/GeometryCollectionAssetNodes.cpp
@@ -102,15 +102,13 @@ FGetGeometryCollectionAssetDataflowNode::FGetGeometryCollectionAssetDataflowNode
 void FGetGeometryCollectionAssetDataflowNode::Evaluate(Dataflow::FContext& Context, const FDataflowOutput* Out) const
 {
 	ensure(Out->IsA(&Asset));
+
+	TObjectPtr<UGeometryCollection> CollectionAsset(nullptr);
 	if (const Dataflow::FEngineContext* EngineContext = Context.AsType<Dataflow::FEngineContext>())
 	{
-		if (const TObjectPtr<UGeometryCollection> CollectionAsset = Cast<UGeometryCollection>(EngineContext->Owner))
-		{
-			SetValue(Context, CollectionAsset, &Asset);
-		}
+		CollectionAsset = Cast<UGeometryCollection>(EngineContext->Owner);
 	}
-
-	SetValue(Context, TObjectPtr<UGeometryCollection>(), &Asset);
+	SetValue(Context, CollectionAsset, &Asset);
 }
 
 // ===========================================================================================================================
