Dataflow : Float Vertex to Indices node

diff --git a/Engine/Plugins/Experimental/GeometryCollectionPlugin/Source/GeometryCollectionNodes/Private/Dataflow/GeometryCollectionNodesPlugin.cpp b/Engine/Plugins/Experimental/GeometryCollectionPlugin/Source/GeometryCollectionNodes/Private/Dataflow/GeometryCollectionNodesPlugin.cpp
--- a/Engine/Plugins/Experimental/GeometryCollectionPlugin/Source/GeometryCollectionNodes/Private/Dataflow/GeometryCollectionNodesPlugin.cpp
+++ b/Engine/Plugins/Experimental/GeometryCollectionPlugin/Source/GeometryCollectionNodes/Private/Dataflow/GeometryCollectionNodesPlugin.cpp
@@ -24,6 +24,7 @@
 #include "Dataflow/GeometryCollectionVerticesNodes.h"
 #include "Dataflow/GeometryCollectionArrayNodes.h"
 #include "Dataflow/GeometryCollectionDebugNodes.h"
+#include "Dataflow/GeometryCollectionVertexScalarToVertexIndicesNode.h"
 
 #define LOCTEXT_NAMESPACE "DataflowNodes"
 
@@ -49,6 +50,8 @@ void IGeometryCollectionNodesPlugin::StartupModule()
 	Dataflow::GeometryCollectionVerticesNodes();
 	Dataflow::GeometryCollectionArrayNodes();
 	Dataflow::GeometryCollectionDebugNodes();
+	DATAFLOW_NODE_REGISTER_CREATION_FACTORY(FGeometryCollectionVertexScalarToVertexIndicesNode);
+
 }
 
 void IGeometryCollectionNodesPlugin::ShutdownModule()
diff --git a/Engine/Plugins/Experimental/GeometryCollectionPlugin/Source/GeometryCollectionNodes/Private/Dataflow/GeometryCollectionVertexScalarToVertexIndicesNode.cpp b/Engine/Plugins/Experimental/GeometryCollectionPlugin/Source/GeometryCollectionNodes/Private/Dataflow/GeometryCollectionVertexScalarToVertexIndicesNode.cpp
new file mode 100644
--- /dev/null
+++ b/Engine/Plugins/Experimental/GeometryCollectionPlugin/Source/GeometryCollectionNodes/Private/Dataflow/GeometryCollectionVertexScalarToVertexIndicesNode.cpp
@@ -0,0 +1,42 @@
+// Copyright Epic Games, Inc. All Rights Reserved.
+
+#include "Dataflow/GeometryCollectionVertexScalarToVertexIndicesNode.h"
+
+#include "Dataflow/DataflowInputOutput.h"
+
+#include UE_INLINE_GENERATED_CPP_BY_NAME(GeometryCollectionVertexScalarToVertexIndicesNode)
+#define LOCTEXT_NAMESPACE "FGeometryCollectionVertexScalarToVertexIndicesNode"
+
+FGeometryCollectionVertexScalarToVertexIndicesNode::FGeometryCollectionVertexScalarToVertexIndicesNode(const Dataflow::FNodeParameters& InParam, FGuid InGuid)
+	: FDataflowNode(InParam, InGuid)
+{
+	RegisterInputConnection(&Collection);
+	RegisterInputConnection(&VertexAttributeName);
+	RegisterOutputConnection(&Indices);
+}
+
+void FGeometryCollectionVertexScalarToVertexIndicesNode::Evaluate(Dataflow::FContext& Context, const FDataflowOutput* Out) const
+{
+	if (Out->IsA< TArray<int32> >(&Indices))
+	{
+		TArray<int32> IndicesOut;
+
+		const FManagedArrayCollection& InCollection = GetValue<FManagedArrayCollection>(Context, &Collection);
+		FString VertixAttributeNameVal = GetValue<FString>(Context, &VertexAttributeName, VertexAttributeName);
+
+		if( const TManagedArray<float>* FloatArray = InCollection.FindAttribute<float>(FName(VertixAttributeNameVal), FName("Vertices") ) )
+		{
+			for (int i = 0; i < FloatArray->Num(); i++)
+			{
+				if ((*FloatArray)[i] > SelectionThreshold)
+				{
+					IndicesOut.Add(i);
+				}
+			}
+		}
+		SetValue< TArray<int32> >(Context, MoveTemp(IndicesOut), &Indices);
+	}
+}
+
+
+#undef LOCTEXT_NAMESPACE
diff --git a/Engine/Plugins/Experimental/GeometryCollectionPlugin/Source/GeometryCollectionNodes/Public/Dataflow/GeometryCollectionVertexScalarToVertexIndicesNode.h b/Engine/Plugins/Experimental/GeometryCollectionPlugin/Source/GeometryCollectionNodes/Public/Dataflow/GeometryCollectionVertexScalarToVertexIndicesNode.h
new file mode 100644
--- /dev/null
+++ b/Engine/Plugins/Experimental/GeometryCollectionPlugin/Source/GeometryCollectionNodes/Public/Dataflow/GeometryCollectionVertexScalarToVertexIndicesNode.h
@@ -0,0 +1,39 @@
+// Copyright Epic Games, Inc. All Rights Reserved.
+
+#pragma once 
+
+#include "Dataflow/DataflowNode.h"
+#include "GeometryCollection/ManagedArrayCollection.h"
+
+#include "GeometryCollectionVertexScalarToVertexIndicesNode.generated.h"
+
+/** Convert an vertex float array to a list of indices */
+USTRUCT(meta = (DataflowGeometryCollection))
+struct FGeometryCollectionVertexScalarToVertexIndicesNode : public FDataflowNode
+{
+	GENERATED_USTRUCT_BODY()
+	DATAFLOW_NODE_DEFINE_INTERNAL(FGeometryCollectionVertexScalarToVertexIndicesNode, "VertexScalarToVertexIndices", "GeometryCollection", "Collection Vertex Weight Map to Indices")
+
+public:
+
+	UPROPERTY(Meta = (DataflowInput))
+	FManagedArrayCollection Collection;
+
+	/** The name of the vertex attribute to generate indices from. */
+	UPROPERTY(EditAnywhere, Category = "Selection Filter", Meta = (DataflowInput))
+	FString VertexAttributeName;
+
+	/** The value threshold for what is included in the vertex list. */
+	UPROPERTY(EditAnywhere, Category = "Selection Filter", Meta = (ClampMin = "0", ClampMax = "1"))
+	float SelectionThreshold = 0.f;
+
+	/** Output list of indices */
+	UPROPERTY(Meta = (DataflowOutput))
+	TArray<int32> Indices = {};
+
+	FGeometryCollectionVertexScalarToVertexIndicesNode(const Dataflow::FNodeParameters& InParam, FGuid InGuid = FGuid::NewGuid());
+
+private:
+
+	virtual void Evaluate(Dataflow::FContext& Context, const FDataflowOutput* Out) const override;
+};
