add new convexity-based clustering options to dataflow autocluster node

diff --git a/Engine/Plugins/Experimental/GeometryCollectionPlugin/Source/GeometryCollectionNodes/Private/Dataflow/GeometryCollectionClusteringNodes.cpp b/Engine/Plugins/Experimental/GeometryCollectionPlugin/Source/GeometryCollectionNodes/Private/Dataflow/GeometryCollectionClusteringNodes.cpp
--- a/Engine/Plugins/Experimental/GeometryCollectionPlugin/Source/GeometryCollectionNodes/Private/Dataflow/GeometryCollectionClusteringNodes.cpp
+++ b/Engine/Plugins/Experimental/GeometryCollectionPlugin/Source/GeometryCollectionNodes/Private/Dataflow/GeometryCollectionClusteringNodes.cpp
@@ -75,6 +75,8 @@ void FAutoClusterDataflowNode::Evaluate(Dataflow::FContext& Context, const FData
 			int32 InGridZ = GetValue<int32>(Context, &ClusterGridHeight);
 			float InMinimumClusterSize = GetValue<float>(Context, &MinimumSize);
 			int32 InKMeansIterations = ClusterSizeMethod == EClusterSizeMethodEnum::Dataflow_ClusterSizeMethod_ByGrid ? DriftIterations : 500;
+			bool bInPreferConvexity = GetValue(Context, &bPreferConvexity);
+			float InConcavityTolerance = GetValue(Context, &ConcavityTolerance);
 
 			TArray<int32> SelectedBones;
 			InTransformSelection.AsArray(SelectedBones);
@@ -88,7 +90,8 @@ void FAutoClusterDataflowNode::Evaluate(Dataflow::FContext& Context, const FData
 				InAutoCluster,
 				InAvoidIsolated, 
 				InEnforceSiteParameters,
-				InGridX, InGridY, InGridZ, InMinimumClusterSize, InKMeansIterations);
+				InGridX, InGridY, InGridZ, InMinimumClusterSize, InKMeansIterations, 
+				bInPreferConvexity, InConcavityTolerance);
 
 			SetValue<const FManagedArrayCollection&>(Context, *GeomCollection, &Collection);
 		}
diff --git a/Engine/Plugins/Experimental/GeometryCollectionPlugin/Source/GeometryCollectionNodes/Public/Dataflow/GeometryCollectionClusteringNodes.h b/Engine/Plugins/Experimental/GeometryCollectionPlugin/Source/GeometryCollectionNodes/Public/Dataflow/GeometryCollectionClusteringNodes.h
--- a/Engine/Plugins/Experimental/GeometryCollectionPlugin/Source/GeometryCollectionNodes/Public/Dataflow/GeometryCollectionClusteringNodes.h
+++ b/Engine/Plugins/Experimental/GeometryCollectionPlugin/Source/GeometryCollectionNodes/Public/Dataflow/GeometryCollectionClusteringNodes.h
@@ -72,6 +72,14 @@ public:
 	UPROPERTY(EditAnywhere, Category = ClusterSize, meta = (DataflowInput, ClampMin = "0"))
 	float MinimumSize = 0;
 
+	/** Whether to favor clusters that have a convex shape. (Note: Does not support ByGrid clustering.)  */
+	UPROPERTY(EditAnywhere, Category = AutoCluster, meta = (EditCondition = "ClusterSizeMethod != EClusterSizeMethod::ByGrid"))
+	bool bPreferConvexity = false;
+
+	/** If > 0, cube root of maximum concave volume to add per cluster (ignoring concavity of individual parts) */
+	UPROPERTY(EditAnywhere, Category = AutoCluster, meta = (EditCondition = "bPreferConvexity && ClusterSizeMethod != EClusterSizeMethod::ByGrid"))
+	float ConcavityTolerance = 0;
+
 	/** If true, bones will only be added to the same cluster if they are physically connected (either directly, or via other bones in the same cluster) */
 	UPROPERTY(EditAnywhere, Category = AutoCluster, meta = (DisplayName = "Enforce Cluster Connectivity"))
 	bool AutoCluster = true;
@@ -104,6 +112,8 @@ public:
 		RegisterInputConnection(&ClusterGridDepth);
 		RegisterInputConnection(&ClusterGridHeight);
 		RegisterInputConnection(&MinimumSize);
+		RegisterInputConnection(&bPreferConvexity);
+		RegisterInputConnection(&ConcavityTolerance);
 
 		RegisterOutputConnection(&Collection, &Collection);
 	}
