Dataflow:
- Implemented StaticMeshToCollection node

diff --git a/Engine/Plugins/Experimental/GeometryCollectionPlugin/Source/GeometryCollectionNodes/Private/Dataflow/Nodes/GeometryCollectionAssetNodes.cpp b/Engine/Plugins/Experimental/GeometryCollectionPlugin/Source/GeometryCollectionNodes/Private/Dataflow/Nodes/GeometryCollectionAssetNodes.cpp
--- a/Engine/Plugins/Experimental/GeometryCollectionPlugin/Source/GeometryCollectionNodes/Private/Dataflow/Nodes/GeometryCollectionAssetNodes.cpp
+++ b/Engine/Plugins/Experimental/GeometryCollectionPlugin/Source/GeometryCollectionNodes/Private/Dataflow/Nodes/GeometryCollectionAssetNodes.cpp
@@ -21,6 +21,7 @@ namespace Dataflow
 		DATAFLOW_NODE_REGISTER_CREATION_FACTORY(FGetGeometryCollectionAssetDataflowNode);
 		DATAFLOW_NODE_REGISTER_CREATION_FACTORY(FGetGeometryCollectionSourcesDataflowNode);
 		DATAFLOW_NODE_REGISTER_CREATION_FACTORY(FCreateGeometryCollectionFromSourcesDataflowNode);
+		DATAFLOW_NODE_REGISTER_CREATION_FACTORY(FStaticMeshToCollectionDataflowNode);
 
 		// Terminal
 		DATAFLOW_NODE_REGISTER_CREATION_FACTORY_NODE_COLORS_BY_CATEGORY("Terminal", FLinearColor(0.f, 0.f, 0.f), CDefaultNodeBodyTintColor);
@@ -223,3 +224,32 @@ void FCreateGeometryCollectionFromSourcesDataflowNode::Evaluate(Dataflow::FConte
 	SetValue(Context, MoveTemp(OutMaterials), &Materials);
 	SetValue(Context, MoveTemp(OutInstancedMeshes), &InstancedMeshes);
 }
+
+// ===========================================================================================================================
+
+FStaticMeshToCollectionDataflowNode::FStaticMeshToCollectionDataflowNode(const Dataflow::FNodeParameters& InParam, FGuid InGuid)
+	: FDataflowNode(InParam, InGuid)
+{
+	RegisterOutputConnection(&Collection);
+	RegisterOutputConnection(&Materials);
+	RegisterOutputConnection(&InstancedMeshes);
+}
+
+void FStaticMeshToCollectionDataflowNode::Evaluate(Dataflow::FContext& Context, const FDataflowOutput* Out) const
+{
+	ensure(Out->IsA(&Collection) || Out->IsA(&Materials) || Out->IsA(&InstancedMeshes));
+
+	FManagedArrayCollection OutCollection;
+	TArray<TObjectPtr<UMaterial>> OutMaterials;
+	TArray<FGeometryCollectionAutoInstanceMesh> OutInstancedMeshes;
+
+	if (StaticMesh)
+	{
+		FGeometryCollectionEngineConversion::ConvertStaticMeshToGeometryCollection(StaticMesh, OutCollection, OutMaterials, OutInstancedMeshes, bSetInternalFromMaterialIndex, bSplitComponents);
+	}
+
+	// Set Outputs
+	SetValue(Context, MoveTemp(OutCollection), &Collection);
+	SetValue(Context, MoveTemp(OutMaterials), &Materials);
+	SetValue(Context, MoveTemp(OutInstancedMeshes), &InstancedMeshes);
+}
diff --git a/Engine/Plugins/Experimental/GeometryCollectionPlugin/Source/GeometryCollectionNodes/Public/Dataflow/Nodes/GeometryCollectionAssetNodes.h b/Engine/Plugins/Experimental/GeometryCollectionPlugin/Source/GeometryCollectionNodes/Public/Dataflow/Nodes/GeometryCollectionAssetNodes.h
--- a/Engine/Plugins/Experimental/GeometryCollectionPlugin/Source/GeometryCollectionNodes/Public/Dataflow/Nodes/GeometryCollectionAssetNodes.h
+++ b/Engine/Plugins/Experimental/GeometryCollectionPlugin/Source/GeometryCollectionNodes/Public/Dataflow/Nodes/GeometryCollectionAssetNodes.h
@@ -10,6 +10,7 @@
 #include "GeometryCollectionAssetNodes.generated.h"
 
 class UGeometryCollection;
+class UStaticMesh;
 
 USTRUCT(meta = (DataflowGeometryCollection, DataflowTerminal))
 struct FGeometryCollectionTerminalDataflowNode : public FDataflowTerminalNode
@@ -103,7 +104,7 @@ public:
 	UPROPERTY(meta = (DataflowOutput, DisplayName = "Collection"))
 	FManagedArrayCollection Collection;
 
-	/** Geometry collection newly created */
+	/** Materials array to use for this asset */
 	UPROPERTY(meta = (DataflowOutput, DisplayName = "Materials"))
 	TArray<TObjectPtr<UMaterial>> Materials;
 
@@ -112,6 +113,46 @@ public:
 	TArray<FGeometryCollectionAutoInstanceMesh> InstancedMeshes;
 };
 
+
+/**
+ * Create a geometry collection from a UStaticMesh
+ */
+USTRUCT(meta = (DataflowContext = "GeometryCollection", DataflowGeometryCollection, DataflowTerminal))
+struct FStaticMeshToCollectionDataflowNode : public FDataflowNode
+{
+	GENERATED_USTRUCT_BODY()
+	DATAFLOW_NODE_DEFINE_INTERNAL(FStaticMeshToCollectionDataflowNode, "StaticMeshToCollection", "GeometryCollection|Asset", "")
+	DATAFLOW_NODE_RENDER_TYPE(FGeometryCollection::StaticType(), "Collection")
+
+public:
+	/** Asset input */
+	UPROPERTY(EditAnywhere, Category = "Asset");
+	TObjectPtr<UStaticMesh> StaticMesh;
+
+	/** Set the internal faces from material index */
+	UPROPERTY(EditAnywhere, Category = "Asset");
+	bool bSetInternalFromMaterialIndex = true;
+
+	/** Split components */
+	UPROPERTY(EditAnywhere, Category = "Asset");
+	bool bSplitComponents = false;
+
+	/** Geometry collection newly created */
+	UPROPERTY(meta = (DataflowOutput, DisplayName = "Collection"))
+	FManagedArrayCollection Collection;
+
+	/** Materials array to use for this asset */
+	UPROPERTY(meta = (DataflowOutput, DisplayName = "Materials"))
+	TArray<TObjectPtr<UMaterial>> Materials;
+
+	/** Array of instanced meshes*/
+	UPROPERTY(meta = (DataflowOutput, DisplayName = "InstancedMeshes"))
+	TArray<FGeometryCollectionAutoInstanceMesh> InstancedMeshes;
+
+	FStaticMeshToCollectionDataflowNode(const Dataflow::FNodeParameters& InParam, FGuid InGuid = FGuid::NewGuid());
+	virtual void Evaluate(Dataflow::FContext& Context, const FDataflowOutput* Out) const override;
+};
+
 namespace Dataflow
 {
 	void GeometryCollectionEngineAssetNodes();
diff --git a/Engine/Source/Runtime/Experimental/GeometryCollectionEngine/Private/GeometryCollection/GeometryCollectionEngineConversion.cpp b/Engine/Source/Runtime/Experimental/GeometryCollectionEngine/Private/GeometryCollection/GeometryCollectionEngineConversion.cpp
--- a/Engine/Source/Runtime/Experimental/GeometryCollectionEngine/Private/GeometryCollection/GeometryCollectionEngineConversion.cpp
+++ b/Engine/Source/Runtime/Experimental/GeometryCollectionEngine/Private/GeometryCollection/GeometryCollectionEngineConversion.cpp
@@ -1469,4 +1469,50 @@ void FGeometryCollectionEngineConversion::AppendGeometryCollectionSource(const F
 	}
 }
 
-#undef LOCTEXT_NAMESPACE 
\ No newline at end of file
+void FGeometryCollectionEngineConversion::ConvertStaticMeshToGeometryCollection(const TObjectPtr<UStaticMesh> StaticMesh, FManagedArrayCollection& OutCollection, TArray<TObjectPtr<UMaterial>>& OutMaterials, TArray<FGeometryCollectionAutoInstanceMesh>& OutInstancedMeshes, bool bSetInternalFromMaterialIndex, bool bSplitComponents)
+{
+#if WITH_EDITORONLY_DATA
+	if (UGeometryCollection* NewGeometryCollection = NewObject<UGeometryCollection>())
+	{
+		// If any of the static meshes have Nanite enabled, also enable on the new geometry collection asset for convenience.
+		NewGeometryCollection->EnableNanite |= StaticMesh->IsNaniteEnabled();
+
+		FTransform ComponentTransform = FTransform::Identity;
+
+		// Record the contributing source on the asset.
+		FSoftObjectPath SourceSoftObjectPath(StaticMesh);
+
+		// Materials		
+		TArray<TObjectPtr<UMaterialInterface>> MatArr;
+		for (auto& StaticMaterial : StaticMesh->GetStaticMaterials())
+		{
+			MatArr.Emplace(StaticMaterial.MaterialInterface);
+		}
+		TArray<TObjectPtr<UMaterialInterface>> SourceMaterials(MatArr);
+
+		// InstanceMeshes
+		FGeometryCollectionAutoInstanceMesh NewInstanceMesh;
+		NewInstanceMesh.Mesh = StaticMesh;
+		NewInstanceMesh.Materials = SourceMaterials;
+		OutInstancedMeshes.Emplace(NewInstanceMesh);
+
+		bool bAddInternalMaterials = true;
+
+		NewGeometryCollection->GeometrySource.Emplace(SourceSoftObjectPath, ComponentTransform, SourceMaterials, bSplitComponents, bSetInternalFromMaterialIndex);
+		FGeometryCollectionEngineConversion::AppendStaticMesh(StaticMesh, SourceMaterials, ComponentTransform, NewGeometryCollection, false, bAddInternalMaterials, bSplitComponents, bSetInternalFromMaterialIndex);
+
+		NewGeometryCollection->InitializeMaterials();
+
+		// Materials
+		for (auto& Material : NewGeometryCollection->Materials)
+		{
+			OutMaterials.Emplace(Material->GetMaterial());
+		}
+
+		TSharedPtr<FGeometryCollection> OutCollectionPtr = NewGeometryCollection->GetGeometryCollection();
+		OutCollectionPtr->CopyTo(&OutCollection);
+	}
+#endif //WITH_EDITORONLY_DATA
+}
+
+#undef LOCTEXT_NAMESPACE 
diff --git a/Engine/Source/Runtime/Experimental/GeometryCollectionEngine/Public/GeometryCollection/GeometryCollectionEngineConversion.h b/Engine/Source/Runtime/Experimental/GeometryCollectionEngine/Public/GeometryCollection/GeometryCollectionEngineConversion.h
--- a/Engine/Source/Runtime/Experimental/GeometryCollectionEngine/Public/GeometryCollection/GeometryCollectionEngineConversion.h
+++ b/Engine/Source/Runtime/Experimental/GeometryCollectionEngine/Public/GeometryCollection/GeometryCollectionEngineConversion.h
@@ -174,4 +174,15 @@ public:
 	*  @param MaterialsInOut  : array of materials to append to
 	*/
 	static GEOMETRYCOLLECTIONENGINE_API void AppendGeometryCollectionSource(const FGeometryCollectionSource& GeometryCollectionSource, FGeometryCollection& GeometryCollectionInOut, TArray<UMaterial*>& MaterialsInOut, bool ReindexMaterials = true);
+
+	/**
+	*  Converts a StaticMesh to a GeometryCollection
+	*  @param StaticMesh : Const mesh to read vertex/normals/index data from
+	*  @param OutCollection : FGeometryCollection output
+	*  @param OutMaterials : materials from the StaticMesh
+	*  @param OutInstancedMeshes : InstancedMeshes
+	*  @param bSetInternalFromMaterialIndex : Set the internal faces using the materials
+	*  @param bSplitComponents : Split the components
+	*/
+	static GEOMETRYCOLLECTIONENGINE_API void ConvertStaticMeshToGeometryCollection(const TObjectPtr<UStaticMesh> StaticMesh, FManagedArrayCollection& OutCollection, TArray<TObjectPtr<UMaterial>>& OutMaterials, TArray<FGeometryCollectionAutoInstanceMesh>& OutInstancedMeshes, bool bSetInternalFromMaterialIndex = true, bool bSplitComponents = false);
 };
