Dataflow:
- Implemented GeometryCollectionToCollection node

diff --git a/Engine/Plugins/Experimental/GeometryCollectionPlugin/Source/GeometryCollectionNodes/Private/Dataflow/Nodes/GeometryCollectionAssetNodes.cpp b/Engine/Plugins/Experimental/GeometryCollectionPlugin/Source/GeometryCollectionNodes/Private/Dataflow/Nodes/GeometryCollectionAssetNodes.cpp
--- a/Engine/Plugins/Experimental/GeometryCollectionPlugin/Source/GeometryCollectionNodes/Private/Dataflow/Nodes/GeometryCollectionAssetNodes.cpp
+++ b/Engine/Plugins/Experimental/GeometryCollectionPlugin/Source/GeometryCollectionNodes/Private/Dataflow/Nodes/GeometryCollectionAssetNodes.cpp
@@ -22,6 +22,7 @@ namespace Dataflow
 		DATAFLOW_NODE_REGISTER_CREATION_FACTORY(FGetGeometryCollectionSourcesDataflowNode);
 		DATAFLOW_NODE_REGISTER_CREATION_FACTORY(FCreateGeometryCollectionFromSourcesDataflowNode);
 		DATAFLOW_NODE_REGISTER_CREATION_FACTORY(FStaticMeshToCollectionDataflowNode);
+		DATAFLOW_NODE_REGISTER_CREATION_FACTORY(FGeometryCollectionToCollectionDataflowNode);
 
 		// Terminal
 		DATAFLOW_NODE_REGISTER_CREATION_FACTORY_NODE_COLORS_BY_CATEGORY("Terminal", FLinearColor(0.f, 0.f, 0.f), CDefaultNodeBodyTintColor);
@@ -253,3 +254,32 @@ void FStaticMeshToCollectionDataflowNode::Evaluate(Dataflow::FContext& Context,
 	SetValue(Context, MoveTemp(OutMaterials), &Materials);
 	SetValue(Context, MoveTemp(OutInstancedMeshes), &InstancedMeshes);
 }
+
+// ===========================================================================================================================
+
+FGeometryCollectionToCollectionDataflowNode::FGeometryCollectionToCollectionDataflowNode(const Dataflow::FNodeParameters& InParam, FGuid InGuid)
+	: FDataflowNode(InParam, InGuid)
+{
+	RegisterOutputConnection(&Collection);
+	RegisterOutputConnection(&Materials);
+	RegisterOutputConnection(&InstancedMeshes);
+}
+
+void FGeometryCollectionToCollectionDataflowNode::Evaluate(Dataflow::FContext& Context, const FDataflowOutput* Out) const
+{
+	ensure(Out->IsA(&Collection) || Out->IsA(&Materials) || Out->IsA(&InstancedMeshes));
+
+	FManagedArrayCollection OutCollection;
+	TArray<TObjectPtr<UMaterial>> OutMaterials;
+	TArray<FGeometryCollectionAutoInstanceMesh> OutInstancedMeshes;
+
+	if (GeometryCollection)
+	{
+		FGeometryCollectionEngineConversion::ConvertGeometryCollectionToGeometryCollection(GeometryCollection, OutCollection, OutMaterials, OutInstancedMeshes);
+	}
+
+	// Set Outputs
+	SetValue(Context, MoveTemp(OutCollection), &Collection);
+	SetValue(Context, MoveTemp(OutMaterials), &Materials);
+	SetValue(Context, MoveTemp(OutInstancedMeshes), &InstancedMeshes);
+}
diff --git a/Engine/Plugins/Experimental/GeometryCollectionPlugin/Source/GeometryCollectionNodes/Public/Dataflow/Nodes/GeometryCollectionAssetNodes.h b/Engine/Plugins/Experimental/GeometryCollectionPlugin/Source/GeometryCollectionNodes/Public/Dataflow/Nodes/GeometryCollectionAssetNodes.h
--- a/Engine/Plugins/Experimental/GeometryCollectionPlugin/Source/GeometryCollectionNodes/Public/Dataflow/Nodes/GeometryCollectionAssetNodes.h
+++ b/Engine/Plugins/Experimental/GeometryCollectionPlugin/Source/GeometryCollectionNodes/Public/Dataflow/Nodes/GeometryCollectionAssetNodes.h
@@ -153,6 +153,38 @@ public:
 	virtual void Evaluate(Dataflow::FContext& Context, const FDataflowOutput* Out) const override;
 };
 
+
+/**
+ * Converts a UGeometryCollection asset to an FManagedArrayCollection
+ */
+USTRUCT(meta = (DataflowContext = "GeometryCollection", DataflowGeometryCollection, DataflowTerminal))
+struct FGeometryCollectionToCollectionDataflowNode : public FDataflowNode
+{
+	GENERATED_USTRUCT_BODY()
+	DATAFLOW_NODE_DEFINE_INTERNAL(FGeometryCollectionToCollectionDataflowNode, "GeometryCollectionToCollection", "GeometryCollection|Asset", "")
+	DATAFLOW_NODE_RENDER_TYPE(FGeometryCollection::StaticType(), "Collection")
+
+public:
+	/** Asset input */
+	UPROPERTY(EditAnywhere, Category = "Asset");
+	TObjectPtr<UGeometryCollection> GeometryCollection;
+
+	/** Geometry collection newly created */
+	UPROPERTY(meta = (DataflowOutput, DisplayName = "Collection"))
+	FManagedArrayCollection Collection;
+
+	/** Materials array to use for this asset */
+	UPROPERTY(meta = (DataflowOutput, DisplayName = "Materials"))
+	TArray<TObjectPtr<UMaterial>> Materials;
+
+	/** Array of instanced meshes*/
+	UPROPERTY(meta = (DataflowOutput, DisplayName = "InstancedMeshes"))
+	TArray<FGeometryCollectionAutoInstanceMesh> InstancedMeshes;
+
+	FGeometryCollectionToCollectionDataflowNode(const Dataflow::FNodeParameters& InParam, FGuid InGuid = FGuid::NewGuid());
+	virtual void Evaluate(Dataflow::FContext& Context, const FDataflowOutput* Out) const override;
+};
+
 namespace Dataflow
 {
 	void GeometryCollectionEngineAssetNodes();
diff --git a/Engine/Source/Runtime/Experimental/GeometryCollectionEngine/Private/GeometryCollection/GeometryCollectionEngineConversion.cpp b/Engine/Source/Runtime/Experimental/GeometryCollectionEngine/Private/GeometryCollection/GeometryCollectionEngineConversion.cpp
--- a/Engine/Source/Runtime/Experimental/GeometryCollectionEngine/Private/GeometryCollection/GeometryCollectionEngineConversion.cpp
+++ b/Engine/Source/Runtime/Experimental/GeometryCollectionEngine/Private/GeometryCollection/GeometryCollectionEngineConversion.cpp
@@ -1515,4 +1515,26 @@ void FGeometryCollectionEngineConversion::ConvertStaticMeshToGeometryCollection(
 #endif //WITH_EDITORONLY_DATA
 }
 
+void FGeometryCollectionEngineConversion::ConvertGeometryCollectionToGeometryCollection(const TObjectPtr<UGeometryCollection> InGeometryCollectionAssetPtr, FManagedArrayCollection& OutCollection, TArray<TObjectPtr<UMaterial>>& OutMaterials, TArray<FGeometryCollectionAutoInstanceMesh>& OutInstancedMeshes)
+{
+	if (InGeometryCollectionAssetPtr)
+	{
+		const TSharedPtr<FGeometryCollection, ESPMode::ThreadSafe> NewGeometryCollectionPtr = InGeometryCollectionAssetPtr->GetGeometryCollection();
+
+		// Materials
+		for (auto& Material : InGeometryCollectionAssetPtr->Materials)
+		{
+			OutMaterials.Emplace(Material->GetMaterial());
+		}
+
+		// InstanceMeshes
+		OutInstancedMeshes = InGeometryCollectionAssetPtr->AutoInstanceMeshes;
+			
+		if (NewGeometryCollectionPtr)
+		{
+			NewGeometryCollectionPtr->CopyTo(&OutCollection);
+		}
+	}
+}
+
 #undef LOCTEXT_NAMESPACE 
diff --git a/Engine/Source/Runtime/Experimental/GeometryCollectionEngine/Public/GeometryCollection/GeometryCollectionEngineConversion.h b/Engine/Source/Runtime/Experimental/GeometryCollectionEngine/Public/GeometryCollection/GeometryCollectionEngineConversion.h
--- a/Engine/Source/Runtime/Experimental/GeometryCollectionEngine/Public/GeometryCollection/GeometryCollectionEngineConversion.h
+++ b/Engine/Source/Runtime/Experimental/GeometryCollectionEngine/Public/GeometryCollection/GeometryCollectionEngineConversion.h
@@ -185,4 +185,14 @@ public:
 	*  @param bSplitComponents : Split the components
 	*/
 	static GEOMETRYCOLLECTIONENGINE_API void ConvertStaticMeshToGeometryCollection(const TObjectPtr<UStaticMesh> StaticMesh, FManagedArrayCollection& OutCollection, TArray<TObjectPtr<UMaterial>>& OutMaterials, TArray<FGeometryCollectionAutoInstanceMesh>& OutInstancedMeshes, bool bSetInternalFromMaterialIndex = true, bool bSplitComponents = false);
+
+	/**
+	*  Converts a UGeometryCollection asset to an FManagedArrayCollection, and arrays of materials and instanced meshes
+	*  @param GeometryCollection : UGeometryCollection input to convert
+	*  @param OutCollection : FGeometryCollection output
+	*  @param OutMaterials : Materials from the UGeometryCollection
+	*  @param OutInstancedMeshes : InstancedMeshes
+	*/	
+	static GEOMETRYCOLLECTIONENGINE_API void ConvertGeometryCollectionToGeometryCollection(const TObjectPtr<UGeometryCollection> InGeometryCollectionAssetPtr, FManagedArrayCollection& OutCollection, TArray<TObjectPtr<UMaterial>>& OutMaterials, TArray<FGeometryCollectionAutoInstanceMesh>& OutInstancedMeshes);
+
 };
