GeometryCollection : SkeletonToCollection Node

diff --git a/Engine/Plugins/Experimental/GeometryCollectionPlugin/Source/GeometryCollectionNodes/Private/Dataflow/GeometryCollectionNodesPlugin.cpp b/Engine/Plugins/Experimental/GeometryCollectionPlugin/Source/GeometryCollectionNodes/Private/Dataflow/GeometryCollectionNodesPlugin.cpp
--- a/Engine/Plugins/Experimental/GeometryCollectionPlugin/Source/GeometryCollectionNodes/Private/Dataflow/GeometryCollectionNodesPlugin.cpp
+++ b/Engine/Plugins/Experimental/GeometryCollectionPlugin/Source/GeometryCollectionNodes/Private/Dataflow/GeometryCollectionNodesPlugin.cpp
@@ -8,7 +8,8 @@
 #include "Dataflow/GeometryCollectionNodes.h"
 #include "Dataflow/Nodes/GeometryCollectionAssetNodes.h"
 #include "Dataflow/GeometryCollectionProcessingNodes.h"
-#include "Dataflow/GeometryCollectionSkeletalMeshNodes.h"
+#include "Dataflow/GeometryCollectionSkeletalMeshToCollectionNode.h"
+#include "Dataflow/GeometryCollectionSkeletonToCollectionNode.h"
 #include "Dataflow/GeometryCollectionSelectionNodes.h"
 #include "Dataflow/GeometryCollectionMeshNodes.h"
 #include "Dataflow/GeometryCollectionClusteringNodes.h"
@@ -34,7 +35,6 @@ void IGeometryCollectionNodesPlugin::StartupModule()
 	Dataflow::GeometryCollectionEngineNodes();
 	Dataflow::GeometryCollectionEngineAssetNodes();
 	Dataflow::GeometryCollectionProcessingNodes();
-	Dataflow::GeometryCollectionSkeletalMeshNodes();
 	Dataflow::GeometryCollectionSelectionNodes();
 	Dataflow::GeometryCollectionMeshNodes();
 	Dataflow::GeometryCollectionClusteringNodes();
@@ -51,7 +51,8 @@ void IGeometryCollectionNodesPlugin::StartupModule()
 	Dataflow::GeometryCollectionArrayNodes();
 	Dataflow::GeometryCollectionDebugNodes();
 	DATAFLOW_NODE_REGISTER_CREATION_FACTORY(FGeometryCollectionVertexScalarToVertexIndicesNode);
-
+	DATAFLOW_NODE_REGISTER_CREATION_FACTORY(FSkeletonToCollectionDataflowNode);
+	DATAFLOW_NODE_REGISTER_CREATION_FACTORY(FSkeletalMeshToCollectionDataflowNode);
 }
 
 void IGeometryCollectionNodesPlugin::ShutdownModule()
diff --git a/Engine/Plugins/Experimental/GeometryCollectionPlugin/Source/GeometryCollectionNodes/Private/Dataflow/GeometryCollectionSkeletalMeshNodes.cpp b/Engine/Plugins/Experimental/GeometryCollectionPlugin/Source/GeometryCollectionNodes/Private/Dataflow/GeometryCollectionSkeletalMeshNodes.cpp
deleted file mode 100644
diff --git a/Engine/Plugins/Experimental/GeometryCollectionPlugin/Source/GeometryCollectionNodes/Private/Dataflow/GeometryCollectionSkeletalMeshToCollectionNode.cpp b/Engine/Plugins/Experimental/GeometryCollectionPlugin/Source/GeometryCollectionNodes/Private/Dataflow/GeometryCollectionSkeletalMeshToCollectionNode.cpp
new file mode 100644
--- /dev/null
+++ b/Engine/Plugins/Experimental/GeometryCollectionPlugin/Source/GeometryCollectionNodes/Private/Dataflow/GeometryCollectionSkeletalMeshToCollectionNode.cpp
@@ -0,0 +1,135 @@
+// Copyright Epic Games, Inc. All Rights Reserved.
+
+#include "Dataflow/GeometryCollectionSkeletalMeshToCollectionNode.h"
+
+#include "Engine/SkeletalMesh.h"
+#include "GeometryCollection/GeometryCollectionObject.h"
+#include "GeometryCollection/ManagedArrayCollection.h"
+#include "GeometryCollection/GeometryCollection.h"
+#include "GeometryCollection/GeometryCollectionEngineConversion.h"
+
+#include "DynamicMesh/DynamicMesh3.h"
+#include "MeshDescriptionToDynamicMesh.h"
+
+#include "MeshDescription.h"
+#include "Rendering/SkeletalMeshLODImporterData.h"
+#include "Rendering/SkeletalMeshModel.h"
+#include "IndexTypes.h"
+#include "GeometryCollection/GeometryCollectionAlgo.h"
+
+void FSkeletalMeshToCollectionDataflowNode::Evaluate(Dataflow::FContext& Context, const FDataflowOutput* Out) const
+{
+	if (Out->IsA<DataType>(&Collection))
+	{
+		FGeometryCollection OutCollection;
+		TObjectPtr<const USkeletalMesh> InSkeletalMesh = GetValue<TObjectPtr<const USkeletalMesh>>(Context, &SkeletalMesh);
+		if (InSkeletalMesh)
+		{
+			if (Method == EDF_ConversionMethod::Skinned)
+			{
+				FGeometryCollectionEngineConversion::AppendSkeletalMesh(InSkeletalMesh, 0, FTransform::Identity, &OutCollection);
+			}
+			else if (Method == EDF_ConversionMethod::StaticComponents)
+			{
+				AppendSkeletalMeshComponentsToGeometryCollection(InSkeletalMesh, OutCollection);
+			}
+		}
+		FManagedArrayCollection OutManagedArrayCollection = OutCollection;
+		SetValue(Context, MoveTemp(OutManagedArrayCollection), &Collection);
+	}
+}
+
+
+
+
+void FSkeletalMeshToCollectionDataflowNode::AppendSkeletalMeshComponentsToGeometryCollection(const USkeletalMesh* InSkeletalMesh, FGeometryCollection& OutCollection) const
+{
+#if WITH_EDITORONLY_DATA
+	UE::Geometry::FDynamicMesh3 DynamicMesh;
+
+	// Check first if we have bulk data available and non-empty.
+	constexpr int32 LODIndex = 0;
+	FMeshDescription SourceMesh;
+	if (InSkeletalMesh->HasMeshDescription(LODIndex))
+	{
+		InSkeletalMesh->CloneMeshDescription(LODIndex, SourceMesh);
+	}
+	FMeshDescriptionToDynamicMesh Converter;
+	Converter.Convert(&SourceMesh, DynamicMesh);
+
+	//
+	// Compute by Component
+	//
+	bool bComputeByComponent = true;
+	if (bComputeByComponent)
+	{
+		TArray<TArray<int32>> ConnectedComponents;
+
+		TArray<FIntVector3> Faces;
+		Faces.SetNum(DynamicMesh.TriangleCount());
+		for (int32 i = 0; i < DynamicMesh.TriangleCount(); ++i)
+		{
+			Faces[i] = FIntVector3(DynamicMesh.GetTriangle(i));
+		}
+		Chaos::Utilities::FindConnectedRegions(Faces, ConnectedComponents);
+
+		TArray<TUniquePtr<FGeometryCollection>> CollectionBuffer;
+		for (int32 i = 0; i < ConnectedComponents.Num(); i++)
+		{
+			CollectionBuffer.Add(nullptr);
+		}
+		ParallelFor(ConnectedComponents.Num(),[&](int32 i)
+		{
+			UE::Geometry::FDynamicMesh3 Component;
+
+			for (FVector V : DynamicMesh.VerticesItr()) {
+				Component.AppendVertex(V);
+			}
+			for (int32 j = 0; j < ConnectedComponents[i].Num(); j++) {
+				int32 ElementIndex = ConnectedComponents[i][j];
+				for (int32 ie = 0; ie < 3; ie++) {
+					Component.AppendTriangle(Faces[ElementIndex][0], Faces[ElementIndex][1], Faces[ElementIndex][2]);
+				}
+			}
+			Component.CompactInPlace();
+			Component.ReverseOrientation();
+
+			int Idx = 0;
+			TArray<float> VertexArray;
+			VertexArray.SetNumUninitialized(Component.VertexCount() * 3);
+			for (FVector V : Component.VerticesItr())
+			{
+				VertexArray[Idx++] = V[0];
+				VertexArray[Idx++] = V[1];
+				VertexArray[Idx++] = V[2];
+			}
+
+			Idx = 0;
+			TArray<int> IndicesArray;
+			IndicesArray.SetNumUninitialized(Component.TriangleCount() * 3);
+			for (UE::Geometry::FIndex3i Tri : Component.TrianglesItr())
+			{
+				IndicesArray[Idx++] = Tri[0];
+				IndicesArray[Idx++] = Tri[1];
+				IndicesArray[Idx++] = Tri[2];
+			}
+
+			CollectionBuffer[i] = TUniquePtr<FGeometryCollection>(FGeometryCollection::NewGeometryCollection(VertexArray, IndicesArray));
+		});
+
+		for (int32 i = 0; i < CollectionBuffer.Num(); i++)
+		{
+			if (CollectionBuffer[i]->NumElements(FGeometryCollection::VerticesGroup))
+			{
+				TSet<int32> VertexToDeleteSet;
+				GeometryCollectionAlgo::ComputeStaleVertices(CollectionBuffer[i].Get(), VertexToDeleteSet);
+				TArray<int32> SortedVertices = VertexToDeleteSet.Array(); SortedVertices.Sort();
+				if (VertexToDeleteSet.Num()) CollectionBuffer[i]->RemoveElements(FGeometryCollection::VerticesGroup, SortedVertices);
+
+				int32 GeomIndex = OutCollection.AppendGeometry(*CollectionBuffer[i].Get());
+			}
+		}
+	}
+#endif
+}
+
diff --git a/Engine/Plugins/Experimental/GeometryCollectionPlugin/Source/GeometryCollectionNodes/Private/Dataflow/GeometryCollectionSkeletonToCollectionNode.cpp b/Engine/Plugins/Experimental/GeometryCollectionPlugin/Source/GeometryCollectionNodes/Private/Dataflow/GeometryCollectionSkeletonToCollectionNode.cpp
new file mode 100644
--- /dev/null
+++ b/Engine/Plugins/Experimental/GeometryCollectionPlugin/Source/GeometryCollectionNodes/Private/Dataflow/GeometryCollectionSkeletonToCollectionNode.cpp
@@ -0,0 +1,25 @@
+// Copyright Epic Games, Inc. All Rights Reserved.
+
+#include "Dataflow/GeometryCollectionSkeletonToCollectionNode.h"
+
+#include "Engine/SkeletalMesh.h"
+#include "GeometryCollection/GeometryCollectionObject.h"
+#include "GeometryCollection/ManagedArrayCollection.h"
+#include "GeometryCollection/GeometryCollection.h"
+#include "GeometryCollection/GeometryCollectionEngineConversion.h"
+
+void FSkeletonToCollectionDataflowNode::Evaluate(Dataflow::FContext& Context, const FDataflowOutput* Out) const
+{
+	if (Out->IsA<DataType>(&Collection))
+	{
+		DataType OutCollection;
+		TObjectPtr<const USkeleton> SkeletonValue = GetValue<TObjectPtr<const USkeleton>>(Context, &Skeleton);
+		if (SkeletonValue)
+		{
+			FGeometryCollectionEngineConversion::AppendSkeleton(SkeletonValue.Get(), FTransform::Identity, &OutCollection);
+		}
+		SetValue(Context, MoveTemp(OutCollection), &Collection);
+	}
+}
+
+
diff --git a/Engine/Plugins/Experimental/GeometryCollectionPlugin/Source/GeometryCollectionNodes/Public/Dataflow/GeometryCollectionSkeletalMeshNodes.h b/Engine/Plugins/Experimental/GeometryCollectionPlugin/Source/GeometryCollectionNodes/Public/Dataflow/GeometryCollectionSkeletalMeshNodes.h
deleted file mode 100644
diff --git a/Engine/Plugins/Experimental/GeometryCollectionPlugin/Source/GeometryCollectionNodes/Public/Dataflow/GeometryCollectionSkeletalMeshToCollectionNode.h b/Engine/Plugins/Experimental/GeometryCollectionPlugin/Source/GeometryCollectionNodes/Public/Dataflow/GeometryCollectionSkeletalMeshToCollectionNode.h
new file mode 100644
--- /dev/null
+++ b/Engine/Plugins/Experimental/GeometryCollectionPlugin/Source/GeometryCollectionNodes/Public/Dataflow/GeometryCollectionSkeletalMeshToCollectionNode.h
@@ -0,0 +1,52 @@
+// Copyright Epic Games, Inc. All Rights Reserved.
+
+#pragma once 
+
+#include "CoreMinimal.h"
+#include "Dataflow/DataflowEngine.h"
+#include "GeometryCollection/ManagedArrayCollection.h"
+
+#include "GeometryCollectionSkeletalMeshToCollectionNode.generated.h"
+
+class USkeletalMesh;
+
+UENUM()
+enum EDF_ConversionMethod : int
+{
+	Skinned				UMETA(DisplayName = "Skinned"),
+	StaticComponents	UMETA(DisplayName = "Static Components"),
+};
+
+USTRUCT(meta = (DataflowGeometryCollection))
+struct FSkeletalMeshToCollectionDataflowNode : public FDataflowNode
+{
+	GENERATED_USTRUCT_BODY()
+	DATAFLOW_NODE_DEFINE_INTERNAL(FSkeletalMeshToCollectionDataflowNode, "SkeletalMeshToCollection", "GeometryCollection", "")
+	DATAFLOW_NODE_RENDER_TYPE(FGeometryCollection::StaticType(), "Collection")
+
+public:
+	typedef FManagedArrayCollection DataType;
+
+	UPROPERTY(EditAnywhere, Category = "Dataflow")
+	TEnumAsByte<EDF_ConversionMethod> Method = EDF_ConversionMethod::StaticComponents;
+
+	UPROPERTY(EditAnywhere, Category = "Dataflow", meta = (DataflowInput, DisplayName = "SkeletalMesh"))
+	TObjectPtr<const USkeletalMesh> SkeletalMesh = nullptr;
+
+	UPROPERTY(meta = (DataflowOutput, DisplayName = "Collection"))
+	FManagedArrayCollection Collection;
+
+	FSkeletalMeshToCollectionDataflowNode(const Dataflow::FNodeParameters& InParam, FGuid InGuid = FGuid::NewGuid())
+		: FDataflowNode(InParam, InGuid)
+	{
+		RegisterInputConnection(&SkeletalMesh);
+		RegisterOutputConnection(&Collection);
+	}
+
+	virtual void Evaluate(Dataflow::FContext& Context, const FDataflowOutput* Out) const override;
+
+private:
+	void AppendSkeletalMeshComponentsToGeometryCollection(const USkeletalMesh* InSkeletalMesh, FGeometryCollection& OutCollection) const;
+
+};
+
diff --git a/Engine/Plugins/Experimental/GeometryCollectionPlugin/Source/GeometryCollectionNodes/Public/Dataflow/GeometryCollectionSkeletonToCollectionNode.h b/Engine/Plugins/Experimental/GeometryCollectionPlugin/Source/GeometryCollectionNodes/Public/Dataflow/GeometryCollectionSkeletonToCollectionNode.h
new file mode 100644
--- /dev/null
+++ b/Engine/Plugins/Experimental/GeometryCollectionPlugin/Source/GeometryCollectionNodes/Public/Dataflow/GeometryCollectionSkeletonToCollectionNode.h
@@ -0,0 +1,38 @@
+// Copyright Epic Games, Inc. All Rights Reserved.
+
+#pragma once 
+
+#include "CoreMinimal.h"
+#include "Dataflow/DataflowEngine.h"
+#include "GeometryCollection/ManagedArrayCollection.h"
+
+#include "GeometryCollectionSkeletonToCollectionNode.generated.h"
+
+class USkeletalMesh;
+
+USTRUCT(meta = (DataflowGeometryCollection))
+struct FSkeletonToCollectionDataflowNode : public FDataflowNode
+{
+	GENERATED_USTRUCT_BODY()
+	DATAFLOW_NODE_DEFINE_INTERNAL(FSkeletonToCollectionDataflowNode, "SkeletonToCollection", "GeometryCollection", "")
+
+public:
+	typedef FManagedArrayCollection DataType;
+
+	UPROPERTY(EditAnywhere, Category = "Dataflow", meta = (DataflowInput, DisplayName = "Skeleton"))
+	TObjectPtr<const USkeleton> Skeleton = nullptr;
+
+	UPROPERTY(meta = (DataflowOutput, DisplayName = "Collection"))
+	FManagedArrayCollection Collection;
+
+	FSkeletonToCollectionDataflowNode(const Dataflow::FNodeParameters& InParam, FGuid InGuid = FGuid::NewGuid())
+		: FDataflowNode(InParam, InGuid)
+	{
+		RegisterInputConnection(&Skeleton);
+		RegisterOutputConnection(&Collection);
+	}
+
+	virtual void Evaluate(Dataflow::FContext& Context, const FDataflowOutput* Out) const override;
+};
+
+
