GeometryCollection : Fix for object slicing

diff --git a/Engine/Plugins/Experimental/GeometryCollectionPlugin/Source/GeometryCollectionNodes/Private/Dataflow/GeometryCollectionSkeletalMeshToCollectionNode.cpp b/Engine/Plugins/Experimental/GeometryCollectionPlugin/Source/GeometryCollectionNodes/Private/Dataflow/GeometryCollectionSkeletalMeshToCollectionNode.cpp
--- a/Engine/Plugins/Experimental/GeometryCollectionPlugin/Source/GeometryCollectionNodes/Private/Dataflow/GeometryCollectionSkeletalMeshToCollectionNode.cpp
+++ b/Engine/Plugins/Experimental/GeometryCollectionPlugin/Source/GeometryCollectionNodes/Private/Dataflow/GeometryCollectionSkeletalMeshToCollectionNode.cpp
@@ -34,8 +34,7 @@ void FSkeletalMeshToCollectionDataflowNode::Evaluate(Dataflow::FContext& Context
 				AppendSkeletalMeshComponentsToGeometryCollection(InSkeletalMesh, OutCollection);
 			}
 		}
-		FManagedArrayCollection OutManagedArrayCollection = OutCollection;
-		SetValue(Context, MoveTemp(OutManagedArrayCollection), &Collection);
+		SetValue(Context, FManagedArrayCollection(OutCollection), &Collection);
 	}
 }
 
