GeometryCollection : Transfer Vertex Attributes between collections.

diff --git a/Engine/Plugins/Experimental/GeometryCollectionPlugin/Source/GeometryCollectionNodes/Private/Dataflow/GeometryCollectionNodesPlugin.cpp b/Engine/Plugins/Experimental/GeometryCollectionPlugin/Source/GeometryCollectionNodes/Private/Dataflow/GeometryCollectionNodesPlugin.cpp
--- a/Engine/Plugins/Experimental/GeometryCollectionPlugin/Source/GeometryCollectionNodes/Private/Dataflow/GeometryCollectionNodesPlugin.cpp
+++ b/Engine/Plugins/Experimental/GeometryCollectionPlugin/Source/GeometryCollectionNodes/Private/Dataflow/GeometryCollectionNodesPlugin.cpp
@@ -26,6 +26,7 @@
 #include "Dataflow/GeometryCollectionArrayNodes.h"
 #include "Dataflow/GeometryCollectionDebugNodes.h"
 #include "Dataflow/GeometryCollectionVertexScalarToVertexIndicesNode.h"
+#include "Dataflow/GeometryCollectionTransferVertexScalarAttributeNode.h"
 
 #define LOCTEXT_NAMESPACE "DataflowNodes"
 
@@ -53,6 +54,7 @@ void IGeometryCollectionNodesPlugin::StartupModule()
 	DATAFLOW_NODE_REGISTER_CREATION_FACTORY(FGeometryCollectionVertexScalarToVertexIndicesNode);
 	DATAFLOW_NODE_REGISTER_CREATION_FACTORY(FSkeletonToCollectionDataflowNode);
 	DATAFLOW_NODE_REGISTER_CREATION_FACTORY(FSkeletalMeshToCollectionDataflowNode);
+	DATAFLOW_NODE_REGISTER_CREATION_FACTORY(FGeometryCollectionTransferVertexScalarAttributeNode);
 }
 
 void IGeometryCollectionNodesPlugin::ShutdownModule()
diff --git a/Engine/Plugins/Experimental/GeometryCollectionPlugin/Source/GeometryCollectionNodes/Private/Dataflow/GeometryCollectionTransferVertexScalarAttributeNode.cpp b/Engine/Plugins/Experimental/GeometryCollectionPlugin/Source/GeometryCollectionNodes/Private/Dataflow/GeometryCollectionTransferVertexScalarAttributeNode.cpp
new file mode 100644
--- /dev/null
+++ b/Engine/Plugins/Experimental/GeometryCollectionPlugin/Source/GeometryCollectionNodes/Private/Dataflow/GeometryCollectionTransferVertexScalarAttributeNode.cpp
@@ -0,0 +1,191 @@
+// Copyright Epic Games, Inc. All Rights Reserved.
+
+#include "Dataflow/GeometryCollectionTransferVertexScalarAttributeNode.h"
+
+#include "Chaos/BoundingVolumeHierarchy.h"
+#include "Chaos/Triangle.h"
+#include "Dataflow/DataflowInputOutput.h"
+#include "GeometryCollection/GeometryCollectionAlgo.h"
+#include "GeometryCollection/TransformCollection.h"
+
+
+#include UE_INLINE_GENERATED_CPP_BY_NAME(GeometryCollectionTransferVertexScalarAttributeNode)
+#define LOCTEXT_NAMESPACE "FGeometryCollectionTransferVertexScalarAttributeNode"
+
+
+
+void FGeometryCollectionTransferVertexScalarAttributeNode::Evaluate(Dataflow::FContext& Context, const FDataflowOutput* Out) const
+{
+	auto BuildComponentSpaceVertices = [](const TManagedArray<FTransform3f>* LocalSpaceTransform,
+		const TManagedArray<int32>* Parent, const TManagedArray<int32>* BoneMapArray,
+		const TManagedArray<FVector3f>* VertexArray,
+		TArray<FVector>& ComponentSpaceVertices)
+	{
+                                              		TArray<FTransform> ComponentTransform;
+		GeometryCollectionAlgo::GlobalMatrices(*LocalSpaceTransform, *Parent, ComponentTransform);
+		
+		ComponentSpaceVertices.SetNumUninitialized(VertexArray->Num());
+		for (int i = 0; i < VertexArray->Num(); i++)
+		{
+			if (0 < (*BoneMapArray)[i] && (*BoneMapArray)[i] < ComponentTransform.Num())
+			{
+				ComponentSpaceVertices[i] = ComponentTransform[(*BoneMapArray)[i]].TransformPosition(FVector((*VertexArray)[i]));
+			}
+			else
+			{
+				ComponentSpaceVertices[i] = FVector((*VertexArray)[i]);
+			}
+		}
+	};
+
+	FString AttributName = GetValue<FString>(Context, &Name, Name);
+	if (Out->IsA< FManagedArrayCollection >(&Collection))
+	{
+		FManagedArrayCollection CollectionVal = GetValue<FManagedArrayCollection>(Context, &Collection);
+		const FManagedArrayCollection& AttributeCollectionVal = GetValue<FManagedArrayCollection>(Context, &FromCollection);
+
+		const TManagedArray<int32>* TargetBoneMapArray = CollectionVal.FindAttribute<int32>(FName("BoneMap"), FName("Vertices"));
+		const TManagedArray<FVector3f>* TargetVertexArray = CollectionVal.FindAttribute<FVector3f>(FName("Vertex"), FName("Vertices"));
+		const TManagedArray<FIntVector3>* TargetIndicesArray = CollectionVal.FindAttribute<FIntVector3>(FName("Indices"), FName("Faces"));
+		const TManagedArray<FTransform3f>* TargetLocalSpaceTransform = CollectionVal.FindAttribute<FTransform3f>(FTransformCollection::TransformAttribute, FTransformCollection::TransformGroup);
+		const TManagedArray<int32>* TargetParent = CollectionVal.FindAttribute<int32>(FTransformCollection::ParentAttribute, FTransformCollection::TransformGroup);
+
+		const TManagedArray<float>* FloatArray = AttributeCollectionVal.FindAttribute<float>(FName(AttributName), FName("Vertices"));
+		const TManagedArray<int32>* BoneMapArray = AttributeCollectionVal.FindAttribute<int32>(FName("BoneMap"), FName("Vertices"));
+		const TManagedArray<FVector3f>* VertexArray = AttributeCollectionVal.FindAttribute<FVector3f>(FName("Vertex"), FName("Vertices"));
+		const TManagedArray<FIntVector3>* IndicesArray = AttributeCollectionVal.FindAttribute<FIntVector3>(FName("Indices"), FName("Faces"));
+		const TManagedArray<FTransform3f>* LocalSpaceTransform = AttributeCollectionVal.FindAttribute<FTransform3f>(FTransformCollection::TransformAttribute, FTransformCollection::TransformGroup);
+		const TManagedArray<int32>* Parent = AttributeCollectionVal.FindAttribute<int32>(FTransformCollection::ParentAttribute, FTransformCollection::TransformGroup);
+
+		if (TargetBoneMapArray && TargetVertexArray && TargetIndicesArray && TargetLocalSpaceTransform && TargetParent)
+		{
+			if (FloatArray && BoneMapArray && VertexArray && IndicesArray && LocalSpaceTransform && Parent)
+			{
+				TManagedArray<float>* TargetFloatArray = CollectionVal.FindAttributeTyped<float>(FName(AttributName), FName("Vertices"));
+				if (!TargetFloatArray && !CollectionVal.HasAttribute(FName(AttributName), FName("Vertices")))
+				{
+					CollectionVal.AddAttribute<float>(FName(AttributName), FName("Vertices"));
+					TargetFloatArray = CollectionVal.FindAttribute<float>(FName(AttributName), FName("Vertices"));
+				}
+
+				if(TargetFloatArray)
+				{
+					TargetFloatArray->Fill(0.f);
+
+					// Build component space vertices for TargetCollection
+					TArray<FVector> ComponentSpaceTargetVertices;
+					BuildComponentSpaceVertices(TargetLocalSpaceTransform, TargetParent, TargetBoneMapArray, TargetVertexArray, ComponentSpaceTargetVertices);
+
+					// Build component space vertices for TargetCollection
+					TArray<FVector> ComponentSpaceVertices;
+					BuildComponentSpaceVertices(LocalSpaceTransform, Parent, BoneMapArray, VertexArray, ComponentSpaceVertices);
+
+
+					// build Sphere based BVH
+					Chaos::FReal SphereRadius = (Chaos::FReal)0.;
+
+					Chaos::TVec3<float> CoordMaxs(-FLT_MAX);
+					Chaos::TVec3<float> CoordMins(FLT_MAX);
+					for (int32 i = 0; i < ComponentSpaceTargetVertices.Num(); i++)
+					{
+						for (int32 j = 0; j < 3; j++)
+						{
+							if (ComponentSpaceTargetVertices[i][j] > CoordMaxs[j])
+							{
+								CoordMaxs[j] = ComponentSpaceTargetVertices[i][j];
+							}
+							if (ComponentSpaceTargetVertices[i][j] < CoordMins[j])
+							{
+								CoordMins[j] = ComponentSpaceTargetVertices[i][j];
+							}
+						}
+					}
+					Chaos::TVec3<float> CoordDiff = (CoordMaxs - CoordMins) * VertexRadiusRatio;
+					SphereRadius = Chaos::FReal(FGenericPlatformMath::Min(CoordDiff[0], FGenericPlatformMath::Min(CoordDiff[1], CoordDiff[2])));
+
+					TArray<Chaos::TSphere<Chaos::FReal, 3>*> VertexSpherePtrs;
+					TArray<Chaos::TSphere<Chaos::FReal, 3>> VertexSpheres;
+
+					VertexSpheres.Init(Chaos::TSphere<Chaos::FReal, 3>(Chaos::TVec3<Chaos::FReal>(0), SphereRadius), ComponentSpaceTargetVertices.Num());
+					VertexSpherePtrs.SetNum(ComponentSpaceTargetVertices.Num());
+
+					for (int32 i = 0; i < ComponentSpaceTargetVertices.Num(); i++)
+					{
+						Chaos::TVec3<Chaos::FReal> SphereCenter(ComponentSpaceTargetVertices[i]);
+						Chaos::TSphere<Chaos::FReal, 3> VertexSphere(SphereCenter, SphereRadius);
+						VertexSpheres[i] = Chaos::TSphere<Chaos::FReal, 3>(SphereCenter, SphereRadius);
+						VertexSpherePtrs[i] = &VertexSpheres[i];
+					}
+					Chaos::TBoundingVolumeHierarchy<
+						TArray<Chaos::TSphere<Chaos::FReal, 3>*>,
+						TArray<int32>,
+						Chaos::FReal,
+						3> VertexBVH(VertexSpherePtrs);
+
+					for (int32 i = 0; i < IndicesArray->Num(); i++)
+					{
+						TArray<int32> TriangleIntersections0 = VertexBVH.FindAllIntersections(ComponentSpaceVertices[(*IndicesArray)[i][0]]);
+						TArray<int32> TriangleIntersections1 = VertexBVH.FindAllIntersections(ComponentSpaceVertices[(*IndicesArray)[i][1]]);
+						TArray<int32> TriangleIntersections2 = VertexBVH.FindAllIntersections(ComponentSpaceVertices[(*IndicesArray)[i][2]]);
+						TriangleIntersections0.Sort();
+						TriangleIntersections1.Sort();
+						TriangleIntersections2.Sort();
+
+						TArray<int32> TriangleIntersections({});
+						for (int32 k = 0; k < TriangleIntersections0.Num(); k++)
+						{
+							if (TriangleIntersections1.Contains(TriangleIntersections0[k])
+								&& TriangleIntersections2.Contains(TriangleIntersections0[k]))
+							{
+								TriangleIntersections.Emplace(TriangleIntersections0[k]);
+							}
+						}
+
+						int32 MinIndex = -1;
+						float MinDis = SphereRadius;
+						Chaos::TVector<float, 3> ClosestBary(0.f);
+
+						for (int32 j = 0; j < TriangleIntersections.Num(); j++)
+						{
+							Chaos::TVector<float, 3> Bary,
+								TriPos0(ComponentSpaceVertices[(*IndicesArray)[i][0]]),
+								TriPos1(ComponentSpaceVertices[(*IndicesArray)[i][1]]),
+								TriPos2(ComponentSpaceVertices[(*IndicesArray)[i][2]]),
+								ParticlePos(ComponentSpaceTargetVertices[TriangleIntersections[j]]);
+
+							Chaos::TVector<Chaos::FRealSingle, 3> ClosestPoint = Chaos::FindClosestPointAndBaryOnTriangle(TriPos0, TriPos1, TriPos2, ParticlePos, Bary);
+							Chaos::FRealSingle CurrentDistance = (ParticlePos - ClosestPoint).Size();
+							if (CurrentDistance < MinDis)
+							{
+								MinDis = CurrentDistance;
+								MinIndex = TriangleIntersections[j];
+								ClosestBary = Bary;
+							}
+
+							if (MinIndex != -1
+								&& MinIndex != (*IndicesArray)[i][0]
+								&& MinIndex != (*IndicesArray)[i][1]
+								&& MinIndex != (*IndicesArray)[i][2])
+							{
+								(*TargetFloatArray)[MinIndex] = 0;
+								for (int32 k = 0; k < 3; k++)
+								{
+									(*TargetFloatArray)[MinIndex] += ClosestBary[k] * (*FloatArray)[(*IndicesArray)[i][k]];
+								}
+							}
+						}
+					}
+				}
+			}
+		}
+
+		SetValue<FManagedArrayCollection >(Context, MoveTemp(CollectionVal), &Collection);
+	}
+	else if (Out->IsA<FString>(&Name))
+	{
+	SetValue< FString >(Context, MoveTemp(AttributName), &Name);
+	}
+}
+
+
+#undef LOCTEXT_NAMESPACE
diff --git a/Engine/Plugins/Experimental/GeometryCollectionPlugin/Source/GeometryCollectionNodes/Public/Dataflow/GeometryCollectionTransferVertexScalarAttributeNode.h b/Engine/Plugins/Experimental/GeometryCollectionPlugin/Source/GeometryCollectionNodes/Public/Dataflow/GeometryCollectionTransferVertexScalarAttributeNode.h
new file mode 100644
--- /dev/null
+++ b/Engine/Plugins/Experimental/GeometryCollectionPlugin/Source/GeometryCollectionNodes/Public/Dataflow/GeometryCollectionTransferVertexScalarAttributeNode.h
@@ -0,0 +1,46 @@
+// Copyright Epic Games, Inc. All Rights Reserved.
+
+#pragma once 
+
+#include "Dataflow/DataflowNode.h"
+#include "GeometryCollection/ManagedArrayCollection.h"
+
+#include "GeometryCollectionTransferVertexScalarAttributeNode.generated.h"
+
+/** Convert an vertex float array to a list of indices */
+USTRUCT(meta = (DataflowGeometryCollection))
+struct FGeometryCollectionTransferVertexScalarAttributeNode : public FDataflowNode
+{
+	GENERATED_USTRUCT_BODY()
+	DATAFLOW_NODE_DEFINE_INTERNAL(FGeometryCollectionTransferVertexScalarAttributeNode, "TransferVertexScalarAttribute", "GeometryCollection", "Transfer a named scalar attribute from the Source Collection to the Target Collection")
+
+public:
+
+	UPROPERTY(Meta = (DataflowInput, DataflowOutput, DisplayName = "Collection", DataflowPassthrough = "Collection"))
+	FManagedArrayCollection Collection;
+
+	UPROPERTY(Meta = (DataflowInput, DisplayName = "FromCollection"))
+	FManagedArrayCollection FromCollection;
+
+	/** The name of the vertex attribute to generate indices from. */
+	UPROPERTY(EditAnywhere, Category = "Dataflow", Meta = (DataflowInput, DataflowOutput, DisplayName = "Name", DataflowPassthrough = "Name"))
+	FString Name = FString("");
+
+	UPROPERTY(EditAnywhere, Category = "Dataflow", meta = (ClampMin = "0.0", ClampMax = "2.0", UIMin = "0.0", UIMax = "2.0"))
+	float VertexRadiusRatio = .001f;
+
+	FGeometryCollectionTransferVertexScalarAttributeNode(const Dataflow::FNodeParameters& InParam, FGuid InGuid = FGuid::NewGuid())
+		: FDataflowNode(InParam, InGuid)
+	{
+		RegisterInputConnection(&Collection);
+		RegisterInputConnection(&FromCollection);
+		RegisterInputConnection(&Name);
+
+		RegisterOutputConnection(&Collection, &Collection);
+		RegisterOutputConnection(&Name, &Name);
+	}
+
+private:
+
+	virtual void Evaluate(Dataflow::FContext& Context, const FDataflowOutput* Out) const override;
+};
