GeometryCollection : Transfer attributes from source geometry.

diff --git a/Engine/Plugins/Experimental/GeometryCollectionPlugin/Source/GeometryCollectionNodes/Private/Dataflow/GeometryCollectionTransferVertexScalarAttributeNode.cpp b/Engine/Plugins/Experimental/GeometryCollectionPlugin/Source/GeometryCollectionNodes/Private/Dataflow/GeometryCollectionTransferVertexScalarAttributeNode.cpp
--- a/Engine/Plugins/Experimental/GeometryCollectionPlugin/Source/GeometryCollectionNodes/Private/Dataflow/GeometryCollectionTransferVertexScalarAttributeNode.cpp
+++ b/Engine/Plugins/Experimental/GeometryCollectionPlugin/Source/GeometryCollectionNodes/Private/Dataflow/GeometryCollectionTransferVertexScalarAttributeNode.cpp
@@ -7,6 +7,7 @@
 #include "Dataflow/DataflowInputOutput.h"
 #include "GeometryCollection/GeometryCollectionAlgo.h"
 #include "GeometryCollection/TransformCollection.h"
+#include "GeometryCollection/GeometryCollection.h"
 
 
 #include UE_INLINE_GENERATED_CPP_BY_NAME(GeometryCollectionTransferVertexScalarAttributeNode)
@@ -16,27 +17,7 @@
 
 void FGeometryCollectionTransferVertexScalarAttributeNode::Evaluate(Dataflow::FContext& Context, const FDataflowOutput* Out) const
 {
-	auto BuildComponentSpaceVertices = [](const TManagedArray<FTransform3f>* LocalSpaceTransform,
-		const TManagedArray<int32>* Parent, const TManagedArray<int32>* BoneMapArray,
-		const TManagedArray<FVector3f>* VertexArray,
-		TArray<FVector>& ComponentSpaceVertices)
-	{
-                                              		TArray<FTransform> ComponentTransform;
-		GeometryCollectionAlgo::GlobalMatrices(*LocalSpaceTransform, *Parent, ComponentTransform);
-		
-		ComponentSpaceVertices.SetNumUninitialized(VertexArray->Num());
-		for (int i = 0; i < VertexArray->Num(); i++)
-		{
-			if (0 < (*BoneMapArray)[i] && (*BoneMapArray)[i] < ComponentTransform.Num())
-			{
-				ComponentSpaceVertices[i] = ComponentTransform[(*BoneMapArray)[i]].TransformPosition(FVector((*VertexArray)[i]));
-			}
-			else
-			{
-				ComponentSpaceVertices[i] = FVector((*VertexArray)[i]);
-			}
-		}
-	};
+
 
 	FString AttributName = GetValue<FString>(Context, &Name, Name);
 	if (Out->IsA< FManagedArrayCollection >(&Collection))
@@ -49,6 +30,10 @@ void FGeometryCollectionTransferVertexScalarAttributeNode::Evaluate(Dataflow::FC
 		const TManagedArray<FIntVector3>* TargetIndicesArray = CollectionVal.FindAttribute<FIntVector3>(FName("Indices"), FName("Faces"));
 		const TManagedArray<FTransform3f>* TargetLocalSpaceTransform = CollectionVal.FindAttribute<FTransform3f>(FTransformCollection::TransformAttribute, FTransformCollection::TransformGroup);
 		const TManagedArray<int32>* TargetParent = CollectionVal.FindAttribute<int32>(FTransformCollection::ParentAttribute, FTransformCollection::TransformGroup);
+		const TManagedArray<int32>* TargetVertexStart = CollectionVal.FindAttribute<int32>("VertexStart", FGeometryCollection::GeometryGroup);
+		const TManagedArray<int32>* TargetVertexCount = CollectionVal.FindAttribute<int32>("VertexCount", FGeometryCollection::GeometryGroup);
+		const TManagedArray<int32>* TargetFaceStart = CollectionVal.FindAttribute<int32>("FaceStart", FGeometryCollection::GeometryGroup);
+		const TManagedArray<int32>* TargetFaceCount = CollectionVal.FindAttribute<int32>("FaceCount", FGeometryCollection::GeometryGroup);
 
 		const TManagedArray<float>* FloatArray = AttributeCollectionVal.FindAttribute<float>(FName(AttributName), FName("Vertices"));
 		const TManagedArray<int32>* BoneMapArray = AttributeCollectionVal.FindAttribute<int32>(FName("BoneMap"), FName("Vertices"));
@@ -56,10 +41,14 @@ void FGeometryCollectionTransferVertexScalarAttributeNode::Evaluate(Dataflow::FC
 		const TManagedArray<FIntVector3>* IndicesArray = AttributeCollectionVal.FindAttribute<FIntVector3>(FName("Indices"), FName("Faces"));
 		const TManagedArray<FTransform3f>* LocalSpaceTransform = AttributeCollectionVal.FindAttribute<FTransform3f>(FTransformCollection::TransformAttribute, FTransformCollection::TransformGroup);
 		const TManagedArray<int32>* Parent = AttributeCollectionVal.FindAttribute<int32>(FTransformCollection::ParentAttribute, FTransformCollection::TransformGroup);
+		const TManagedArray<int32>* VertexStart = AttributeCollectionVal.FindAttribute<int32>("VertexStart", FGeometryCollection::GeometryGroup);
+		const TManagedArray<int32>* VertexCount = AttributeCollectionVal.FindAttribute<int32>("VertexCount", FGeometryCollection::GeometryGroup);
+		const TManagedArray<int32>* FaceStart = AttributeCollectionVal.FindAttribute<int32>("FaceStart", FGeometryCollection::GeometryGroup);
+		const TManagedArray<int32>* FaceCount = AttributeCollectionVal.FindAttribute<int32>("FaceCount", FGeometryCollection::GeometryGroup);
 
-		if (TargetBoneMapArray && TargetVertexArray && TargetIndicesArray && TargetLocalSpaceTransform && TargetParent)
+		if (TargetBoneMapArray && TargetVertexArray && TargetIndicesArray && TargetLocalSpaceTransform && TargetParent && TargetVertexStart && TargetVertexCount && TargetFaceStart && TargetFaceCount)
 		{
-			if (FloatArray && BoneMapArray && VertexArray && IndicesArray && LocalSpaceTransform && Parent)
+			if (FloatArray && BoneMapArray && VertexArray && IndicesArray && LocalSpaceTransform && Parent && VertexStart && VertexCount && FaceStart && FaceCount)
 			{
 				TManagedArray<float>* TargetFloatArray = CollectionVal.FindAttributeTyped<float>(FName(AttributName), FName("Vertices"));
 				if (!TargetFloatArray && !CollectionVal.HasAttribute(FName(AttributName), FName("Vertices")))
@@ -72,120 +61,370 @@ void FGeometryCollectionTransferVertexScalarAttributeNode::Evaluate(Dataflow::FC
 				{
 					TargetFloatArray->Fill(0.f);
 
-					// Build component space vertices for TargetCollection
-					TArray<FVector> ComponentSpaceTargetVertices;
-					BuildComponentSpaceVertices(TargetLocalSpaceTransform, TargetParent, TargetBoneMapArray, TargetVertexArray, ComponentSpaceTargetVertices);
 
-					// Build component space vertices for TargetCollection
-					TArray<FVector> ComponentSpaceVertices;
-					BuildComponentSpaceVertices(LocalSpaceTransform, Parent, BoneMapArray, VertexArray, ComponentSpaceVertices);
+					TArray<FIntVector2> AlignedGeometry = FindSourceToTargetGeometryMap(AttributeCollectionVal, CollectionVal);
+					if (AlignedGeometry.Num() == AttributeCollectionVal.NumElements(FGeometryCollection::GeometryGroup))
+					{
+						PairedGeometryTransfer(AttributName, AlignedGeometry, AttributeCollectionVal,CollectionVal,TargetFloatArray);
+					}
+					else
+					{
+						NearestVertexTransfer(AttributName, AttributeCollectionVal, CollectionVal, TargetFloatArray);
+					}
+				}
+			}
+		}
+
+		SetValue<FManagedArrayCollection >(Context, MoveTemp(CollectionVal), &Collection);
+	}
+	else if (Out->IsA<FString>(&Name))
+	{
+	SetValue< FString >(Context, MoveTemp(AttributName), &Name);
+	}
+}
+
+TArray<FIntVector2> FGeometryCollectionTransferVertexScalarAttributeNode::FindSourceToTargetGeometryMap(const FManagedArrayCollection& AttributeCollectionVal, const FManagedArrayCollection& CollectionVal) const
+{
+	TArray<FIntVector2> Mapping;
+	const TManagedArray<FString>* SourceName = AttributeCollectionVal.FindAttribute<FString>(FName("BoneName"), FTransformCollection::TransformGroup);
+	const TManagedArray<int32>* SourceGeometryGroup = AttributeCollectionVal.FindAttribute<int32>(FName("TransformToGeometryIndex"), FTransformCollection::TransformGroup);
+	const TManagedArray<FString>* TargetName = CollectionVal.FindAttribute<FString>(FName("BoneName"), FTransformCollection::TransformGroup);
+	const TManagedArray<int32>* TargetGeometryGroup = CollectionVal.FindAttribute<int32>(FName("TransformToGeometryIndex"), FTransformCollection::TransformGroup);
+	if (SourceName && SourceGeometryGroup && TargetName && TargetGeometryGroup)
+	{
+		for (int i = 0; i < SourceName->Num(); i++)
+		{
+			for (int j = 0; j < TargetName->Num(); j++)
+			{
+				FString TestName = FString::Printf(TEXT("%s_Tet"), *(*SourceName)[i]);
+				if ((*TargetName)[j].StartsWith(TestName))
+				{
+					Mapping.Add(FIntVector2((*SourceGeometryGroup)[i], (*TargetGeometryGroup)[j]));
+					break;
+				}
+			}
+		}
+	}
+
+	return Mapping;
+}
+
+
+void FGeometryCollectionTransferVertexScalarAttributeNode::PairedGeometryTransfer(FString AttributName, const TArray<FIntVector2>& PairedGeometry, const FManagedArrayCollection& AttributeCollectionVal, const FManagedArrayCollection& CollectionVal, TManagedArray<float>* TargetFloatArray) const
+{
+	auto BuildComponentSpaceVertices = [](const TManagedArray<FTransform3f>* LocalSpaceTransform,
+		const TManagedArray<int32>* Parent, const TManagedArray<int32>* BoneMapArray,
+		const TManagedArray<FVector3f>* VertexArray,
+		int32 Start, int32 Count,
+		TArray<FVector>& ComponentSpaceVertices)
+	{
+		int32 End = Start + Count;
+		TArray<FTransform> ComponentTransform;
+		GeometryCollectionAlgo::GlobalMatrices(*LocalSpaceTransform, *Parent, ComponentTransform);
+
+		ComponentSpaceVertices.SetNumUninitialized(Count);
+		for (int i = 0; i < Count; i++)
+		{
+			int j = i + Start;
+			if (0 < (*BoneMapArray)[i] && (*BoneMapArray)[i] < ComponentTransform.Num())
+			{
+				ComponentSpaceVertices[i] = ComponentTransform[(*BoneMapArray)[j]].TransformPosition(FVector((*VertexArray)[j]));
+			}
+			else
+			{
+				ComponentSpaceVertices[i] = FVector((*VertexArray)[j]);
+			}
+		}
+	};
+
+	const TManagedArray<int32>* TargetBoneMapArray = CollectionVal.FindAttribute<int32>(FName("BoneMap"), FName("Vertices"));
+	const TManagedArray<FVector3f>* TargetVertexArray = CollectionVal.FindAttribute<FVector3f>(FName("Vertex"), FName("Vertices"));
+	const TManagedArray<FIntVector3>* TargetIndicesArray = CollectionVal.FindAttribute<FIntVector3>(FName("Indices"), FName("Faces"));
+	const TManagedArray<FTransform3f>* TargetLocalSpaceTransform = CollectionVal.FindAttribute<FTransform3f>(FTransformCollection::TransformAttribute, FTransformCollection::TransformGroup);
+	const TManagedArray<int32>* TargetParent = CollectionVal.FindAttribute<int32>(FTransformCollection::ParentAttribute, FTransformCollection::TransformGroup);
+	const TManagedArray<int32>* TargetVertexStart = CollectionVal.FindAttribute<int32>("VertexStart", FGeometryCollection::GeometryGroup);
+	const TManagedArray<int32>* TargetVertexCount = CollectionVal.FindAttribute<int32>("VertexCount", FGeometryCollection::GeometryGroup);
+	const TManagedArray<int32>* TargetFaceStart = CollectionVal.FindAttribute<int32>("FaceStart", FGeometryCollection::GeometryGroup);
+	const TManagedArray<int32>* TargetFaceCount = CollectionVal.FindAttribute<int32>("FaceCount", FGeometryCollection::GeometryGroup);
+
+	const TManagedArray<float>* FloatArray = AttributeCollectionVal.FindAttribute<float>(FName(AttributName), FName("Vertices"));
+	const TManagedArray<int32>* BoneMapArray = AttributeCollectionVal.FindAttribute<int32>(FName("BoneMap"), FName("Vertices"));
+	const TManagedArray<FVector3f>* VertexArray = AttributeCollectionVal.FindAttribute<FVector3f>(FName("Vertex"), FName("Vertices"));
+	const TManagedArray<FIntVector3>* IndicesArray = AttributeCollectionVal.FindAttribute<FIntVector3>(FName("Indices"), FName("Faces"));
+	const TManagedArray<FTransform3f>* LocalSpaceTransform = AttributeCollectionVal.FindAttribute<FTransform3f>(FTransformCollection::TransformAttribute, FTransformCollection::TransformGroup);
+	const TManagedArray<int32>* Parent = AttributeCollectionVal.FindAttribute<int32>(FTransformCollection::ParentAttribute, FTransformCollection::TransformGroup);
+	const TManagedArray<int32>* VertexStart = AttributeCollectionVal.FindAttribute<int32>("VertexStart", FGeometryCollection::GeometryGroup);
+	const TManagedArray<int32>* VertexCount = AttributeCollectionVal.FindAttribute<int32>("VertexCount", FGeometryCollection::GeometryGroup);
+	const TManagedArray<int32>* FaceStart = AttributeCollectionVal.FindAttribute<int32>("FaceStart", FGeometryCollection::GeometryGroup);
+	const TManagedArray<int32>* FaceCount = AttributeCollectionVal.FindAttribute<int32>("FaceCount", FGeometryCollection::GeometryGroup);
+
+	ParallelFor(PairedGeometry.Num(), [&](int32 Pdx)
+	{
+		int32 AttributeGeometryIndex = PairedGeometry[Pdx][0];
+		int32 TargetGeometryIndex = PairedGeometry[Pdx][1];
+
+		if (ensure(0 <= AttributeGeometryIndex && AttributeGeometryIndex < VertexStart->Num()))
+		{
+			if (ensure(0 <= TargetGeometryIndex && TargetGeometryIndex < TargetVertexStart->Num()))
+			{
+				// Build component space vertices for TargetCollection
+				TArray<FVector> ComponentSpaceTargetVertices; // size of num vertices of the geometry entry. 
+				BuildComponentSpaceVertices(TargetLocalSpaceTransform, TargetParent, TargetBoneMapArray, TargetVertexArray, 
+					(*TargetVertexStart)[TargetGeometryIndex], (*TargetVertexCount)[TargetGeometryIndex], ComponentSpaceTargetVertices);
 
+				// Build component space vertices for TargetCollection
+				TArray<FVector> ComponentSpaceVertices; // size of num vertices of the geometry entry
+				BuildComponentSpaceVertices(LocalSpaceTransform, Parent, BoneMapArray, VertexArray, 
+					(*VertexStart)[AttributeGeometryIndex], (*VertexCount)[AttributeGeometryIndex], ComponentSpaceVertices);
 
-					// build Sphere based BVH
-					Chaos::FReal SphereRadius = (Chaos::FReal)0.;
 
-					Chaos::TVec3<float> CoordMaxs(-FLT_MAX);
-					Chaos::TVec3<float> CoordMins(FLT_MAX);
-					for (int32 i = 0; i < ComponentSpaceTargetVertices.Num(); i++)
+				// build Sphere based BVH
+				Chaos::FReal SphereRadius = (Chaos::FReal)0.;
+
+				Chaos::TVec3<float> CoordMaxs(-FLT_MAX);
+				Chaos::TVec3<float> CoordMins(FLT_MAX);
+				for (int32 i = 0; i < ComponentSpaceTargetVertices.Num(); i++)
+				{
+					for (int32 j = 0; j < 3; j++)
 					{
-						for (int32 j = 0; j < 3; j++)
+						if (ComponentSpaceTargetVertices[i][j] > CoordMaxs[j])
 						{
-							if (ComponentSpaceTargetVertices[i][j] > CoordMaxs[j])
-							{
-								CoordMaxs[j] = ComponentSpaceTargetVertices[i][j];
-							}
-							if (ComponentSpaceTargetVertices[i][j] < CoordMins[j])
-							{
-								CoordMins[j] = ComponentSpaceTargetVertices[i][j];
-							}
+							CoordMaxs[j] = ComponentSpaceTargetVertices[i][j];
+						}
+						if (ComponentSpaceTargetVertices[i][j] < CoordMins[j])
+						{
+							CoordMins[j] = ComponentSpaceTargetVertices[i][j];
 						}
 					}
-					Chaos::TVec3<float> CoordDiff = (CoordMaxs - CoordMins) * VertexRadiusRatio;
-					SphereRadius = Chaos::FReal(FGenericPlatformMath::Min(CoordDiff[0], FGenericPlatformMath::Min(CoordDiff[1], CoordDiff[2])));
+				}
+				Chaos::TVec3<float> CoordDiff = (CoordMaxs - CoordMins) * VertexRadiusRatio;
+				SphereRadius = Chaos::FReal(FGenericPlatformMath::Min(CoordDiff[0], FGenericPlatformMath::Min(CoordDiff[1], CoordDiff[2])));
 
-					TArray<Chaos::TSphere<Chaos::FReal, 3>*> VertexSpherePtrs;
-					TArray<Chaos::TSphere<Chaos::FReal, 3>> VertexSpheres;
+				TArray<Chaos::TSphere<Chaos::FReal, 3>*> VertexSpherePtrs;
+				TArray<Chaos::TSphere<Chaos::FReal, 3>> VertexSpheres;
 
-					VertexSpheres.Init(Chaos::TSphere<Chaos::FReal, 3>(Chaos::TVec3<Chaos::FReal>(0), SphereRadius), ComponentSpaceTargetVertices.Num());
-					VertexSpherePtrs.SetNum(ComponentSpaceTargetVertices.Num());
+				VertexSpheres.Init(Chaos::TSphere<Chaos::FReal, 3>(Chaos::TVec3<Chaos::FReal>(0), SphereRadius), ComponentSpaceTargetVertices.Num());
+				VertexSpherePtrs.SetNum(ComponentSpaceTargetVertices.Num());
 
-					for (int32 i = 0; i < ComponentSpaceTargetVertices.Num(); i++)
-					{
-						Chaos::TVec3<Chaos::FReal> SphereCenter(ComponentSpaceTargetVertices[i]);
-						Chaos::TSphere<Chaos::FReal, 3> VertexSphere(SphereCenter, SphereRadius);
-						VertexSpheres[i] = Chaos::TSphere<Chaos::FReal, 3>(SphereCenter, SphereRadius);
-						VertexSpherePtrs[i] = &VertexSpheres[i];
-					}
-					Chaos::TBoundingVolumeHierarchy<
-						TArray<Chaos::TSphere<Chaos::FReal, 3>*>,
-						TArray<int32>,
-						Chaos::FReal,
-						3> VertexBVH(VertexSpherePtrs);
+				for (int32 i = 0; i < ComponentSpaceTargetVertices.Num(); i++)
+				{
+					Chaos::TVec3<Chaos::FReal> SphereCenter(ComponentSpaceTargetVertices[i]);
+					Chaos::TSphere<Chaos::FReal, 3> VertexSphere(SphereCenter, SphereRadius);
+					VertexSpheres[i] = Chaos::TSphere<Chaos::FReal, 3>(SphereCenter, SphereRadius);
+					VertexSpherePtrs[i] = &VertexSpheres[i];
+				}
+				Chaos::TBoundingVolumeHierarchy<
+					TArray<Chaos::TSphere<Chaos::FReal, 3>*>,
+					TArray<int32>,
+					Chaos::FReal,
+					3> VertexBVH(VertexSpherePtrs);
+
+				int32 TargetVertexStartVal = (*TargetVertexStart)[TargetGeometryIndex];
+				int32 TargetVertexCountVal = (*TargetVertexCount)[TargetGeometryIndex];
+				int32 VertexStartVal = (*VertexStart)[AttributeGeometryIndex];
+				int32 FaceStartVal = (*FaceStart)[AttributeGeometryIndex];
+				int32 FaceCountVal = (*FaceCount)[AttributeGeometryIndex];
+				for (int32 i = 0; i < FaceCountVal; i++)
+				{
+					FIntVector3 Triangle = (*IndicesArray)[FaceStartVal+i] - FIntVector3(VertexStartVal);
+					TArray<int32> TriangleIntersections0 = VertexBVH.FindAllIntersections(ComponentSpaceVertices[Triangle[0]]);
+					TArray<int32> TriangleIntersections1 = VertexBVH.FindAllIntersections(ComponentSpaceVertices[Triangle[1]]);
+					TArray<int32> TriangleIntersections2 = VertexBVH.FindAllIntersections(ComponentSpaceVertices[Triangle[2]]);
+					TriangleIntersections0.Sort();
+					TriangleIntersections1.Sort();
+					TriangleIntersections2.Sort();
 
-					for (int32 i = 0; i < IndicesArray->Num(); i++)
+					TArray<int32> TriangleIntersections({});
+					for (int32 k = 0; k < TriangleIntersections0.Num(); k++)
 					{
-						TArray<int32> TriangleIntersections0 = VertexBVH.FindAllIntersections(ComponentSpaceVertices[(*IndicesArray)[i][0]]);
-						TArray<int32> TriangleIntersections1 = VertexBVH.FindAllIntersections(ComponentSpaceVertices[(*IndicesArray)[i][1]]);
-						TArray<int32> TriangleIntersections2 = VertexBVH.FindAllIntersections(ComponentSpaceVertices[(*IndicesArray)[i][2]]);
-						TriangleIntersections0.Sort();
-						TriangleIntersections1.Sort();
-						TriangleIntersections2.Sort();
-
-						TArray<int32> TriangleIntersections({});
-						for (int32 k = 0; k < TriangleIntersections0.Num(); k++)
+						if (TriangleIntersections1.Contains(TriangleIntersections0[k])
+							&& TriangleIntersections2.Contains(TriangleIntersections0[k]))
 						{
-							if (TriangleIntersections1.Contains(TriangleIntersections0[k])
-								&& TriangleIntersections2.Contains(TriangleIntersections0[k]))
-							{
-								TriangleIntersections.Emplace(TriangleIntersections0[k]);
-							}
+							TriangleIntersections.Emplace(TriangleIntersections0[k]);
 						}
+					}
+
+					int32 MinIndex = -1;
+					float MinDis = SphereRadius;
+					Chaos::TVector<float, 3> ClosestBary(0.f);
 
-						int32 MinIndex = -1;
-						float MinDis = SphereRadius;
-						Chaos::TVector<float, 3> ClosestBary(0.f);
+					for (int32 j = 0; j < TriangleIntersections.Num(); j++)
+					{
+						Chaos::TVector<float, 3> Bary,
+							TriPos0(ComponentSpaceVertices[Triangle[0]]),
+							TriPos1(ComponentSpaceVertices[Triangle[1]]),
+							TriPos2(ComponentSpaceVertices[Triangle[2]]),
+							ParticlePos(ComponentSpaceTargetVertices[TriangleIntersections[j]]);
 
-						for (int32 j = 0; j < TriangleIntersections.Num(); j++)
+						Chaos::TVector<Chaos::FRealSingle, 3> ClosestPoint = Chaos::FindClosestPointAndBaryOnTriangle(TriPos0, TriPos1, TriPos2, ParticlePos, Bary);
+						Chaos::FRealSingle CurrentDistance = (ParticlePos - ClosestPoint).Size();
+						if (CurrentDistance < MinDis)
 						{
-							Chaos::TVector<float, 3> Bary,
-								TriPos0(ComponentSpaceVertices[(*IndicesArray)[i][0]]),
-								TriPos1(ComponentSpaceVertices[(*IndicesArray)[i][1]]),
-								TriPos2(ComponentSpaceVertices[(*IndicesArray)[i][2]]),
-								ParticlePos(ComponentSpaceTargetVertices[TriangleIntersections[j]]);
-
-							Chaos::TVector<Chaos::FRealSingle, 3> ClosestPoint = Chaos::FindClosestPointAndBaryOnTriangle(TriPos0, TriPos1, TriPos2, ParticlePos, Bary);
-							Chaos::FRealSingle CurrentDistance = (ParticlePos - ClosestPoint).Size();
-							if (CurrentDistance < MinDis)
-							{
-								MinDis = CurrentDistance;
-								MinIndex = TriangleIntersections[j];
-								ClosestBary = Bary;
-							}
+							MinDis = CurrentDistance;
+							MinIndex = TriangleIntersections[j];
+							ClosestBary = Bary;
+						}
 
-							if (MinIndex != -1
-								&& MinIndex != (*IndicesArray)[i][0]
-								&& MinIndex != (*IndicesArray)[i][1]
-								&& MinIndex != (*IndicesArray)[i][2])
+						if (MinIndex != -1
+							&& MinIndex != (*IndicesArray)[i][0]
+							&& MinIndex != (*IndicesArray)[i][1]
+							&& MinIndex != (*IndicesArray)[i][2])
+						{
+							(*TargetFloatArray)[MinIndex+ TargetVertexStartVal] = 0;
+							for (int32 k = 0; k < 3; k++)
 							{
-								(*TargetFloatArray)[MinIndex] = 0;
-								for (int32 k = 0; k < 3; k++)
-								{
-									(*TargetFloatArray)[MinIndex] += ClosestBary[k] * (*FloatArray)[(*IndicesArray)[i][k]];
-								}
+								(*TargetFloatArray)[MinIndex+ TargetVertexStartVal] += ClosestBary[k] * (*FloatArray)[Triangle[k]+VertexStartVal];
 							}
 						}
 					}
 				}
 			}
 		}
+	});
+}
 
-		SetValue<FManagedArrayCollection >(Context, MoveTemp(CollectionVal), &Collection);
+void FGeometryCollectionTransferVertexScalarAttributeNode::NearestVertexTransfer(FString AttributName, const FManagedArrayCollection& AttributeCollectionVal, const FManagedArrayCollection& CollectionVal, TManagedArray<float>* TargetFloatArray) const
+{
+	auto BuildComponentSpaceVertices = [](const TManagedArray<FTransform3f>* LocalSpaceTransform,
+		const TManagedArray<int32>* Parent, const TManagedArray<int32>* BoneMapArray,
+		const TManagedArray<FVector3f>* VertexArray,
+		TArray<FVector>& ComponentSpaceVertices)
+	{
+		TArray<FTransform> ComponentTransform;
+		GeometryCollectionAlgo::GlobalMatrices(*LocalSpaceTransform, *Parent, ComponentTransform);
+
+		ComponentSpaceVertices.SetNumUninitialized(VertexArray->Num());
+		for (int i = 0; i < VertexArray->Num(); i++)
+		{
+			if (0 < (*BoneMapArray)[i] && (*BoneMapArray)[i] < ComponentTransform.Num())
+			{
+				ComponentSpaceVertices[i] = ComponentTransform[(*BoneMapArray)[i]].TransformPosition(FVector((*VertexArray)[i]));
+			}
+			else
+			{
+				ComponentSpaceVertices[i] = FVector((*VertexArray)[i]);
+			}
+		}
+	};
+
+	const TManagedArray<int32>* TargetBoneMapArray = CollectionVal.FindAttribute<int32>(FName("BoneMap"), FName("Vertices"));
+	const TManagedArray<FVector3f>* TargetVertexArray = CollectionVal.FindAttribute<FVector3f>(FName("Vertex"), FName("Vertices"));
+	const TManagedArray<FIntVector3>* TargetIndicesArray = CollectionVal.FindAttribute<FIntVector3>(FName("Indices"), FName("Faces"));
+	const TManagedArray<FTransform3f>* TargetLocalSpaceTransform = CollectionVal.FindAttribute<FTransform3f>(FTransformCollection::TransformAttribute, FTransformCollection::TransformGroup);
+	const TManagedArray<int32>* TargetParent = CollectionVal.FindAttribute<int32>(FTransformCollection::ParentAttribute, FTransformCollection::TransformGroup);
+
+	const TManagedArray<float>* FloatArray = AttributeCollectionVal.FindAttribute<float>(FName(AttributName), FName("Vertices"));
+	const TManagedArray<int32>* BoneMapArray = AttributeCollectionVal.FindAttribute<int32>(FName("BoneMap"), FName("Vertices"));
+	const TManagedArray<FVector3f>* VertexArray = AttributeCollectionVal.FindAttribute<FVector3f>(FName("Vertex"), FName("Vertices"));
+	const TManagedArray<FIntVector3>* IndicesArray = AttributeCollectionVal.FindAttribute<FIntVector3>(FName("Indices"), FName("Faces"));
+	const TManagedArray<FTransform3f>* LocalSpaceTransform = AttributeCollectionVal.FindAttribute<FTransform3f>(FTransformCollection::TransformAttribute, FTransformCollection::TransformGroup);
+	const TManagedArray<int32>* Parent = AttributeCollectionVal.FindAttribute<int32>(FTransformCollection::ParentAttribute, FTransformCollection::TransformGroup);
+
+	// Build component space vertices for TargetCollection
+	TArray<FVector> ComponentSpaceTargetVertices;
+	BuildComponentSpaceVertices(TargetLocalSpaceTransform, TargetParent, TargetBoneMapArray, TargetVertexArray, ComponentSpaceTargetVertices);
+
+	// Build component space vertices for TargetCollection
+	TArray<FVector> ComponentSpaceVertices;
+	BuildComponentSpaceVertices(LocalSpaceTransform, Parent, BoneMapArray, VertexArray, ComponentSpaceVertices);
+
+
+	// build Sphere based BVH
+	Chaos::FReal SphereRadius = (Chaos::FReal)0.;
+
+	Chaos::TVec3<float> CoordMaxs(-FLT_MAX);
+	Chaos::TVec3<float> CoordMins(FLT_MAX);
+	for (int32 i = 0; i < ComponentSpaceTargetVertices.Num(); i++)
+	{
+		for (int32 j = 0; j < 3; j++)
+		{
+			if (ComponentSpaceTargetVertices[i][j] > CoordMaxs[j])
+			{
+				CoordMaxs[j] = ComponentSpaceTargetVertices[i][j];
+			}
+			if (ComponentSpaceTargetVertices[i][j] < CoordMins[j])
+			{
+				CoordMins[j] = ComponentSpaceTargetVertices[i][j];
+			}
+		}
 	}
-	else if (Out->IsA<FString>(&Name))
+	Chaos::TVec3<float> CoordDiff = (CoordMaxs - CoordMins) * VertexRadiusRatio;
+	SphereRadius = Chaos::FReal(FGenericPlatformMath::Min(CoordDiff[0], FGenericPlatformMath::Min(CoordDiff[1], CoordDiff[2])));
+
+	TArray<Chaos::TSphere<Chaos::FReal, 3>*> VertexSpherePtrs;
+	TArray<Chaos::TSphere<Chaos::FReal, 3>> VertexSpheres;
+
+	VertexSpheres.Init(Chaos::TSphere<Chaos::FReal, 3>(Chaos::TVec3<Chaos::FReal>(0), SphereRadius), ComponentSpaceTargetVertices.Num());
+	VertexSpherePtrs.SetNum(ComponentSpaceTargetVertices.Num());
+
+	for (int32 i = 0; i < ComponentSpaceTargetVertices.Num(); i++)
 	{
-	SetValue< FString >(Context, MoveTemp(AttributName), &Name);
+		Chaos::TVec3<Chaos::FReal> SphereCenter(ComponentSpaceTargetVertices[i]);
+		Chaos::TSphere<Chaos::FReal, 3> VertexSphere(SphereCenter, SphereRadius);
+		VertexSpheres[i] = Chaos::TSphere<Chaos::FReal, 3>(SphereCenter, SphereRadius);
+		VertexSpherePtrs[i] = &VertexSpheres[i];
+	}
+	Chaos::TBoundingVolumeHierarchy<
+		TArray<Chaos::TSphere<Chaos::FReal, 3>*>,
+		TArray<int32>,
+		Chaos::FReal,
+		3> VertexBVH(VertexSpherePtrs);
+
+	for (int32 i = 0; i < IndicesArray->Num(); i++)
+	{
+		TArray<int32> TriangleIntersections0 = VertexBVH.FindAllIntersections(ComponentSpaceVertices[(*IndicesArray)[i][0]]);
+		TArray<int32> TriangleIntersections1 = VertexBVH.FindAllIntersections(ComponentSpaceVertices[(*IndicesArray)[i][1]]);
+		TArray<int32> TriangleIntersections2 = VertexBVH.FindAllIntersections(ComponentSpaceVertices[(*IndicesArray)[i][2]]);
+		TriangleIntersections0.Sort();
+		TriangleIntersections1.Sort();
+		TriangleIntersections2.Sort();
+
+		TArray<int32> TriangleIntersections({});
+		for (int32 k = 0; k < TriangleIntersections0.Num(); k++)
+		{
+			if (TriangleIntersections1.Contains(TriangleIntersections0[k])
+				&& TriangleIntersections2.Contains(TriangleIntersections0[k]))
+			{
+				TriangleIntersections.Emplace(TriangleIntersections0[k]);
+			}
+		}
+
+		int32 MinIndex = -1;
+		float MinDis = SphereRadius;
+		Chaos::TVector<float, 3> ClosestBary(0.f);
+
+		for (int32 j = 0; j < TriangleIntersections.Num(); j++)
+		{
+			Chaos::TVector<float, 3> Bary,
+				TriPos0(ComponentSpaceVertices[(*IndicesArray)[i][0]]),
+				TriPos1(ComponentSpaceVertices[(*IndicesArray)[i][1]]),
+				TriPos2(ComponentSpaceVertices[(*IndicesArray)[i][2]]),
+				ParticlePos(ComponentSpaceTargetVertices[TriangleIntersections[j]]);
+
+			Chaos::TVector<Chaos::FRealSingle, 3> ClosestPoint = Chaos::FindClosestPointAndBaryOnTriangle(TriPos0, TriPos1, TriPos2, ParticlePos, Bary);
+			Chaos::FRealSingle CurrentDistance = (ParticlePos - ClosestPoint).Size();
+			if (CurrentDistance < MinDis)
+			{
+				MinDis = CurrentDistance;
+				MinIndex = TriangleIntersections[j];
+				ClosestBary = Bary;
+			}
+
+			if (MinIndex != -1
+				&& MinIndex != (*IndicesArray)[i][0]
+				&& MinIndex != (*IndicesArray)[i][1]
+				&& MinIndex != (*IndicesArray)[i][2])
+			{
+				(*TargetFloatArray)[MinIndex] = 0;
+				for (int32 k = 0; k < 3; k++)
+				{
+					(*TargetFloatArray)[MinIndex] += ClosestBary[k] * (*FloatArray)[(*IndicesArray)[i][k]];
+				}
+			}
+		}
 	}
 }
 
 
+
 #undef LOCTEXT_NAMESPACE
diff --git a/Engine/Plugins/Experimental/GeometryCollectionPlugin/Source/GeometryCollectionNodes/Public/Dataflow/GeometryCollectionTransferVertexScalarAttributeNode.h b/Engine/Plugins/Experimental/GeometryCollectionPlugin/Source/GeometryCollectionNodes/Public/Dataflow/GeometryCollectionTransferVertexScalarAttributeNode.h
--- a/Engine/Plugins/Experimental/GeometryCollectionPlugin/Source/GeometryCollectionNodes/Public/Dataflow/GeometryCollectionTransferVertexScalarAttributeNode.h
+++ b/Engine/Plugins/Experimental/GeometryCollectionPlugin/Source/GeometryCollectionNodes/Public/Dataflow/GeometryCollectionTransferVertexScalarAttributeNode.h
@@ -43,4 +43,9 @@ public:
 private:
 
 	virtual void Evaluate(Dataflow::FContext& Context, const FDataflowOutput* Out) const override;
+
+	TArray<FIntVector2> FindSourceToTargetGeometryMap(const FManagedArrayCollection& SourceCollection, const FManagedArrayCollection& TargetCollection) const;
+	void PairedGeometryTransfer(FString AttributName, const TArray<FIntVector2>& PairedGeometry, const FManagedArrayCollection& AttributeCollectionVal, const FManagedArrayCollection& CollectionVal, TManagedArray<float>* FloatArray) const;
+	void NearestVertexTransfer(FString AttributName, const FManagedArrayCollection& AttributeCollectionVal, const FManagedArrayCollection& CollectionVal, TManagedArray<float>* FloatArray) const;
+
 };
