ChaosFlesh : Transfer Weights
- Updated query culling threshold based on max edge length.
#preflgiht 66289673e363eda2c46f7a4f

diff --git a/Engine/Plugins/Experimental/GeometryCollectionPlugin/Source/GeometryCollectionNodes/Private/Dataflow/GeometryCollectionTransferVertexScalarAttributeNode.cpp b/Engine/Plugins/Experimental/GeometryCollectionPlugin/Source/GeometryCollectionNodes/Private/Dataflow/GeometryCollectionTransferVertexScalarAttributeNode.cpp
--- a/Engine/Plugins/Experimental/GeometryCollectionPlugin/Source/GeometryCollectionNodes/Private/Dataflow/GeometryCollectionTransferVertexScalarAttributeNode.cpp
+++ b/Engine/Plugins/Experimental/GeometryCollectionPlugin/Source/GeometryCollectionNodes/Private/Dataflow/GeometryCollectionTransferVertexScalarAttributeNode.cpp
@@ -2,87 +2,126 @@
 
 #include "Dataflow/GeometryCollectionTransferVertexScalarAttributeNode.h"
 
-#include "Chaos/BoundingVolumeHierarchy.h"
 #include "Chaos/Triangle.h"
 #include "Dataflow/DataflowInputOutput.h"
 #include "GeometryCollection/GeometryCollectionAlgo.h"
 #include "GeometryCollection/TransformCollection.h"
 #include "GeometryCollection/GeometryCollection.h"
+#include "GeometryCollection/ManagedArrayAccessor.h"
 
 
 #include UE_INLINE_GENERATED_CPP_BY_NAME(GeometryCollectionTransferVertexScalarAttributeNode)
 #define LOCTEXT_NAMESPACE "FGeometryCollectionTransferVertexScalarAttributeNode"
 
+namespace UE::Private {
 
+	class FTransferFacade
+	{
+		const FManagedArrayCollection& ConstCollection;
+		FManagedArrayCollection* Collection = nullptr;
+	public:
+		FTransferFacade(FManagedArrayCollection& InCollection)
+			: ConstCollection(InCollection)
+			, Collection(&InCollection)
+			, BoneMap(InCollection, FName("BoneMap"), FName("Vertices"))
+			, Vertex(InCollection, FName("Vertex"), FName("Vertices"))
+			, Indices(InCollection, FName("Indices"), FName("Faces"))
+			, Transform(InCollection, FTransformCollection::TransformAttribute, FTransformCollection::TransformGroup)
+			, Parent(InCollection, FTransformCollection::ParentAttribute, FTransformCollection::TransformGroup)
+			, VertexStart(InCollection, "VertexStart", FGeometryCollection::GeometryGroup)
+			, VertexCount(InCollection, "VertexCount", FGeometryCollection::GeometryGroup)
+			, FaceStart(InCollection, "FaceStart", FGeometryCollection::GeometryGroup)
+			, FaceCount(InCollection, "FaceCount", FGeometryCollection::GeometryGroup)
+		{}
+
+		FTransferFacade(const FManagedArrayCollection& InCollection)
+			: ConstCollection(InCollection)
+			, Collection(nullptr)
+			, BoneMap(InCollection, FName("BoneMap"), FName("Vertices"))
+			, Vertex(InCollection, FName("Vertex"), FName("Vertices"))
+			, Indices(InCollection, FName("Indices"), FName("Faces"))
+			, Transform(InCollection, FTransformCollection::TransformAttribute, FTransformCollection::TransformGroup)
+			, Parent(InCollection, FTransformCollection::ParentAttribute, FTransformCollection::TransformGroup)
+			, VertexStart(InCollection, "VertexStart", FGeometryCollection::GeometryGroup)
+			, VertexCount(InCollection, "VertexCount", FGeometryCollection::GeometryGroup)
+			, FaceStart(InCollection, "FaceStart", FGeometryCollection::GeometryGroup)
+			, FaceCount(InCollection, "FaceCount", FGeometryCollection::GeometryGroup)
+		{}
+
+
+		bool IsValid() const {
+			return BoneMap.IsValid() && Vertex.IsValid() && Indices.IsValid() && Transform.IsValid() && Parent.IsValid() && VertexStart.IsValid() &&
+				VertexCount.IsValid() && FaceStart.IsValid() && FaceCount.IsValid();
+		}
+
+		const TManagedArray<float>* GetFloatArray(FString AttributeName, FString Group) const
+		{
+			return ConstCollection.FindAttribute<float>(FName(AttributeName), FName(Group));
+		}
+
+		TManagedArray<float>* GetFloatArray(FString AttributeName, FString Group)
+		{
+			TManagedArray<float>* TargetFloatArray = nullptr;
+			if (!Collection->HasAttribute(FName(AttributeName), FName(Group)))
+			{
+				Collection->AddAttribute<float>(FName(AttributeName), FName(Group));
+			}
+			return Collection->FindAttribute<float>(FName(AttributeName), FName(Group));
+		}
+
+
+		TManagedArrayAccessor<int32> BoneMap;
+		TManagedArrayAccessor<FVector3f> Vertex;
+		TManagedArrayAccessor<FIntVector3> Indices;
+		TManagedArrayAccessor<FTransform3f> Transform;
+		TManagedArrayAccessor<int32> Parent;
+		TManagedArrayAccessor<int32> VertexStart;
+		TManagedArrayAccessor<int32> VertexCount;
+		TManagedArrayAccessor<int32> FaceStart;
+		TManagedArrayAccessor<int32> FaceCount;
+	};
+
+}
 
 void FGeometryCollectionTransferVertexScalarAttributeNode::Evaluate(Dataflow::FContext& Context, const FDataflowOutput* Out) const
 {
 
-
-	FString AttributName = GetValue<FString>(Context, &Name, Name);
+	FString AttributeName = GetValue<FString>(Context, &Name, Name);
 	if (Out->IsA< FManagedArrayCollection >(&Collection))
 	{
-		FManagedArrayCollection CollectionVal = GetValue<FManagedArrayCollection>(Context, &Collection);
-		const FManagedArrayCollection& AttributeCollectionVal = GetValue<FManagedArrayCollection>(Context, &FromCollection);
-
-		const TManagedArray<int32>* TargetBoneMapArray = CollectionVal.FindAttribute<int32>(FName("BoneMap"), FName("Vertices"));
-		const TManagedArray<FVector3f>* TargetVertexArray = CollectionVal.FindAttribute<FVector3f>(FName("Vertex"), FName("Vertices"));
-		const TManagedArray<FIntVector3>* TargetIndicesArray = CollectionVal.FindAttribute<FIntVector3>(FName("Indices"), FName("Faces"));
-		const TManagedArray<FTransform3f>* TargetLocalSpaceTransform = CollectionVal.FindAttribute<FTransform3f>(FTransformCollection::TransformAttribute, FTransformCollection::TransformGroup);
-		const TManagedArray<int32>* TargetParent = CollectionVal.FindAttribute<int32>(FTransformCollection::ParentAttribute, FTransformCollection::TransformGroup);
-		const TManagedArray<int32>* TargetVertexStart = CollectionVal.FindAttribute<int32>("VertexStart", FGeometryCollection::GeometryGroup);
-		const TManagedArray<int32>* TargetVertexCount = CollectionVal.FindAttribute<int32>("VertexCount", FGeometryCollection::GeometryGroup);
-		const TManagedArray<int32>* TargetFaceStart = CollectionVal.FindAttribute<int32>("FaceStart", FGeometryCollection::GeometryGroup);
-		const TManagedArray<int32>* TargetFaceCount = CollectionVal.FindAttribute<int32>("FaceCount", FGeometryCollection::GeometryGroup);
-
-		const TManagedArray<float>* FloatArray = AttributeCollectionVal.FindAttribute<float>(FName(AttributName), FName("Vertices"));
-		const TManagedArray<int32>* BoneMapArray = AttributeCollectionVal.FindAttribute<int32>(FName("BoneMap"), FName("Vertices"));
-		const TManagedArray<FVector3f>* VertexArray = AttributeCollectionVal.FindAttribute<FVector3f>(FName("Vertex"), FName("Vertices"));
-		const TManagedArray<FIntVector3>* IndicesArray = AttributeCollectionVal.FindAttribute<FIntVector3>(FName("Indices"), FName("Faces"));
-		const TManagedArray<FTransform3f>* LocalSpaceTransform = AttributeCollectionVal.FindAttribute<FTransform3f>(FTransformCollection::TransformAttribute, FTransformCollection::TransformGroup);
-		const TManagedArray<int32>* Parent = AttributeCollectionVal.FindAttribute<int32>(FTransformCollection::ParentAttribute, FTransformCollection::TransformGroup);
-		const TManagedArray<int32>* VertexStart = AttributeCollectionVal.FindAttribute<int32>("VertexStart", FGeometryCollection::GeometryGroup);
-		const TManagedArray<int32>* VertexCount = AttributeCollectionVal.FindAttribute<int32>("VertexCount", FGeometryCollection::GeometryGroup);
-		const TManagedArray<int32>* FaceStart = AttributeCollectionVal.FindAttribute<int32>("FaceStart", FGeometryCollection::GeometryGroup);
-		const TManagedArray<int32>* FaceCount = AttributeCollectionVal.FindAttribute<int32>("FaceCount", FGeometryCollection::GeometryGroup);
-
-		if (TargetBoneMapArray && TargetVertexArray && TargetIndicesArray && TargetLocalSpaceTransform && TargetParent && TargetVertexStart && TargetVertexCount && TargetFaceStart && TargetFaceCount)
+		FManagedArrayCollection TargetCollection = GetValue<FManagedArrayCollection>(Context, &Collection);
+		const FManagedArrayCollection& SampleCollection = GetValue<FManagedArrayCollection>(Context, &FromCollection);
+
+		UE::Private::FTransferFacade Target(TargetCollection);
+		const UE::Private::FTransferFacade Sample(SampleCollection);
+
+		if (Target.IsValid() && Sample.IsValid())
 		{
-			if (FloatArray && BoneMapArray && VertexArray && IndicesArray && LocalSpaceTransform && Parent && VertexStart && VertexCount && FaceStart && FaceCount)
+			if (TManagedArray<float>* TargetFloatArray = Target.GetFloatArray(AttributeName, "Vertices"))
 			{
-				TManagedArray<float>* TargetFloatArray = CollectionVal.FindAttributeTyped<float>(FName(AttributName), FName("Vertices"));
-				if (!TargetFloatArray && !CollectionVal.HasAttribute(FName(AttributName), FName("Vertices")))
+				TargetFloatArray->Fill(0.f);
+
+				TArray<FIntVector2> AlignedGeometry = FindSourceToTargetGeometryMap(SampleCollection, TargetCollection);
+				if (AlignedGeometry.Num() == TargetCollection.NumElements(FGeometryCollection::GeometryGroup))
 				{
-					CollectionVal.AddAttribute<float>(FName(AttributName), FName("Vertices"));
-					TargetFloatArray = CollectionVal.FindAttribute<float>(FName(AttributName), FName("Vertices"));
+					PairedGeometryTransfer(AttributeName, AlignedGeometry, Sample, Target, TargetFloatArray);
 				}
-
-				if(TargetFloatArray)
+				else
 				{
-					TargetFloatArray->Fill(0.f);
-
-
-					TArray<FIntVector2> AlignedGeometry = FindSourceToTargetGeometryMap(AttributeCollectionVal, CollectionVal);
-					if (AlignedGeometry.Num() == AttributeCollectionVal.NumElements(FGeometryCollection::GeometryGroup))
-					{
-						PairedGeometryTransfer(AttributName, AlignedGeometry, AttributeCollectionVal,CollectionVal,TargetFloatArray);
-					}
-					else
-					{
-						NearestVertexTransfer(AttributName, AttributeCollectionVal, CollectionVal, TargetFloatArray);
-					}
+					NearestVertexTransfer(AttributeName, Sample, Target, TargetFloatArray);
 				}
 			}
 		}
 
-		SetValue<FManagedArrayCollection >(Context, MoveTemp(CollectionVal), &Collection);
+		SetValue<FManagedArrayCollection >(Context, MoveTemp(TargetCollection), &Collection);
 	}
 	else if (Out->IsA<FString>(&Name))
 	{
-	SetValue< FString >(Context, MoveTemp(AttributName), &Name);
+		SetValue< FString >(Context, MoveTemp(AttributeName), &Name);
 	}
 }
 
+
 TArray<FIntVector2> FGeometryCollectionTransferVertexScalarAttributeNode::FindSourceToTargetGeometryMap(const FManagedArrayCollection& AttributeCollectionVal, const FManagedArrayCollection& CollectionVal) const
 {
 	TArray<FIntVector2> Mapping;
@@ -109,322 +148,226 @@ TArray<FIntVector2> FGeometryCollectionTransferVertexScalarAttributeNode::FindSo
 	return Mapping;
 }
 
-
-void FGeometryCollectionTransferVertexScalarAttributeNode::PairedGeometryTransfer(FString AttributName, const TArray<FIntVector2>& PairedGeometry, const FManagedArrayCollection& AttributeCollectionVal, const FManagedArrayCollection& CollectionVal, TManagedArray<float>* TargetFloatArray) const
+void FGeometryCollectionTransferVertexScalarAttributeNode::PairedGeometryTransfer(FString AttributeName, const TArray<FIntVector2>& PairedGeometry, 
+	const UE::Private::FTransferFacade& Sample, UE::Private::FTransferFacade&  Target, TManagedArray<float>* TargetFloatArray) const
 {
-	auto BuildComponentSpaceVertices = [](const TManagedArray<FTransform3f>* LocalSpaceTransform,
-		const TManagedArray<int32>* Parent, const TManagedArray<int32>* BoneMapArray,
-		const TManagedArray<FVector3f>* VertexArray,
-		int32 Start, int32 Count,
-		TArray<FVector>& ComponentSpaceVertices)
+	if (const TManagedArray<float>* FloatArray = Sample.GetFloatArray(AttributeName, "Vertices"))
 	{
-		int32 End = Start + Count;
-		TArray<FTransform> ComponentTransform;
-		GeometryCollectionAlgo::GlobalMatrices(*LocalSpaceTransform, *Parent, ComponentTransform);
-
-		ComponentSpaceVertices.SetNumUninitialized(Count);
-		for (int i = 0; i < Count; i++)
+		ParallelFor(PairedGeometry.Num(), [&](int32 Pdx)
 		{
-			int j = i + Start;
-			if (0 < (*BoneMapArray)[i] && (*BoneMapArray)[i] < ComponentTransform.Num())
-			{
-				ComponentSpaceVertices[i] = ComponentTransform[(*BoneMapArray)[j]].TransformPosition(FVector((*VertexArray)[j]));
-			}
-			else
-			{
-				ComponentSpaceVertices[i] = FVector((*VertexArray)[j]);
-			}
-		}
-	};
+			int32 AttributeGeometryIndex = PairedGeometry[Pdx][0];
+			int32 TargetGeometryIndex = PairedGeometry[Pdx][1];
 
-	const TManagedArray<int32>* TargetBoneMapArray = CollectionVal.FindAttribute<int32>(FName("BoneMap"), FName("Vertices"));
-	const TManagedArray<FVector3f>* TargetVertexArray = CollectionVal.FindAttribute<FVector3f>(FName("Vertex"), FName("Vertices"));
-	const TManagedArray<FIntVector3>* TargetIndicesArray = CollectionVal.FindAttribute<FIntVector3>(FName("Indices"), FName("Faces"));
-	const TManagedArray<FTransform3f>* TargetLocalSpaceTransform = CollectionVal.FindAttribute<FTransform3f>(FTransformCollection::TransformAttribute, FTransformCollection::TransformGroup);
-	const TManagedArray<int32>* TargetParent = CollectionVal.FindAttribute<int32>(FTransformCollection::ParentAttribute, FTransformCollection::TransformGroup);
-	const TManagedArray<int32>* TargetVertexStart = CollectionVal.FindAttribute<int32>("VertexStart", FGeometryCollection::GeometryGroup);
-	const TManagedArray<int32>* TargetVertexCount = CollectionVal.FindAttribute<int32>("VertexCount", FGeometryCollection::GeometryGroup);
-	const TManagedArray<int32>* TargetFaceStart = CollectionVal.FindAttribute<int32>("FaceStart", FGeometryCollection::GeometryGroup);
-	const TManagedArray<int32>* TargetFaceCount = CollectionVal.FindAttribute<int32>("FaceCount", FGeometryCollection::GeometryGroup);
-
-	const TManagedArray<float>* FloatArray = AttributeCollectionVal.FindAttribute<float>(FName(AttributName), FName("Vertices"));
-	const TManagedArray<int32>* BoneMapArray = AttributeCollectionVal.FindAttribute<int32>(FName("BoneMap"), FName("Vertices"));
-	const TManagedArray<FVector3f>* VertexArray = AttributeCollectionVal.FindAttribute<FVector3f>(FName("Vertex"), FName("Vertices"));
-	const TManagedArray<FIntVector3>* IndicesArray = AttributeCollectionVal.FindAttribute<FIntVector3>(FName("Indices"), FName("Faces"));
-	const TManagedArray<FTransform3f>* LocalSpaceTransform = AttributeCollectionVal.FindAttribute<FTransform3f>(FTransformCollection::TransformAttribute, FTransformCollection::TransformGroup);
-	const TManagedArray<int32>* Parent = AttributeCollectionVal.FindAttribute<int32>(FTransformCollection::ParentAttribute, FTransformCollection::TransformGroup);
-	const TManagedArray<int32>* VertexStart = AttributeCollectionVal.FindAttribute<int32>("VertexStart", FGeometryCollection::GeometryGroup);
-	const TManagedArray<int32>* VertexCount = AttributeCollectionVal.FindAttribute<int32>("VertexCount", FGeometryCollection::GeometryGroup);
-	const TManagedArray<int32>* FaceStart = AttributeCollectionVal.FindAttribute<int32>("FaceStart", FGeometryCollection::GeometryGroup);
-	const TManagedArray<int32>* FaceCount = AttributeCollectionVal.FindAttribute<int32>("FaceCount", FGeometryCollection::GeometryGroup);
-
-	ParallelFor(PairedGeometry.Num(), [&](int32 Pdx)
-	{
-		int32 AttributeGeometryIndex = PairedGeometry[Pdx][0];
-		int32 TargetGeometryIndex = PairedGeometry[Pdx][1];
-
-		if (ensure(0 <= AttributeGeometryIndex && AttributeGeometryIndex < VertexStart->Num()))
-		{
-			if (ensure(0 <= TargetGeometryIndex && TargetGeometryIndex < TargetVertexStart->Num()))
+			if (ensure(0 <= AttributeGeometryIndex && AttributeGeometryIndex < Sample.VertexStart.Num()))
 			{
-				// Build component space vertices for TargetCollection
-				TArray<FVector> ComponentSpaceTargetVertices; // size of num vertices of the geometry entry. 
-				BuildComponentSpaceVertices(TargetLocalSpaceTransform, TargetParent, TargetBoneMapArray, TargetVertexArray, 
-					(*TargetVertexStart)[TargetGeometryIndex], (*TargetVertexCount)[TargetGeometryIndex], ComponentSpaceTargetVertices);
-
-				// Build component space vertices for TargetCollection
-				TArray<FVector> ComponentSpaceVertices; // size of num vertices of the geometry entry
-				BuildComponentSpaceVertices(LocalSpaceTransform, Parent, BoneMapArray, VertexArray, 
-					(*VertexStart)[AttributeGeometryIndex], (*VertexCount)[AttributeGeometryIndex], ComponentSpaceVertices);
-
-
-				// build Sphere based BVH
-				Chaos::FReal SphereRadius = (Chaos::FReal)0.;
-
-				Chaos::TVec3<float> CoordMaxs(-FLT_MAX);
-				Chaos::TVec3<float> CoordMins(FLT_MAX);
-				for (int32 i = 0; i < ComponentSpaceTargetVertices.Num(); i++)
+				if (ensure(0 <= TargetGeometryIndex && TargetGeometryIndex < Target.VertexStart.Num()))
 				{
-					for (int32 j = 0; j < 3; j++)
+					// Build component space vertices for TargetCollection
+					TArray<FVector> ComponentSpaceTargetVertices; // size of num vertices of the geometry entry. 
+					BuildComponentSpaceVertices(&Target.Transform.Get(), &Target.Parent.Get(), &Target.BoneMap.Get(), &Target.Vertex.Get(),
+						Target.VertexStart[TargetGeometryIndex], Target.VertexCount[TargetGeometryIndex], ComponentSpaceTargetVertices);
+
+					// Build component space vertices for TargetCollection
+					TArray<FVector> ComponentSpaceVertices; // size of num vertices of the geometry entry
+					BuildComponentSpaceVertices(&Sample.Transform.Get(), &Sample.Parent.Get(), &Sample.BoneMap.Get(), &Sample.Vertex.Get(),
+						Sample.VertexStart[AttributeGeometryIndex], Sample.VertexCount[AttributeGeometryIndex], ComponentSpaceVertices);
+
+					// build Sphere based BVH
+					Chaos::FReal SphereRadius = Chaos::FReal(EdgeMultiplier * FMath::Max(
+						MaxEdgeLength(ComponentSpaceTargetVertices, Target.Indices.Get(), Target.VertexStart[TargetGeometryIndex], Target.FaceStart[TargetGeometryIndex], Target.FaceCount[TargetGeometryIndex]),
+						MaxEdgeLength(ComponentSpaceVertices, Sample.Indices.Get(), Sample.VertexStart[AttributeGeometryIndex], Sample.FaceStart[AttributeGeometryIndex], Sample.FaceCount[AttributeGeometryIndex])));
+
+					TUniquePtr<UE::Private::BVH> VertexBVH(BuildParticleSphereBVH(ComponentSpaceTargetVertices, SphereRadius));
+
+					int32 TargetVertexStartVal = Target.VertexStart[TargetGeometryIndex];
+					int32 TargetVertexCountVal = Target.VertexCount[TargetGeometryIndex];
+					int32 VertexStartVal = Sample.VertexStart[AttributeGeometryIndex];
+					int32 FaceStartVal = Sample.FaceStart[AttributeGeometryIndex];
+					int32 FaceCountVal = Sample.FaceCount[AttributeGeometryIndex];
+					for (int32 i = 0; i < FaceCountVal; i++)
 					{
-						if (ComponentSpaceTargetVertices[i][j] > CoordMaxs[j])
+						FIntVector3 ComponentTriangle = Sample.Indices[FaceStartVal + i] - FIntVector3(VertexStartVal);
+						if (TriangleHasWeightsToTransfer(Sample.Indices[FaceStartVal + i], *FloatArray))
 						{
-							CoordMaxs[j] = ComponentSpaceTargetVertices[i][j];
-						}
-						if (ComponentSpaceTargetVertices[i][j] < CoordMins[j])
-						{
-							CoordMins[j] = ComponentSpaceTargetVertices[i][j];
+							TArray<int32> TargetVertexIntersection({});
+							TriangleToVertexIntersections(*VertexBVH, ComponentSpaceVertices, ComponentTriangle, TargetVertexIntersection);
+
+							for (int32 j = 0; j < TargetVertexIntersection.Num(); j++)
+							{
+								Chaos::TVector<float, 3> Bary,
+									TriPos0(ComponentSpaceVertices[ComponentTriangle[0]]),
+									TriPos1(ComponentSpaceVertices[ComponentTriangle[1]]),
+									TriPos2(ComponentSpaceVertices[ComponentTriangle[2]]),
+									ParticlePos(ComponentSpaceTargetVertices[TargetVertexIntersection[j]]);
+
+								Chaos::TVector<Chaos::FRealSingle, 3> ClosestPoint = Chaos::FindClosestPointAndBaryOnTriangle(TriPos0, TriPos1, TriPos2, ParticlePos, Bary);
+								Chaos::FRealSingle CurrentDistance = (ParticlePos - ClosestPoint).Size();
+								float TriRadius = FalloffThreshold * MaxEdgeLength(ComponentSpaceVertices, Sample.Indices.Get(), 0, i, 1);
+								if (CurrentDistance < TriRadius)
+								{
+									int32 TargetIndex = TargetVertexIntersection[j] + TargetVertexStartVal;
+									if (ensure(0 <= TargetIndex && TargetIndex < TargetFloatArray->Num()))
+									{
+										float Value = 0.f;
+										for (int32 k = 0; k < 3; k++)
+										{
+											Value += Bary[k] * (*FloatArray)[ComponentTriangle[k] + VertexStartVal];
+										}
+										(*TargetFloatArray)[TargetIndex] = FMath::Max((*TargetFloatArray)[TargetIndex], Value);
+									}
+								}
+							}
 						}
 					}
 				}
-				Chaos::TVec3<float> CoordDiff = (CoordMaxs - CoordMins) * VertexRadiusRatio;
-				SphereRadius = Chaos::FReal(FGenericPlatformMath::Max(CoordDiff[0], FGenericPlatformMath::Max(CoordDiff[1], CoordDiff[2])));
+			}
+		});
+	}
+}
 
-				TArray<Chaos::TSphere<Chaos::FReal, 3>*> VertexSpherePtrs;
-				TArray<Chaos::TSphere<Chaos::FReal, 3>> VertexSpheres;
+void FGeometryCollectionTransferVertexScalarAttributeNode::NearestVertexTransfer(FString AttributeName, const UE::Private::FTransferFacade& Sample, UE::Private::FTransferFacade& Target, TManagedArray<float>* TargetFloatArray) const
+{
+	if (const TManagedArray<float>* FloatArray = Sample.GetFloatArray(AttributeName, "Vertices"))
+	{
+		// Build component space vertices for TargetCollection
+		TArray<FVector> ComponentSpaceTargetVertices;
+		BuildComponentSpaceVertices(&Target.Transform.Get(), &Target.Parent.Get(), &Target.BoneMap.Get(), &Target.Vertex.Get(), 0, Target.Vertex.Num(), ComponentSpaceTargetVertices);
 
-				VertexSpheres.Init(Chaos::TSphere<Chaos::FReal, 3>(Chaos::TVec3<Chaos::FReal>(0), SphereRadius), ComponentSpaceTargetVertices.Num());
-				VertexSpherePtrs.SetNum(ComponentSpaceTargetVertices.Num());
+		// Build component space vertices for SourceCollection
+		TArray<FVector> ComponentSpaceVertices;
+		BuildComponentSpaceVertices(&Sample.Transform.Get(), &Sample.Parent.Get(), &Sample.BoneMap.Get(), &Sample.Vertex.Get(), 0, Sample.Vertex.Num(), ComponentSpaceVertices);
 
-				for (int32 i = 0; i < ComponentSpaceTargetVertices.Num(); i++)
-				{
-					Chaos::TVec3<Chaos::FReal> SphereCenter(ComponentSpaceTargetVertices[i]);
-					Chaos::TSphere<Chaos::FReal, 3> VertexSphere(SphereCenter, SphereRadius);
-					VertexSpheres[i] = Chaos::TSphere<Chaos::FReal, 3>(SphereCenter, SphereRadius);
-					VertexSpherePtrs[i] = &VertexSpheres[i];
-				}
-				Chaos::TBoundingVolumeHierarchy<
-					TArray<Chaos::TSphere<Chaos::FReal, 3>*>,
-					TArray<int32>,
-					Chaos::FReal,
-					3> VertexBVH(VertexSpherePtrs);
-
-				int32 TargetVertexStartVal = (*TargetVertexStart)[TargetGeometryIndex];
-				int32 TargetVertexCountVal = (*TargetVertexCount)[TargetGeometryIndex];
-				int32 VertexStartVal = (*VertexStart)[AttributeGeometryIndex];
-				int32 FaceStartVal = (*FaceStart)[AttributeGeometryIndex];
-				int32 FaceCountVal = (*FaceCount)[AttributeGeometryIndex];
-				for (int32 i = 0; i < FaceCountVal; i++)
-				{
-					FIntVector3 Triangle = (*IndicesArray)[FaceStartVal+i] - FIntVector3(VertexStartVal);
-					TArray<int32> TriangleIntersections0 = VertexBVH.FindAllIntersections(ComponentSpaceVertices[Triangle[0]]);
-					TArray<int32> TriangleIntersections1 = VertexBVH.FindAllIntersections(ComponentSpaceVertices[Triangle[1]]);
-					TArray<int32> TriangleIntersections2 = VertexBVH.FindAllIntersections(ComponentSpaceVertices[Triangle[2]]);
-					TriangleIntersections0.Sort();
-					TriangleIntersections1.Sort();
-					TriangleIntersections2.Sort();
-
-					TArray<int32> TriangleIntersections({});
-					for (int32 k = 0; k < TriangleIntersections0.Num(); k++)
-					{
-						if (TriangleIntersections1.Contains(TriangleIntersections0[k])
-							&& TriangleIntersections2.Contains(TriangleIntersections0[k]))
-						{
-							TriangleIntersections.Emplace(TriangleIntersections0[k]);
-						}
-					}
+		// build Sphere based BVH
+		Chaos::FReal SphereRadius = Chaos::FReal(EdgeMultiplier * FMath::Max(
+			MaxEdgeLength(ComponentSpaceTargetVertices, Target.Indices.Get(), 0, 0, Target.Indices.Num()),
+			MaxEdgeLength(ComponentSpaceVertices, Sample.Indices.Get(), 0, 0, Sample.Indices.Num())));
+		TUniquePtr<UE::Private::BVH> VertexBVH(BuildParticleSphereBVH(ComponentSpaceTargetVertices, SphereRadius));
 
-					int32 MinIndex = -1;
-					float MinDis = SphereRadius;
-					Chaos::TVector<float, 3> ClosestBary(0.f);
+		for (int32 i = 0; i < Sample.Indices.Num(); i++)
+		{
+			FIntVector3 Triangle = Sample.Indices[i];
+			if (TriangleHasWeightsToTransfer(Triangle, *FloatArray))
+			{
+				TArray<int32> TargetVertexIntersection({});
+				TriangleToVertexIntersections(*VertexBVH, ComponentSpaceVertices, Triangle, TargetVertexIntersection);
 
-					for (int32 j = 0; j < TriangleIntersections.Num(); j++)
+				for (int32 j = 0; j < TargetVertexIntersection.Num(); j++)
+				{
+					Chaos::TVector<float, 3> Bary,
+						TriPos0(ComponentSpaceVertices[Sample.Indices[i][0]]),
+						TriPos1(ComponentSpaceVertices[Sample.Indices[i][1]]),
+						TriPos2(ComponentSpaceVertices[Sample.Indices[i][2]]),
+						ParticlePos(ComponentSpaceTargetVertices[TargetVertexIntersection[j]]);
+
+					Chaos::TVector<Chaos::FRealSingle, 3> ClosestPoint = Chaos::FindClosestPointAndBaryOnTriangle(TriPos0, TriPos1, TriPos2, ParticlePos, Bary);
+					Chaos::FRealSingle CurrentDistance = (ParticlePos - ClosestPoint).Size();
+					float TriRadius = FalloffThreshold * MaxEdgeLength(ComponentSpaceVertices, Sample.Indices.Get(), 0, i, 1);
+					if (CurrentDistance < TriRadius)
 					{
-						Chaos::TVector<float, 3> Bary,
-							TriPos0(ComponentSpaceVertices[Triangle[0]]),
-							TriPos1(ComponentSpaceVertices[Triangle[1]]),
-							TriPos2(ComponentSpaceVertices[Triangle[2]]),
-							ParticlePos(ComponentSpaceTargetVertices[TriangleIntersections[j]]);
-
-						Chaos::TVector<Chaos::FRealSingle, 3> ClosestPoint = Chaos::FindClosestPointAndBaryOnTriangle(TriPos0, TriPos1, TriPos2, ParticlePos, Bary);
-						Chaos::FRealSingle CurrentDistance = (ParticlePos - ClosestPoint).Size();
-						if (CurrentDistance < MinDis)
+						int32 TargetIndex = TargetVertexIntersection[j];
+						if (ensure(0 <= TargetIndex && TargetIndex < TargetFloatArray->Num()))
 						{
-							MinDis = CurrentDistance;
-							MinIndex = TriangleIntersections[j];
-							ClosestBary = Bary;
-						}
-
-						if (MinIndex != -1
-							&& MinIndex != (*IndicesArray)[i][0]
-							&& MinIndex != (*IndicesArray)[i][1]
-							&& MinIndex != (*IndicesArray)[i][2])
-						{
-							(*TargetFloatArray)[MinIndex+ TargetVertexStartVal] = 0;
+							float Value = 0.f;
 							for (int32 k = 0; k < 3; k++)
 							{
-								(*TargetFloatArray)[MinIndex+ TargetVertexStartVal] += ClosestBary[k] * (*FloatArray)[Triangle[k]+VertexStartVal];
+								Value += Bary[k] * (*FloatArray)[Sample.Indices[i][k]];
 							}
+							(*TargetFloatArray)[TargetIndex] = FMath::Max((*TargetFloatArray)[TargetIndex], Value);
 						}
 					}
 				}
 			}
 		}
-	});
+	}
 }
 
-void FGeometryCollectionTransferVertexScalarAttributeNode::NearestVertexTransfer(FString AttributName, const FManagedArrayCollection& AttributeCollectionVal, const FManagedArrayCollection& CollectionVal, TManagedArray<float>* TargetFloatArray) const
+
+float FGeometryCollectionTransferVertexScalarAttributeNode::MaxEdgeLength(TArray<FVector>& Vert, const TManagedArray<FIntVector3>& Tri, int VertexOFfset, int TriStart, int TriCount)
 {
-	auto BuildComponentSpaceVertices = [](const TManagedArray<FTransform3f>* LocalSpaceTransform,
-		const TManagedArray<int32>* Parent, const TManagedArray<int32>* BoneMapArray,
-		const TManagedArray<FVector3f>* VertexArray,
-		TArray<FVector>& ComponentSpaceVertices)
-	{
-		TArray<FTransform> ComponentTransform;
-		GeometryCollectionAlgo::GlobalMatrices(*LocalSpaceTransform, *Parent, ComponentTransform);
+	auto TriInRange = [](const FIntVector3& T, int Max) {
+		for (int k = 0; k < 3; k++) if (T[0] < 0 || Max <= T[0]) return false;
+		return true;
+	};
 
-		ComponentSpaceVertices.SetNumUninitialized(VertexArray->Num());
-		for (int i = 0; i < VertexArray->Num(); i++)
+	float Max = 0;
+	int TriStop = TriStart + TriCount;
+	for (int i = TriStart; i < TriStop; i++)
+	{
+		if (TriInRange(Tri[i] - FIntVector3(VertexOFfset), Vert.Num()))
 		{
-			if (0 < (*BoneMapArray)[i] && (*BoneMapArray)[i] < ComponentTransform.Num())
-			{
-				ComponentSpaceVertices[i] = ComponentTransform[(*BoneMapArray)[i]].TransformPosition(FVector((*VertexArray)[i]));
-			}
-			else
-			{
-				ComponentSpaceVertices[i] = FVector((*VertexArray)[i]);
-			}
+			Max = FMath::Max(Max, FVector3f(Vert[Tri[i][0] - VertexOFfset] - Vert[Tri[i][1] - VertexOFfset]).SquaredLength());
+			Max = FMath::Max(Max, FVector3f(Vert[Tri[i][0] - VertexOFfset] - Vert[Tri[i][2] - VertexOFfset]).SquaredLength());
+			Max = FMath::Max(Max, FVector3f(Vert[Tri[i][1] - VertexOFfset] - Vert[Tri[i][2] - VertexOFfset]).SquaredLength());
 		}
-	};
-
-	const TManagedArray<int32>* TargetBoneMapArray = CollectionVal.FindAttribute<int32>(FName("BoneMap"), FName("Vertices"));
-	const TManagedArray<FVector3f>* TargetVertexArray = CollectionVal.FindAttribute<FVector3f>(FName("Vertex"), FName("Vertices"));
-	const TManagedArray<FIntVector3>* TargetIndicesArray = CollectionVal.FindAttribute<FIntVector3>(FName("Indices"), FName("Faces"));
-	const TManagedArray<FTransform3f>* TargetLocalSpaceTransform = CollectionVal.FindAttribute<FTransform3f>(FTransformCollection::TransformAttribute, FTransformCollection::TransformGroup);
-	const TManagedArray<int32>* TargetParent = CollectionVal.FindAttribute<int32>(FTransformCollection::ParentAttribute, FTransformCollection::TransformGroup);
-
-	const TManagedArray<float>* FloatArray = AttributeCollectionVal.FindAttribute<float>(FName(AttributName), FName("Vertices"));
-	const TManagedArray<int32>* BoneMapArray = AttributeCollectionVal.FindAttribute<int32>(FName("BoneMap"), FName("Vertices"));
-	const TManagedArray<FVector3f>* VertexArray = AttributeCollectionVal.FindAttribute<FVector3f>(FName("Vertex"), FName("Vertices"));
-	const TManagedArray<FIntVector3>* IndicesArray = AttributeCollectionVal.FindAttribute<FIntVector3>(FName("Indices"), FName("Faces"));
-	const TManagedArray<FTransform3f>* LocalSpaceTransform = AttributeCollectionVal.FindAttribute<FTransform3f>(FTransformCollection::TransformAttribute, FTransformCollection::TransformGroup);
-	const TManagedArray<int32>* Parent = AttributeCollectionVal.FindAttribute<int32>(FTransformCollection::ParentAttribute, FTransformCollection::TransformGroup);
-
-	// Build component space vertices for TargetCollection
-	TArray<FVector> ComponentSpaceTargetVertices;
-	BuildComponentSpaceVertices(TargetLocalSpaceTransform, TargetParent, TargetBoneMapArray, TargetVertexArray, ComponentSpaceTargetVertices);
-
-	// Build component space vertices for TargetCollection
-	TArray<FVector> ComponentSpaceVertices;
-	BuildComponentSpaceVertices(LocalSpaceTransform, Parent, BoneMapArray, VertexArray, ComponentSpaceVertices);
-
+	}
+	return FMath::Sqrt(Max);
+}
 
-	// build Sphere based BVH
-	Chaos::FReal SphereRadius = (Chaos::FReal)0.;
+void FGeometryCollectionTransferVertexScalarAttributeNode::BuildComponentSpaceVertices(const TManagedArray<FTransform3f>* LocalSpaceTransform, const TManagedArray<int32>* Parent, const TManagedArray<int32>* BoneMapArray,
+	const TManagedArray<FVector3f>* VertexArray, int32 Start, int32 Count, TArray<FVector>& ComponentSpaceVertices)
+{
+	TArray<FTransform> ComponentTransform;
+	GeometryCollectionAlgo::GlobalMatrices(*LocalSpaceTransform, *Parent, ComponentTransform);
 
-	Chaos::TVec3<float> CoordMaxs(-FLT_MAX);
-	Chaos::TVec3<float> CoordMins(FLT_MAX);
-	for (int32 i = 0; i < ComponentSpaceTargetVertices.Num(); i++)
+	ComponentSpaceVertices.SetNumUninitialized(Count);
+	for (int i = 0; i < Count; i++)
 	{
-		for (int32 j = 0; j < 3; j++)
+		int j = i + Start;
+		if (0 < (*BoneMapArray)[i] && (*BoneMapArray)[i] < ComponentTransform.Num())
 		{
-			if (ComponentSpaceTargetVertices[i][j] > CoordMaxs[j])
-			{
-				CoordMaxs[j] = ComponentSpaceTargetVertices[i][j];
-			}
-			if (ComponentSpaceTargetVertices[i][j] < CoordMins[j])
-			{
-				CoordMins[j] = ComponentSpaceTargetVertices[i][j];
-			}
+			ComponentSpaceVertices[i] = ComponentTransform[(*BoneMapArray)[j]].TransformPosition(FVector((*VertexArray)[j]));
+		}
+		else
+		{
+			ComponentSpaceVertices[i] = FVector((*VertexArray)[j]);
 		}
 	}
-	Chaos::TVec3<float> CoordDiff = (CoordMaxs - CoordMins) * VertexRadiusRatio;
-	SphereRadius = Chaos::FReal(FGenericPlatformMath::Max(CoordDiff[0], FGenericPlatformMath::Max(CoordDiff[1], CoordDiff[2])));
+}
 
+UE::Private::BVH* FGeometryCollectionTransferVertexScalarAttributeNode::BuildParticleSphereBVH(const TArray<FVector>& Vertices, float Radius)
+{
 	TArray<Chaos::TSphere<Chaos::FReal, 3>*> VertexSpherePtrs;
 	TArray<Chaos::TSphere<Chaos::FReal, 3>> VertexSpheres;
 
-	VertexSpheres.Init(Chaos::TSphere<Chaos::FReal, 3>(Chaos::TVec3<Chaos::FReal>(0), SphereRadius), ComponentSpaceTargetVertices.Num());
-	VertexSpherePtrs.SetNum(ComponentSpaceTargetVertices.Num());
+	VertexSpheres.Init(Chaos::TSphere<Chaos::FReal, 3>(Chaos::TVec3<Chaos::FReal>(0), Radius), Vertices.Num());
+	VertexSpherePtrs.SetNum(Vertices.Num());
 
-	for (int32 i = 0; i < ComponentSpaceTargetVertices.Num(); i++)
+	for (int32 i = 0; i < Vertices.Num(); i++)
 	{
-		Chaos::TVec3<Chaos::FReal> SphereCenter(ComponentSpaceTargetVertices[i]);
-		Chaos::TSphere<Chaos::FReal, 3> VertexSphere(SphereCenter, SphereRadius);
-		VertexSpheres[i] = Chaos::TSphere<Chaos::FReal, 3>(SphereCenter, SphereRadius);
+		Chaos::TVec3<Chaos::FReal> SphereCenter(Vertices[i]);
+		Chaos::TSphere<Chaos::FReal, 3> VertexSphere(SphereCenter, Radius);
+		VertexSpheres[i] = Chaos::TSphere<Chaos::FReal, 3>(SphereCenter, Radius);
 		VertexSpherePtrs[i] = &VertexSpheres[i];
 	}
-	Chaos::TBoundingVolumeHierarchy<
-		TArray<Chaos::TSphere<Chaos::FReal, 3>*>,
-		TArray<int32>,
-		Chaos::FReal,
-		3> VertexBVH(VertexSpherePtrs);
+	return new UE::Private::BVH(VertexSpherePtrs);
+}
 
-	for (int32 i = 0; i < IndicesArray->Num(); i++)
-	{
-		TArray<int32> TriangleIntersections0 = VertexBVH.FindAllIntersections(ComponentSpaceVertices[(*IndicesArray)[i][0]]);
-		TArray<int32> TriangleIntersections1 = VertexBVH.FindAllIntersections(ComponentSpaceVertices[(*IndicesArray)[i][1]]);
-		TArray<int32> TriangleIntersections2 = VertexBVH.FindAllIntersections(ComponentSpaceVertices[(*IndicesArray)[i][2]]);
-		TriangleIntersections0.Sort();
-		TriangleIntersections1.Sort();
-		TriangleIntersections2.Sort();
-
-		TArray<int32> TriangleIntersections({});
-		for (int32 k = 0; k < TriangleIntersections0.Num(); k++)
-		{
-			if (TriangleIntersections1.Contains(TriangleIntersections0[k])
-				&& TriangleIntersections2.Contains(TriangleIntersections0[k]))
-			{
-				TriangleIntersections.Emplace(TriangleIntersections0[k]);
-			}
-		}
+bool FGeometryCollectionTransferVertexScalarAttributeNode::TriangleHasWeightsToTransfer(const FIntVector3& T, const TManagedArray<float>& F)
+{
+	return !FMath::IsNearlyZero((F[T[0]] + F[T[1]] + F[T[2]]));
+}
 
-		int32 MinIndex = -1;
-		float MinDis = SphereRadius;
-		Chaos::TVector<float, 3> ClosestBary(0.f);
+void FGeometryCollectionTransferVertexScalarAttributeNode::TriangleToVertexIntersections(
+	UE::Private::BVH& VertexBVH, const TArray<FVector>& ComponentSpaceVertices, const FIntVector3& Triangle, TArray<int32>& OutTargetVertexIntersection)
+{
+	OutTargetVertexIntersection.Empty();
 
-		for (int32 j = 0; j < TriangleIntersections.Num(); j++)
-		{
-			Chaos::TVector<float, 3> Bary,
-				TriPos0(ComponentSpaceVertices[(*IndicesArray)[i][0]]),
-				TriPos1(ComponentSpaceVertices[(*IndicesArray)[i][1]]),
-				TriPos2(ComponentSpaceVertices[(*IndicesArray)[i][2]]),
-				ParticlePos(ComponentSpaceTargetVertices[TriangleIntersections[j]]);
-
-			Chaos::TVector<Chaos::FRealSingle, 3> ClosestPoint = Chaos::FindClosestPointAndBaryOnTriangle(TriPos0, TriPos1, TriPos2, ParticlePos, Bary);
-			Chaos::FRealSingle CurrentDistance = (ParticlePos - ClosestPoint).Size();
-			if (CurrentDistance < MinDis)
-			{
-				MinDis = CurrentDistance;
-				MinIndex = TriangleIntersections[j];
-				ClosestBary = Bary;
-			}
+	TArray<int32> TargetVertexIntersection0 = VertexBVH.FindAllIntersections(ComponentSpaceVertices[Triangle[0]]);
+	TArray<int32> TargetVertexIntersection1 = VertexBVH.FindAllIntersections(ComponentSpaceVertices[Triangle[1]]);
+	TArray<int32> TargetVertexIntersection2 = VertexBVH.FindAllIntersections(ComponentSpaceVertices[Triangle[2]]);
+	TargetVertexIntersection0.Sort();
+	TargetVertexIntersection1.Sort();
+	TargetVertexIntersection2.Sort();
 
-			if (MinIndex != -1
-				&& MinIndex != (*IndicesArray)[i][0]
-				&& MinIndex != (*IndicesArray)[i][1]
-				&& MinIndex != (*IndicesArray)[i][2])
-			{
-				(*TargetFloatArray)[MinIndex] = 0;
-				for (int32 k = 0; k < 3; k++)
-				{
-					(*TargetFloatArray)[MinIndex] += ClosestBary[k] * (*FloatArray)[(*IndicesArray)[i][k]];
-				}
-			}
+	for (int32 k = 0; k < TargetVertexIntersection0.Num(); k++)
+	{
+		if (TargetVertexIntersection1.Contains(TargetVertexIntersection0[k])
+			&& TargetVertexIntersection2.Contains(TargetVertexIntersection0[k]))
+		{
+			OutTargetVertexIntersection.Emplace(TargetVertexIntersection0[k]);
 		}
 	}
 }
 
 
-
 #undef LOCTEXT_NAMESPACE
diff --git a/Engine/Plugins/Experimental/GeometryCollectionPlugin/Source/GeometryCollectionNodes/Public/Dataflow/GeometryCollectionTransferVertexScalarAttributeNode.h b/Engine/Plugins/Experimental/GeometryCollectionPlugin/Source/GeometryCollectionNodes/Public/Dataflow/GeometryCollectionTransferVertexScalarAttributeNode.h
--- a/Engine/Plugins/Experimental/GeometryCollectionPlugin/Source/GeometryCollectionNodes/Public/Dataflow/GeometryCollectionTransferVertexScalarAttributeNode.h
+++ b/Engine/Plugins/Experimental/GeometryCollectionPlugin/Source/GeometryCollectionNodes/Public/Dataflow/GeometryCollectionTransferVertexScalarAttributeNode.h
@@ -2,11 +2,20 @@
 
 #pragma once 
 
+#include "Chaos/BoundingVolumeHierarchy.h"
 #include "Dataflow/DataflowNode.h"
 #include "GeometryCollection/ManagedArrayCollection.h"
 
 #include "GeometryCollectionTransferVertexScalarAttributeNode.generated.h"
 
+
+namespace UE::Private {
+
+	class FTransferFacade;
+	typedef Chaos::TBoundingVolumeHierarchy<TArray<Chaos::TSphere<Chaos::FReal, 3>*>, TArray<int32>, Chaos::FReal, 3> BVH;
+};
+
+
 /** Convert an vertex float array to a list of indices */
 USTRUCT(meta = (DataflowGeometryCollection))
 struct FGeometryCollectionTransferVertexScalarAttributeNode : public FDataflowNode
@@ -22,12 +31,18 @@ public:
 	UPROPERTY(Meta = (DataflowInput, DisplayName = "FromCollection"))
 	FManagedArrayCollection FromCollection;
 
-	/** The name of the vertex attribute to generate indices from. */
+	// The name of the vertex attribute to generate indices from.
 	UPROPERTY(EditAnywhere, Category = "Dataflow", Meta = (DataflowInput, DataflowOutput, DisplayName = "Name", DataflowPassthrough = "Name"))
 	FString Name = FString("");
 
-	UPROPERTY(EditAnywhere, Category = "Dataflow", meta = (ClampMin = "0.0", ClampMax = "2.0", UIMin = "0.0", UIMax = "2.0"))
-	float VertexRadiusRatio = .001f;
+	// Threshold based on distance from sample triangle. Anything past the threshold is not sampled. [Defualts to 1 percent (0.01)]
+	UPROPERTY(EditAnywhere, Category = "Dataflow|Scales")
+	float  FalloffThreshold = 0.01f;
+
+	// Scalar for the Bounding Volume Hierarchy (BVH) target's particle search radius.
+	UPROPERTY(EditAnywhere, Category = "Dataflow|Scales")
+	float  EdgeMultiplier = 1.0f;
+
 
 	FGeometryCollectionTransferVertexScalarAttributeNode(const Dataflow::FNodeParameters& InParam, FGuid InGuid = FGuid::NewGuid())
 		: FDataflowNode(InParam, InGuid)
@@ -44,8 +59,20 @@ private:
 
 	virtual void Evaluate(Dataflow::FContext& Context, const FDataflowOutput* Out) const override;
 
+
+	//
+	// private helpers
+	//
 	TArray<FIntVector2> FindSourceToTargetGeometryMap(const FManagedArrayCollection& SourceCollection, const FManagedArrayCollection& TargetCollection) const;
-	void PairedGeometryTransfer(FString AttributName, const TArray<FIntVector2>& PairedGeometry, const FManagedArrayCollection& AttributeCollectionVal, const FManagedArrayCollection& CollectionVal, TManagedArray<float>* FloatArray) const;
-	void NearestVertexTransfer(FString AttributName, const FManagedArrayCollection& AttributeCollectionVal, const FManagedArrayCollection& CollectionVal, TManagedArray<float>* FloatArray) const;
+	void PairedGeometryTransfer(FString AttributName, const TArray<FIntVector2>& PairedGeometry,const UE::Private::FTransferFacade& Sample, UE::Private::FTransferFacade& Target, TManagedArray<float>* TargetFloatArray) const;
+	void NearestVertexTransfer(FString AttributName, const UE::Private::FTransferFacade& Sample, UE::Private::FTransferFacade& Target, TManagedArray<float>* TargetFloatArray) const;
+	static float MaxEdgeLength(TArray<FVector>& Vert, const TManagedArray<FIntVector3>& Tri, int VertexOFfset, int TriStart, int TriCount);
+	static void BuildComponentSpaceVertices(const TManagedArray<FTransform3f>* LocalSpaceTransform, const TManagedArray<int32>* Parent, const TManagedArray<int32>* BoneMapArray, const TManagedArray<FVector3f>* VertexArray, int32 Start, int32 Count, TArray<FVector>& ComponentSpaceVertices);
+	static UE::Private::BVH* BuildParticleSphereBVH(const TArray<FVector>& Vertices, float Radius);
+	static bool TriangleHasWeightsToTransfer(const FIntVector3& T, const TManagedArray<float>& F);
+	static void TriangleToVertexIntersections(UE::Private::BVH& VertexBVH, const TArray<FVector>& ComponentSpaceVertices, const FIntVector3& Triangle, TArray<int32>& OutTargetVertexIntersection);
+
 
 };
+
+
