ChaosFlesh : Added falloff to the sample values.

diff --git a/Engine/Plugins/Experimental/GeometryCollectionPlugin/Source/GeometryCollectionNodes/Private/Dataflow/GeometryCollectionTransferVertexScalarAttributeNode.cpp b/Engine/Plugins/Experimental/GeometryCollectionPlugin/Source/GeometryCollectionNodes/Private/Dataflow/GeometryCollectionTransferVertexScalarAttributeNode.cpp
--- a/Engine/Plugins/Experimental/GeometryCollectionPlugin/Source/GeometryCollectionNodes/Private/Dataflow/GeometryCollectionTransferVertexScalarAttributeNode.cpp
+++ b/Engine/Plugins/Experimental/GeometryCollectionPlugin/Source/GeometryCollectionNodes/Private/Dataflow/GeometryCollectionTransferVertexScalarAttributeNode.cpp
@@ -203,7 +203,8 @@ void FGeometryCollectionTransferVertexScalarAttributeNode::PairedGeometryTransfe
 								Chaos::TVector<Chaos::FRealSingle, 3> ClosestPoint = Chaos::FindClosestPointAndBaryOnTriangle(TriPos0, TriPos1, TriPos2, ParticlePos, Bary);
 								Chaos::FRealSingle CurrentDistance = (ParticlePos - ClosestPoint).Size();
 								float TriRadius = FalloffThreshold * MaxEdgeLength(ComponentSpaceVertices, Sample.Indices.Get(), 0, i, 1);
-								if (CurrentDistance < TriRadius)
+								float FalloffScale = CalculateFalloffScale(Falloff, TriRadius, CurrentDistance);
+								if (!FMath::IsNearlyZero(FalloffScale))
 								{
 									int32 TargetIndex = TargetVertexIntersection[j] + TargetVertexStartVal;
 									if (ensure(0 <= TargetIndex && TargetIndex < TargetFloatArray->Num()))
@@ -211,7 +212,7 @@ void FGeometryCollectionTransferVertexScalarAttributeNode::PairedGeometryTransfe
 										float Value = 0.f;
 										for (int32 k = 0; k < 3; k++)
 										{
-											Value += Bary[k] * (*FloatArray)[ComponentTriangle[k] + VertexStartVal];
+											Value += FalloffScale * (Bary[k] * (*FloatArray)[ComponentTriangle[k] + VertexStartVal]);
 										}
 										(*TargetFloatArray)[TargetIndex] = FMath::Max((*TargetFloatArray)[TargetIndex], Value);
 									}
@@ -262,7 +263,8 @@ void FGeometryCollectionTransferVertexScalarAttributeNode::NearestVertexTransfer
 					Chaos::TVector<Chaos::FRealSingle, 3> ClosestPoint = Chaos::FindClosestPointAndBaryOnTriangle(TriPos0, TriPos1, TriPos2, ParticlePos, Bary);
 					Chaos::FRealSingle CurrentDistance = (ParticlePos - ClosestPoint).Size();
 					float TriRadius = FalloffThreshold * MaxEdgeLength(ComponentSpaceVertices, Sample.Indices.Get(), 0, i, 1);
-					if (CurrentDistance < TriRadius)
+					float FalloffScale = CalculateFalloffScale(Falloff, TriRadius, CurrentDistance);
+					if (!FMath::IsNearlyZero(FalloffScale))
 					{
 						int32 TargetIndex = TargetVertexIntersection[j];
 						if (ensure(0 <= TargetIndex && TargetIndex < TargetFloatArray->Num()))
@@ -270,7 +272,7 @@ void FGeometryCollectionTransferVertexScalarAttributeNode::NearestVertexTransfer
 							float Value = 0.f;
 							for (int32 k = 0; k < 3; k++)
 							{
-								Value += Bary[k] * (*FloatArray)[Sample.Indices[i][k]];
+								Value += FalloffScale * (Bary[k] * (*FloatArray)[Sample.Indices[i][k]]);
 							}
 							(*TargetFloatArray)[TargetIndex] = FMath::Max((*TargetFloatArray)[TargetIndex], Value);
 						}
@@ -370,4 +372,22 @@ void FGeometryCollectionTransferVertexScalarAttributeNode::TriangleToVertexInter
 }
 
 
+float FGeometryCollectionTransferVertexScalarAttributeNode::CalculateFalloffScale(EDataflowTransferNodeFalloff FalloffSetting, float Threshold, float Distance)
+{
+	float Demonator = 1.0;
+	if (Distance > Threshold && !FMath::IsNearlyZero(Threshold))
+	{
+		Demonator = Distance / Threshold;
+	}
+	switch (FalloffSetting)
+	{
+	case EDataflowTransferNodeFalloff::Dataflow_Transfer_Linear:
+		return 1. / Demonator;
+	case EDataflowTransferNodeFalloff::Dataflow_Transfer_Squared:
+		return 1. / FMath::Square(Demonator);
+	}
+	return Demonator;
+}
+
+
 #undef LOCTEXT_NAMESPACE
diff --git a/Engine/Plugins/Experimental/GeometryCollectionPlugin/Source/GeometryCollectionNodes/Public/Dataflow/GeometryCollectionTransferVertexScalarAttributeNode.h b/Engine/Plugins/Experimental/GeometryCollectionPlugin/Source/GeometryCollectionNodes/Public/Dataflow/GeometryCollectionTransferVertexScalarAttributeNode.h
--- a/Engine/Plugins/Experimental/GeometryCollectionPlugin/Source/GeometryCollectionNodes/Public/Dataflow/GeometryCollectionTransferVertexScalarAttributeNode.h
+++ b/Engine/Plugins/Experimental/GeometryCollectionPlugin/Source/GeometryCollectionNodes/Public/Dataflow/GeometryCollectionTransferVertexScalarAttributeNode.h
@@ -2,6 +2,7 @@
 
 #pragma once 
 
+#include "CoreMinimal.h"
 #include "Chaos/BoundingVolumeHierarchy.h"
 #include "Dataflow/DataflowNode.h"
 #include "GeometryCollection/ManagedArrayCollection.h"
@@ -16,7 +17,26 @@ namespace UE::Private {
 };
 
 
-/** Convert an vertex float array to a list of indices */
+UENUM(BlueprintType)
+enum class EDataflowTransferNodeFalloff : uint8
+{
+	/** Squared falloff based on distance from triangle*/
+	Dataflow_Transfer_Squared  UMETA(DisplayName = "Squared"),
+
+	/** Linear falloff based on distance from triangle*/
+	Dataflow_Transfer_Linear UMETA(DisplayName = "Linear"),
+
+	/** No distance falloff */
+	Dataflow_Transfer_None UMETA(DisplayName = "None"),
+	//~~~
+	//256th entry
+	Dataflow_Transfer_Dataflow_Max UMETA(Hidden)
+};
+
+
+/**
+ * Transfer float properties from a sample collection to a target collection. 
+ */
 USTRUCT(meta = (DataflowGeometryCollection))
 struct FGeometryCollectionTransferVertexScalarAttributeNode : public FDataflowNode
 {
@@ -25,22 +45,28 @@ struct FGeometryCollectionTransferVertexScalarAttributeNode : public FDataflowNo
 
 public:
 
+	/* Target collection to transfer vertex float attribute to. */
 	UPROPERTY(Meta = (DataflowInput, DataflowOutput, DisplayName = "Collection", DataflowPassthrough = "Collection"))
 	FManagedArrayCollection Collection;
 
+	/* Sample collection to transfer vertex float attribute from. */
 	UPROPERTY(Meta = (DataflowInput, DisplayName = "FromCollection"))
 	FManagedArrayCollection FromCollection;
 
-	// The name of the vertex attribute to generate indices from.
+	/* The name of the vertex attribute to generate indices from. */
 	UPROPERTY(EditAnywhere, Category = "Dataflow", Meta = (DataflowInput, DataflowOutput, DisplayName = "Name", DataflowPassthrough = "Name"))
 	FString Name = FString("");
 
-	// Threshold based on distance from sample triangle. Anything past the threshold is not sampled. [Defualts to 1 percent (0.01)]
-	UPROPERTY(EditAnywhere, Category = "Dataflow|Scales")
+	/* Falloff of sample value based on distance from sample triangle[default: Squared] */
+	UPROPERTY(EditAnywhere, Category = "Thresholds")
+	EDataflowTransferNodeFalloff Falloff = EDataflowTransferNodeFalloff::Dataflow_Transfer_Squared;
+
+	/* Threshold based on distance from sample triangle.Values sampled past the threshold will falloff.[Defualts to 1 percent of triangle size(0.01)] */
+	UPROPERTY(EditAnywhere, Category = "Thresholds")
 	float  FalloffThreshold = 0.01f;
 
-	// Scalar for the Bounding Volume Hierarchy (BVH) target's particle search radius.
-	UPROPERTY(EditAnywhere, Category = "Dataflow|Scales")
+	/* Scalar for the Bounding Volume Hierarchy(BVH) target's particle search radius. */
+	UPROPERTY(EditAnywhere, Category = "Thresholds")
 	float  EdgeMultiplier = 1.0f;
 
 
@@ -71,7 +97,7 @@ private:
 	static UE::Private::BVH* BuildParticleSphereBVH(const TArray<FVector>& Vertices, float Radius);
 	static bool TriangleHasWeightsToTransfer(const FIntVector3& T, const TManagedArray<float>& F);
 	static void TriangleToVertexIntersections(UE::Private::BVH& VertexBVH, const TArray<FVector>& ComponentSpaceVertices, const FIntVector3& Triangle, TArray<int32>& OutTargetVertexIntersection);
-
+	static float CalculateFalloffScale(EDataflowTransferNodeFalloff FalloffSetting, float Threshold, float Distance);
 
 };
 
