Datarflow : Updated node outputs to support FCollectionKey

diff --git a/Engine/Plugins/Experimental/ChaosFlesh/Source/ChaosFleshDeprecatedNodes/Private/Dataflow/ChaosFleshDeprecatedNodesPlugin.cpp b/Engine/Plugins/Experimental/ChaosFlesh/Source/ChaosFleshDeprecatedNodes/Private/Dataflow/ChaosFleshDeprecatedNodesPlugin.cpp
--- a/Engine/Plugins/Experimental/ChaosFlesh/Source/ChaosFleshDeprecatedNodes/Private/Dataflow/ChaosFleshDeprecatedNodesPlugin.cpp
+++ b/Engine/Plugins/Experimental/ChaosFlesh/Source/ChaosFleshDeprecatedNodes/Private/Dataflow/ChaosFleshDeprecatedNodesPlugin.cpp
@@ -6,6 +6,7 @@
 #include "Dataflow/ChaosFleshConstructTetGridNode.h"
 #include "Dataflow/ChaosFleshImportGEO.h"
 #include "Dataflow/ChaosFleshKinematicOriginInsertionInitializationNode.h"
+#include "Dataflow/ChaosFleshKinematicTetrahedralConstraintNode.h"
 
 
 #define LOCTEXT_NAMESPACE "ChaosFleshDeprecatedNodes"
@@ -17,6 +18,7 @@ void IChaosFleshDeprecatedNodesPlugin::StartupModule()
 	DATAFLOW_NODE_REGISTER_CREATION_FACTORY(FGenerateTetrahedralCollectionDataflowNodes);
 	DATAFLOW_NODE_REGISTER_CREATION_FACTORY(FConstructTetGridNode);
 	DATAFLOW_NODE_REGISTER_CREATION_FACTORY(FKinematicOriginInsertionInitializationDataflowNode);
+	DATAFLOW_NODE_REGISTER_CREATION_FACTORY(FKinematicTetrahedralBindingsDataflowNode);
 }
 
 void IChaosFleshDeprecatedNodesPlugin::ShutdownModule()
diff --git a/Engine/Plugins/Experimental/Dataflow/Source/DataflowNodes/Private/Dataflow/DataflowCollectionAddScalarVertexPropertyNode.cpp b/Engine/Plugins/Experimental/Dataflow/Source/DataflowNodes/Private/Dataflow/DataflowCollectionAddScalarVertexPropertyNode.cpp
--- a/Engine/Plugins/Experimental/Dataflow/Source/DataflowNodes/Private/Dataflow/DataflowCollectionAddScalarVertexPropertyNode.cpp
+++ b/Engine/Plugins/Experimental/Dataflow/Source/DataflowNodes/Private/Dataflow/DataflowCollectionAddScalarVertexPropertyNode.cpp
@@ -14,7 +14,7 @@ FDataflowCollectionAddScalarVertexPropertyNode::FDataflowCollectionAddScalarVert
 {
 	RegisterInputConnection(&Collection);
 	RegisterOutputConnection(&Collection, &Collection);
-	RegisterOutputConnection(&Name);
+	RegisterOutputConnection(&AttributeKey);
 }
 
 void FDataflowCollectionAddScalarVertexPropertyNode::Evaluate(Dataflow::FContext& Context, const FDataflowOutput* Out) const
@@ -49,9 +49,9 @@ void FDataflowCollectionAddScalarVertexPropertyNode::Evaluate(Dataflow::FContext
 
 		SetValue(Context, MoveTemp(InCollection), &Collection);
 	}
-	else if (Out->IsA<FString>(&Name))
+	else if (Out->IsA<FCollectionAttributeKey>(&AttributeKey))
 	{
-		SetValue(Context, Name, &Name);
+		SetValue(Context, FCollectionAttributeKey(Name,"Vertices"), &AttributeKey);
 	}
 }
 
diff --git a/Engine/Plugins/Experimental/Dataflow/Source/DataflowNodes/Private/Dataflow/DataflowCollectionAttributeKeyNodes.cpp b/Engine/Plugins/Experimental/Dataflow/Source/DataflowNodes/Private/Dataflow/DataflowCollectionAttributeKeyNodes.cpp
--- a/Engine/Plugins/Experimental/Dataflow/Source/DataflowNodes/Private/Dataflow/DataflowCollectionAttributeKeyNodes.cpp
+++ b/Engine/Plugins/Experimental/Dataflow/Source/DataflowNodes/Private/Dataflow/DataflowCollectionAttributeKeyNodes.cpp
@@ -11,22 +11,22 @@ namespace Dataflow
 {
 	void DataflowCollectionAttributeKeyNodes()
 	{
-		DATAFLOW_NODE_REGISTER_CREATION_FACTORY(FMakeCollectionKeyDataflowNode);
-		DATAFLOW_NODE_REGISTER_CREATION_FACTORY(FBreakCollectionKeyDataflowNode);
+		DATAFLOW_NODE_REGISTER_CREATION_FACTORY(FMakeAttributeKeyDataflowNode);
+		DATAFLOW_NODE_REGISTER_CREATION_FACTORY(FBreakAttributeKeyDataflowNode);
 	}
 }
 
 
-void FMakeCollectionKeyDataflowNode::Evaluate(Dataflow::FContext& Context, const FDataflowOutput* Out) const
+void FMakeAttributeKeyDataflowNode::Evaluate(Dataflow::FContext& Context, const FDataflowOutput* Out) const
 {
-	FString GroupName = GetValue<FString>(Context, &GroupIn);
 	FString AttributeName = GetValue<FString>(Context, &AttributeIn);
-	SetValue(Context, FCollectionKey(GroupName, AttributeName), &CollectionKeyOut);
+	FString GroupName = GetValue<FString>(Context, &GroupIn);
+	SetValue(Context, FCollectionAttributeKey(GroupName, AttributeName), &AttributeKeyOut);
 }
 
-void FBreakCollectionKeyDataflowNode::Evaluate(Dataflow::FContext& Context, const FDataflowOutput* Out) const
+void FBreakAttributeKeyDataflowNode::Evaluate(Dataflow::FContext& Context, const FDataflowOutput* Out) const
 {
-	FCollectionKey CollectionKey = GetValue<FCollectionKey>(Context, &CollectionKeyIn);
-	SetValue(Context, CollectionKey.Group, &GroupOut);
+	FCollectionAttributeKey CollectionKey = GetValue<FCollectionAttributeKey>(Context, &AttributeKeyIn);
 	SetValue(Context, CollectionKey.Attribute, &AttributeOut);
+	SetValue(Context, CollectionKey.Group, &GroupOut);
 }
\ No newline at end of file
diff --git a/Engine/Plugins/Experimental/Dataflow/Source/DataflowNodes/Public/Dataflow/DataflowCollectionAddScalarVertexPropertyNode.h b/Engine/Plugins/Experimental/Dataflow/Source/DataflowNodes/Public/Dataflow/DataflowCollectionAddScalarVertexPropertyNode.h
--- a/Engine/Plugins/Experimental/Dataflow/Source/DataflowNodes/Public/Dataflow/DataflowCollectionAddScalarVertexPropertyNode.h
+++ b/Engine/Plugins/Experimental/Dataflow/Source/DataflowNodes/Public/Dataflow/DataflowCollectionAddScalarVertexPropertyNode.h
@@ -5,6 +5,7 @@
 #include "Dataflow/DataflowNode.h"
 #include "GeometryCollection/GeometryCollection.h"
 #include "GeometryCollection/ManagedArrayCollection.h"
+#include "Dataflow/DataflowCollectionAttributeKeyNodes.h"
 
 #include "DataflowCollectionAddScalarVertexPropertyNode.generated.h"
 
@@ -19,13 +20,16 @@ struct DATAFLOWNODES_API FDataflowCollectionAddScalarVertexPropertyNode : public
 
 public:
 
-	UPROPERTY(Meta = (Dataflowinput, DataflowOutput, DataflowPassthrough = "Collection"))
+	UPROPERTY(Meta = (DataflowInput, DataflowOutput, DataflowPassthrough = "Collection"))
 	FManagedArrayCollection Collection;
 
 	/** The name to be set as a weight map attribute. */
-	UPROPERTY(EditAnywhere, Category = "Vertex Attribute", Meta = (DataflowOutput))
+	UPROPERTY(EditAnywhere, Category = "Vertex Attribute")
 	FString Name;
 
+	UPROPERTY(meta = (DisplayName = "AttributeKey", DataflowOutput))
+	FCollectionAttributeKey AttributeKey;
+
 	UPROPERTY()
 	TArray<float> VertexWeights;
 
diff --git a/Engine/Plugins/Experimental/Dataflow/Source/DataflowNodes/Public/Dataflow/DataflowCollectionAttributeKeyNodes.h b/Engine/Plugins/Experimental/Dataflow/Source/DataflowNodes/Public/Dataflow/DataflowCollectionAttributeKeyNodes.h
--- a/Engine/Plugins/Experimental/Dataflow/Source/DataflowNodes/Public/Dataflow/DataflowCollectionAttributeKeyNodes.h
+++ b/Engine/Plugins/Experimental/Dataflow/Source/DataflowNodes/Public/Dataflow/DataflowCollectionAttributeKeyNodes.h
@@ -11,24 +11,29 @@ class USkeletalMesh;
 
 
 USTRUCT(BlueprintType)
-struct DATAFLOWNODES_API FCollectionKey
+struct DATAFLOWNODES_API FCollectionAttributeKey
 {
 	GENERATED_USTRUCT_BODY()
 public:
-	FCollectionKey() : Group(""), Attribute("") {}
-	FCollectionKey(FString InGroup, FString InAttribute)
-		:Group(InGroup), Attribute(InAttribute) {}
-	TPair<FName, FName> GetNamedKey() { return { FName(Group), FName(Attribute)}; }
-	FString Group;
+	FCollectionAttributeKey() : Attribute(""), Group("") {}
+	FCollectionAttributeKey(FString InAttribute, FString InGroup)
+		: Attribute(InAttribute), Group(InGroup) {}
+	TPair<FName, FName> GetNamedKey() { return { FName(Attribute), FName(Group)}; }
+
+	UPROPERTY(EditAnywhere, Category = "Collection Key")
 	FString Attribute;
+
+	UPROPERTY(EditAnywhere, Category = "Collection Key")
+	FString Group;
+
 };
 
 
 USTRUCT(meta = (Dataflow))
-struct FMakeCollectionKeyDataflowNode : public FDataflowNode
+struct FMakeAttributeKeyDataflowNode : public FDataflowNode
 {
 	GENERATED_USTRUCT_BODY()
-	DATAFLOW_NODE_DEFINE_INTERNAL(FMakeCollectionKeyDataflowNode, "MakeCollectionKey", "GeometryCollection", "")
+	DATAFLOW_NODE_DEFINE_INTERNAL(FMakeAttributeKeyDataflowNode, "MakeAttributeKey", "GeometryCollection", "")
 
 public:
 
@@ -38,44 +43,45 @@ public:
 	UPROPERTY(EditAnywhere, Category = "Collection Key", meta = (DataflowInput, DisplayName = "Attribute"))
 	FString AttributeIn = "";
 
-	UPROPERTY(meta = (DataflowOutput, DisplayName = "CollectionKey"))
-	FCollectionKey CollectionKeyOut;
+	UPROPERTY(meta = (DataflowOutput, DisplayName = "AttributeKey"))
+	FCollectionAttributeKey AttributeKeyOut;
 
-	FMakeCollectionKeyDataflowNode(const Dataflow::FNodeParameters& InParam, FGuid InGuid = FGuid::NewGuid())
+	FMakeAttributeKeyDataflowNode(const Dataflow::FNodeParameters& InParam, FGuid InGuid = FGuid::NewGuid())
 		: FDataflowNode(InParam, InGuid)
 	{
 		RegisterInputConnection(&GroupIn);
 		RegisterInputConnection(&AttributeIn);
-		RegisterOutputConnection(&CollectionKeyOut);
+		RegisterOutputConnection(&AttributeKeyOut);
 	}
 
 	virtual void Evaluate(Dataflow::FContext& Context, const FDataflowOutput* Out) const override;
 };
 
 USTRUCT(meta = (DataflowGeometryCollection))
-struct FBreakCollectionKeyDataflowNode : public FDataflowNode
+struct FBreakAttributeKeyDataflowNode : public FDataflowNode
 {
 	GENERATED_USTRUCT_BODY()
-	DATAFLOW_NODE_DEFINE_INTERNAL(FBreakCollectionKeyDataflowNode, "BreakCollectionKey", "GeometryCollection", "")
+	DATAFLOW_NODE_DEFINE_INTERNAL(FBreakAttributeKeyDataflowNode, "BreakAttributeKey", "Dataflow", "")
 
 public:
 
-	UPROPERTY(meta = (DataflowInput, DisplayName = "CollectionKey"))
-	FCollectionKey CollectionKeyIn;
+	UPROPERTY(meta = (DataflowInput, DisplayName = "AttributeKey"))
+	FCollectionAttributeKey AttributeKeyIn;
 
-	UPROPERTY(EditAnywhere, Category = "Collection Key", meta = (DataflowOutput, DisplayName = "Group"))
+	UPROPERTY(EditAnywhere, Category = "Attribute Key", meta = (DataflowOutput, DisplayName = "Attribute"))
+	FString AttributeOut = "";
+
+	UPROPERTY(EditAnywhere, Category = "Attribute Key", meta = (DataflowOutput, DisplayName = "Group"))
 	FString GroupOut = "";
 
-	UPROPERTY(EditAnywhere, Category = "Collection Key", meta = (DataflowOutput, DisplayName = "Attribute"))
-	FString AttributeOut = "";
 
 
-	FBreakCollectionKeyDataflowNode(const Dataflow::FNodeParameters& InParam, FGuid InGuid = FGuid::NewGuid())
+	FBreakAttributeKeyDataflowNode(const Dataflow::FNodeParameters& InParam, FGuid InGuid = FGuid::NewGuid())
 		: FDataflowNode(InParam, InGuid)
 	{
-		RegisterInputConnection(&CollectionKeyIn);
-		RegisterOutputConnection(&GroupOut);
+		RegisterInputConnection(&AttributeKeyIn);
 		RegisterOutputConnection(&AttributeOut);
+		RegisterOutputConnection(&GroupOut);
 	}
 
 	virtual void Evaluate(Dataflow::FContext& Context, const FDataflowOutput* Out) const override;
diff --git a/Engine/Plugins/Experimental/GeometryCollectionPlugin/Source/GeometryCollectionNodes/GeometryCollectionNodes.Build.cs b/Engine/Plugins/Experimental/GeometryCollectionPlugin/Source/GeometryCollectionNodes/GeometryCollectionNodes.Build.cs
--- a/Engine/Plugins/Experimental/GeometryCollectionPlugin/Source/GeometryCollectionNodes/GeometryCollectionNodes.Build.cs
+++ b/Engine/Plugins/Experimental/GeometryCollectionPlugin/Source/GeometryCollectionNodes/GeometryCollectionNodes.Build.cs
@@ -15,6 +15,7 @@ namespace UnrealBuildTool.Rules
 					"Chaos",
 					"DataflowCore",
 					"DataflowEngine",
+					"DataflowNodes",
 					"GeometryCollectionEngine",
 					"GeometryCore",
 					"Voronoi",
diff --git a/Engine/Plugins/Experimental/GeometryCollectionPlugin/Source/GeometryCollectionNodes/Private/Dataflow/GeometryCollectionNodes.cpp b/Engine/Plugins/Experimental/GeometryCollectionPlugin/Source/GeometryCollectionNodes/Private/Dataflow/GeometryCollectionNodes.cpp
--- a/Engine/Plugins/Experimental/GeometryCollectionPlugin/Source/GeometryCollectionNodes/Private/Dataflow/GeometryCollectionNodes.cpp
+++ b/Engine/Plugins/Experimental/GeometryCollectionPlugin/Source/GeometryCollectionNodes/Private/Dataflow/GeometryCollectionNodes.cpp
@@ -919,15 +919,15 @@ void FGetCollectionAttributeDataTypedDataflowNode::Evaluate(Dataflow::FContext&
 		Out->IsA<TArray<FVector3d>>(&Vector3dAttributeData))
 	{
 		const FManagedArrayCollection& InCollection = GetValue<FManagedArrayCollection>(Context, &Collection);
-
-		FName GroupNameToUse;
+		
+		FName InputGroupName;
 		if (GroupName != EStandardGroupNameEnum::Dataflow_EStandardGroupNameEnum_Custom)
 		{
-			GroupNameToUse = GetGroupName(GroupName);
+			InputGroupName = GetGroupName(GroupName);
 		}
 		else
 		{
-			GroupNameToUse = FName(*CustomGroupName);
+			InputGroupName = FName(*CustomGroupName);
 		}
 
 		SetValue(Context, TArray<bool>(), &BoolAttributeData);
@@ -938,48 +938,53 @@ void FGetCollectionAttributeDataTypedDataflowNode::Evaluate(Dataflow::FContext&
 		SetValue(Context, TArray<FVector3f>(), &Vector3fAttributeData);
 		SetValue(Context, TArray<FVector3d>(), &Vector3dAttributeData);
 
-		if (GroupNameToUse.GetStringLength() > 0 && AttrName.Len() > 0)
+		FCollectionAttributeKey DefaultAttributeKey(AttrName, InputGroupName.ToString());
+		FCollectionAttributeKey AttributeKeyVal = GetValue(Context, &AttributeKey, DefaultAttributeKey);
+		FName GroupNameVal = FName(AttributeKeyVal.Group);
+		FName AttributeNameVal = FName(AttributeKeyVal.Attribute);
+
+		if (GroupNameVal.GetStringLength() > 0 && AttributeNameVal.GetStringLength() > 0)
 		{
-			if (InCollection.HasGroup(GroupNameToUse))
+			if (InCollection.HasGroup(GroupNameVal))
 			{
-				if (InCollection.HasAttribute(FName(*AttrName), GroupNameToUse))
+				if (InCollection.HasAttribute(AttributeNameVal, GroupNameVal))
 				{
-					FString TypeStr = GetArrayTypeString(InCollection.GetAttributeType(FName(*AttrName), GroupNameToUse)).ToString();
+					FString TypeStr = GetArrayTypeString(InCollection.GetAttributeType(AttributeNameVal, GroupNameVal)).ToString();
 
 					if (TypeStr == FString("Bool"))
 					{
-						const TManagedArray<bool>& AttributeArr = InCollection.GetAttribute<bool>(FName(*AttrName), GroupNameToUse);
+						const TManagedArray<bool>& AttributeArr = InCollection.GetAttribute<bool>(AttributeNameVal, GroupNameVal);
 						TArray<bool> BoolArray = AttributeArr.GetAsBoolArray();
 						SetValue(Context, MoveTemp(BoolArray), &BoolAttributeData);
 					}
 					else if (TypeStr == FString("Float"))
 					{
-						const TManagedArray<float>& AttributeArr = InCollection.GetAttribute<float>(FName(*AttrName), GroupNameToUse);
+						const TManagedArray<float>& AttributeArr = InCollection.GetAttribute<float>(AttributeNameVal, GroupNameVal);
 						SetValue(Context, AttributeArr.GetConstArray(), &FloatAttributeData);
 					}
 					else if (TypeStr == FString("Double"))
 					{
-						const TManagedArray<double>& AttributeArr = InCollection.GetAttribute<double>(FName(*AttrName), GroupNameToUse);
+						const TManagedArray<double>& AttributeArr = InCollection.GetAttribute<double>(AttributeNameVal, GroupNameVal);
 						SetValue(Context, AttributeArr.GetConstArray(), &DoubleAttributeData);
 					}
 					else if (TypeStr == FString("Int32"))
 					{
-						const TManagedArray<int32>& AttributeArr = InCollection.GetAttribute<int32>(FName(*AttrName), GroupNameToUse);
+						const TManagedArray<int32>& AttributeArr = InCollection.GetAttribute<int32>(AttributeNameVal, GroupNameVal);
 						SetValue(Context, AttributeArr.GetConstArray(), &Int32AttributeData);
 					}
 					else if (TypeStr == FString("String"))
 					{
-						const TManagedArray<FString>& AttributeArr = InCollection.GetAttribute<FString>(FName(*AttrName), GroupNameToUse);
+						const TManagedArray<FString>& AttributeArr = InCollection.GetAttribute<FString>(AttributeNameVal, GroupNameVal);
 						SetValue(Context, AttributeArr.GetConstArray(), &StringAttributeData);
 					}
 					else if (TypeStr == FString("Vector"))
 					{
-						const TManagedArray<FVector3f>& AttributeArr = InCollection.GetAttribute<FVector3f>(FName(*AttrName), GroupNameToUse);
+						const TManagedArray<FVector3f>& AttributeArr = InCollection.GetAttribute<FVector3f>(AttributeNameVal, GroupNameVal);
 						SetValue(Context, AttributeArr.GetConstArray(), &Vector3fAttributeData);
 					}
 					else if (TypeStr == FString("Vector3d"))
 					{
-						const TManagedArray<FVector3d>& AttributeArr = InCollection.GetAttribute<FVector3d>(FName(*AttrName), GroupNameToUse);
+						const TManagedArray<FVector3d>& AttributeArr = InCollection.GetAttribute<FVector3d>(AttributeNameVal, GroupNameVal);
 						SetValue(Context, AttributeArr.GetConstArray(), &Vector3dAttributeData);
 					}
 				}
diff --git a/Engine/Plugins/Experimental/GeometryCollectionPlugin/Source/GeometryCollectionNodes/Private/Dataflow/GeometryCollectionTransferVertexScalarAttributeNode.cpp b/Engine/Plugins/Experimental/GeometryCollectionPlugin/Source/GeometryCollectionNodes/Private/Dataflow/GeometryCollectionTransferVertexScalarAttributeNode.cpp
--- a/Engine/Plugins/Experimental/GeometryCollectionPlugin/Source/GeometryCollectionNodes/Private/Dataflow/GeometryCollectionTransferVertexScalarAttributeNode.cpp
+++ b/Engine/Plugins/Experimental/GeometryCollectionPlugin/Source/GeometryCollectionNodes/Private/Dataflow/GeometryCollectionTransferVertexScalarAttributeNode.cpp
@@ -85,39 +85,40 @@ namespace UE::Private {
 
 void FGeometryCollectionTransferVertexScalarAttributeNode::Evaluate(Dataflow::FContext& Context, const FDataflowOutput* Out) const
 {
+	typedef TPair<FName, FName> NamePair;
+	FCollectionAttributeKey Key = GetValue(Context, &AttributeKey, AttributeKey);
 
-	FString AttributeName = GetValue<FString>(Context, &Name, Name);
-	if (Out->IsA< FManagedArrayCollection >(&Collection))
+	if (Out->IsA(&Collection))
 	{
-		FManagedArrayCollection TargetCollection = GetValue<FManagedArrayCollection>(Context, &Collection);
-		const FManagedArrayCollection& SampleCollection = GetValue<FManagedArrayCollection>(Context, &FromCollection);
+		FManagedArrayCollection TargetCollection = GetValue(Context, &Collection);
+		const FManagedArrayCollection& SampleCollection = GetValue(Context, &FromCollection);
 
 		UE::Private::FTransferFacade Target(TargetCollection);
 		const UE::Private::FTransferFacade Sample(SampleCollection);
 
 		if (Target.IsValid() && Sample.IsValid())
 		{
-			if (TManagedArray<float>* TargetFloatArray = Target.GetFloatArray(AttributeName, "Vertices"))
+			if (TManagedArray<float>* TargetFloatArray = Target.GetFloatArray(Key.Attribute, Key.Group))
 			{
 				TargetFloatArray->Fill(0.f);
 
 				TArray<FIntVector2> AlignedGeometry = FindSourceToTargetGeometryMap(SampleCollection, TargetCollection);
 				if (AlignedGeometry.Num() == TargetCollection.NumElements(FGeometryCollection::GeometryGroup))
 				{
-					PairedGeometryTransfer(AttributeName, AlignedGeometry, Sample, Target, TargetFloatArray);
+					PairedGeometryTransfer(Key, AlignedGeometry, Sample, Target, TargetFloatArray);
 				}
 				else
 				{
-					NearestVertexTransfer(AttributeName, Sample, Target, TargetFloatArray);
+					NearestVertexTransfer(Key, Sample, Target, TargetFloatArray);
 				}
 			}
 		}
 
-		SetValue<FManagedArrayCollection >(Context, MoveTemp(TargetCollection), &Collection);
+		SetValue(Context, MoveTemp(TargetCollection), &Collection);
 	}
-	else if (Out->IsA<FString>(&Name))
+	else if (Out->IsA(&AttributeKey))
 	{
-		SetValue< FString >(Context, MoveTemp(AttributeName), &Name);
+		SetValue(Context, MoveTemp(Key), &AttributeKey);
 	}
 }
 
@@ -148,10 +149,10 @@ TArray<FIntVector2> FGeometryCollectionTransferVertexScalarAttributeNode::FindSo
 	return Mapping;
 }
 
-void FGeometryCollectionTransferVertexScalarAttributeNode::PairedGeometryTransfer(FString AttributeName, const TArray<FIntVector2>& PairedGeometry, 
+void FGeometryCollectionTransferVertexScalarAttributeNode::PairedGeometryTransfer(FCollectionAttributeKey Key, const TArray<FIntVector2>& PairedGeometry,
 	const UE::Private::FTransferFacade& Sample, UE::Private::FTransferFacade&  Target, TManagedArray<float>* TargetFloatArray) const
 {
-	if (const TManagedArray<float>* FloatArray = Sample.GetFloatArray(AttributeName, "Vertices"))
+	if (const TManagedArray<float>* FloatArray = Sample.GetFloatArray(Key.Attribute, Key.Group))
 	{
 		ParallelFor(PairedGeometry.Num(), [&](int32 Pdx)
 		{
@@ -226,9 +227,9 @@ void FGeometryCollectionTransferVertexScalarAttributeNode::PairedGeometryTransfe
 	}
 }
 
-void FGeometryCollectionTransferVertexScalarAttributeNode::NearestVertexTransfer(FString AttributeName, const UE::Private::FTransferFacade& Sample, UE::Private::FTransferFacade& Target, TManagedArray<float>* TargetFloatArray) const
+void FGeometryCollectionTransferVertexScalarAttributeNode::NearestVertexTransfer(FCollectionAttributeKey Key, const UE::Private::FTransferFacade& Sample, UE::Private::FTransferFacade& Target, TManagedArray<float>* TargetFloatArray) const
 {
-	if (const TManagedArray<float>* FloatArray = Sample.GetFloatArray(AttributeName, "Vertices"))
+	if (const TManagedArray<float>* FloatArray = Sample.GetFloatArray(Key.Attribute, Key.Group))
 	{
 		// Build component space vertices for TargetCollection
 		TArray<FVector> ComponentSpaceTargetVertices;
diff --git a/Engine/Plugins/Experimental/GeometryCollectionPlugin/Source/GeometryCollectionNodes/Private/Dataflow/GeometryCollectionVertexScalarToVertexIndicesNode.cpp b/Engine/Plugins/Experimental/GeometryCollectionPlugin/Source/GeometryCollectionNodes/Private/Dataflow/GeometryCollectionVertexScalarToVertexIndicesNode.cpp
--- a/Engine/Plugins/Experimental/GeometryCollectionPlugin/Source/GeometryCollectionNodes/Private/Dataflow/GeometryCollectionVertexScalarToVertexIndicesNode.cpp
+++ b/Engine/Plugins/Experimental/GeometryCollectionPlugin/Source/GeometryCollectionNodes/Private/Dataflow/GeometryCollectionVertexScalarToVertexIndicesNode.cpp
@@ -11,7 +11,7 @@ FGeometryCollectionVertexScalarToVertexIndicesNode::FGeometryCollectionVertexSca
 	: FDataflowNode(InParam, InGuid)
 {
 	RegisterInputConnection(&Collection);
-	RegisterInputConnection(&VertexAttributeName);
+	RegisterInputConnection(&AttributeKey);
 	RegisterOutputConnection(&Indices);
 }
 
@@ -22,9 +22,9 @@ void FGeometryCollectionVertexScalarToVertexIndicesNode::Evaluate(Dataflow::FCon
 		TArray<int32> IndicesOut;
 
 		const FManagedArrayCollection& InCollection = GetValue<FManagedArrayCollection>(Context, &Collection);
-		FString VertixAttributeNameVal = GetValue<FString>(Context, &VertexAttributeName, VertexAttributeName);
+		FCollectionAttributeKey Key = GetValue(Context, &AttributeKey);
 
-		if( const TManagedArray<float>* FloatArray = InCollection.FindAttribute<float>(FName(VertixAttributeNameVal), FName("Vertices") ) )
+		if( const TManagedArray<float>* FloatArray = InCollection.FindAttribute<float>(FName(Key.Attribute), FName(Key.Group) ) )
 		{
 			for (int i = 0; i < FloatArray->Num(); i++)
 			{
diff --git a/Engine/Plugins/Experimental/GeometryCollectionPlugin/Source/GeometryCollectionNodes/Public/Dataflow/GeometryCollectionNodes.h b/Engine/Plugins/Experimental/GeometryCollectionPlugin/Source/GeometryCollectionNodes/Public/Dataflow/GeometryCollectionNodes.h
--- a/Engine/Plugins/Experimental/GeometryCollectionPlugin/Source/GeometryCollectionNodes/Public/Dataflow/GeometryCollectionNodes.h
+++ b/Engine/Plugins/Experimental/GeometryCollectionPlugin/Source/GeometryCollectionNodes/Public/Dataflow/GeometryCollectionNodes.h
@@ -4,6 +4,7 @@
 
 #include "CoreMinimal.h"
 #include "Dataflow/DataflowEngine.h"
+#include "Dataflow/DataflowCollectionAttributeKeyNodes.h"
 #include "GeometryCollection/ManagedArrayCollection.h"
 #include "DynamicMesh/DynamicMesh3.h"
 #include "UDynamicMesh.h"
@@ -1318,6 +1319,10 @@ public:
 	UPROPERTY(meta = (DataflowInput, DataflowIntrinsic))
 	FManagedArrayCollection Collection;
 
+	/** Input to drive the Attribute and Group name */
+	UPROPERTY(meta = (DataflowInput, DisplayName = "AttributeKey"))
+	FCollectionAttributeKey AttributeKey;
+
 	/** Standard group names */
 	UPROPERTY(EditAnywhere, Category = "Attribute", meta = (DisplayName = "Group"))
 	EStandardGroupNameEnum GroupName = EStandardGroupNameEnum::Dataflow_EStandardGroupNameEnum_Transform;
@@ -1362,6 +1367,7 @@ public:
 		: FDataflowNode(InParam, InGuid)
 	{
 		RegisterInputConnection(&Collection);
+		RegisterInputConnection(&AttributeKey);
 		RegisterOutputConnection(&BoolAttributeData);
 		RegisterOutputConnection(&FloatAttributeData);
 		RegisterOutputConnection(&DoubleAttributeData);
diff --git a/Engine/Plugins/Experimental/GeometryCollectionPlugin/Source/GeometryCollectionNodes/Public/Dataflow/GeometryCollectionTransferVertexScalarAttributeNode.h b/Engine/Plugins/Experimental/GeometryCollectionPlugin/Source/GeometryCollectionNodes/Public/Dataflow/GeometryCollectionTransferVertexScalarAttributeNode.h
--- a/Engine/Plugins/Experimental/GeometryCollectionPlugin/Source/GeometryCollectionNodes/Public/Dataflow/GeometryCollectionTransferVertexScalarAttributeNode.h
+++ b/Engine/Plugins/Experimental/GeometryCollectionPlugin/Source/GeometryCollectionNodes/Public/Dataflow/GeometryCollectionTransferVertexScalarAttributeNode.h
@@ -5,6 +5,7 @@
 #include "CoreMinimal.h"
 #include "Chaos/BoundingVolumeHierarchy.h"
 #include "Dataflow/DataflowNode.h"
+#include "Dataflow/DataflowCollectionAttributeKeyNodes.h"
 #include "GeometryCollection/ManagedArrayCollection.h"
 
 #include "GeometryCollectionTransferVertexScalarAttributeNode.generated.h"
@@ -54,8 +55,8 @@ public:
 	FManagedArrayCollection FromCollection;
 
 	/* The name of the vertex attribute to generate indices from. */
-	UPROPERTY(EditAnywhere, Category = "Dataflow", Meta = (DataflowInput, DataflowOutput, DisplayName = "Name", DataflowPassthrough = "Name"))
-	FString Name = FString("");
+	UPROPERTY(EditAnywhere, Category = "Dataflow", Meta = (DataflowInput, DataflowOutput, DisplayName = "AttributeKey", DataflowPassthrough = "AttributeKey"))
+	FCollectionAttributeKey AttributeKey;
 
 	/* Falloff of sample value based on distance from sample triangle[default: Squared] */
 	UPROPERTY(EditAnywhere, Category = "Thresholds")
@@ -75,10 +76,10 @@ public:
 	{
 		RegisterInputConnection(&Collection);
 		RegisterInputConnection(&FromCollection);
-		RegisterInputConnection(&Name);
+		RegisterInputConnection(&AttributeKey);
 
 		RegisterOutputConnection(&Collection, &Collection);
-		RegisterOutputConnection(&Name, &Name);
+		RegisterOutputConnection(&AttributeKey, &AttributeKey);
 	}
 
 private:
@@ -90,8 +91,8 @@ private:
 	// private helpers
 	//
 	TArray<FIntVector2> FindSourceToTargetGeometryMap(const FManagedArrayCollection& SourceCollection, const FManagedArrayCollection& TargetCollection) const;
-	void PairedGeometryTransfer(FString AttributName, const TArray<FIntVector2>& PairedGeometry,const UE::Private::FTransferFacade& Sample, UE::Private::FTransferFacade& Target, TManagedArray<float>* TargetFloatArray) const;
-	void NearestVertexTransfer(FString AttributName, const UE::Private::FTransferFacade& Sample, UE::Private::FTransferFacade& Target, TManagedArray<float>* TargetFloatArray) const;
+	void PairedGeometryTransfer(FCollectionAttributeKey Key, const TArray<FIntVector2>& PairedGeometry,const UE::Private::FTransferFacade& Sample, UE::Private::FTransferFacade& Target, TManagedArray<float>* TargetFloatArray) const;
+	void NearestVertexTransfer(FCollectionAttributeKey Key, const UE::Private::FTransferFacade& Sample, UE::Private::FTransferFacade& Target, TManagedArray<float>* TargetFloatArray) const;
 	static float MaxEdgeLength(TArray<FVector>& Vert, const TManagedArray<FIntVector3>& Tri, int VertexOFfset, int TriStart, int TriCount);
 	static void BuildComponentSpaceVertices(const TManagedArray<FTransform3f>* LocalSpaceTransform, const TManagedArray<int32>* Parent, const TManagedArray<int32>* BoneMapArray, const TManagedArray<FVector3f>* VertexArray, int32 Start, int32 Count, TArray<FVector>& ComponentSpaceVertices);
 	static UE::Private::BVH* BuildParticleSphereBVH(const TArray<FVector>& Vertices, float Radius);
diff --git a/Engine/Plugins/Experimental/GeometryCollectionPlugin/Source/GeometryCollectionNodes/Public/Dataflow/GeometryCollectionVertexScalarToVertexIndicesNode.h b/Engine/Plugins/Experimental/GeometryCollectionPlugin/Source/GeometryCollectionNodes/Public/Dataflow/GeometryCollectionVertexScalarToVertexIndicesNode.h
--- a/Engine/Plugins/Experimental/GeometryCollectionPlugin/Source/GeometryCollectionNodes/Public/Dataflow/GeometryCollectionVertexScalarToVertexIndicesNode.h
+++ b/Engine/Plugins/Experimental/GeometryCollectionPlugin/Source/GeometryCollectionNodes/Public/Dataflow/GeometryCollectionVertexScalarToVertexIndicesNode.h
@@ -2,7 +2,9 @@
 
 #pragma once 
 
+#include "Dataflow/DataflowCollectionAttributeKeyNodes.h"
 #include "Dataflow/DataflowNode.h"
+#include "GeometryCollection/GeometryCollection.h"
 #include "GeometryCollection/ManagedArrayCollection.h"
 
 #include "GeometryCollectionVertexScalarToVertexIndicesNode.generated.h"
@@ -19,9 +21,9 @@ public:
 	UPROPERTY(Meta = (DataflowInput))
 	FManagedArrayCollection Collection;
 
-	/** The name of the vertex attribute to generate indices from. */
-	UPROPERTY(EditAnywhere, Category = "Selection Filter", Meta = (DataflowInput))
-	FString VertexAttributeName;
+	/** The name of the vertex attribute and group to generate indices from. */
+	UPROPERTY(EditAnywhere, Category = "Dataflow", meta = (DataflowInput, DisplayName = "AttributeKey"))
+	FCollectionAttributeKey AttributeKey = FCollectionAttributeKey("", "Vertices");
 
 	/** The value threshold for what is included in the vertex list. */
 	UPROPERTY(EditAnywhere, Category = "Selection Filter", Meta = (ClampMin = "0", ClampMax = "1"))
