Dataflow : Fix Clear Convex Hull node
- Selection is no longer required
- Do not favor parent when sanitizzing the selection

diff --git a/Engine/Plugins/Experimental/GeometryCollectionPlugin/Source/GeometryCollectionNodes/Private/Dataflow/GeometryCollectionUtilityNodes.cpp b/Engine/Plugins/Experimental/GeometryCollectionPlugin/Source/GeometryCollectionNodes/Private/Dataflow/GeometryCollectionUtilityNodes.cpp
--- a/Engine/Plugins/Experimental/GeometryCollectionPlugin/Source/GeometryCollectionNodes/Private/Dataflow/GeometryCollectionUtilityNodes.cpp
+++ b/Engine/Plugins/Experimental/GeometryCollectionPlugin/Source/GeometryCollectionNodes/Private/Dataflow/GeometryCollectionUtilityNodes.cpp
@@ -403,20 +403,26 @@ void FClearConvexHullsDataflowNode::Evaluate(Dataflow::FContext& Context, const
 	{
 		FManagedArrayCollection InCollection = GetValue(Context, &Collection);
 
-		if (!IsConnected(&Collection) || !IsConnected(&TransformSelection) || !FGeometryCollectionConvexUtility::HasConvexHullData(&InCollection))
+		if (!IsConnected(&Collection) || !FGeometryCollectionConvexUtility::HasConvexHullData(&InCollection))
 		{
 			SetValue(Context, MoveTemp(InCollection), &Collection);
 			return;
 		}
 
-		const FDataflowTransformSelection& InSelection = GetValue(Context, &TransformSelection);
+		GeometryCollection::Facades::FCollectionTransformSelectionFacade SelectionFacade(InCollection);
 
-		const FDataflowTransformSelection& InTransformSelection = GetValue<FDataflowTransformSelection>(Context, &TransformSelection);
-		TArray<int32> Selection = InTransformSelection.AsArray();
+		TArray<int32> ToClear;
+		if (IsConnected(&TransformSelection))
+		{ 
+			const FDataflowTransformSelection& InTransformSelection = GetValue<FDataflowTransformSelection>(Context, &TransformSelection);
+			ToClear = InTransformSelection.AsArray();
 
-		TArray<int32> ToClear = InTransformSelection.AsArray();
-		GeometryCollection::Facades::FCollectionTransformSelectionFacade SelectionFacade(InCollection);
-		SelectionFacade.Sanitize(ToClear);
+			SelectionFacade.Sanitize(ToClear, /* bFavorParent */false);
+		}
+		else
+		{
+			ToClear = SelectionFacade.SelectAll();
+		}
 
 		FGeometryCollectionConvexUtility::RemoveConvexHulls(&InCollection, ToClear);
 		SetValue(Context, MoveTemp(InCollection), &Collection);
diff --git a/Engine/Plugins/Experimental/GeometryCollectionPlugin/Source/GeometryCollectionNodes/Public/Dataflow/GeometryCollectionUtilityNodes.h b/Engine/Plugins/Experimental/GeometryCollectionPlugin/Source/GeometryCollectionNodes/Public/Dataflow/GeometryCollectionUtilityNodes.h
--- a/Engine/Plugins/Experimental/GeometryCollectionPlugin/Source/GeometryCollectionNodes/Public/Dataflow/GeometryCollectionUtilityNodes.h
+++ b/Engine/Plugins/Experimental/GeometryCollectionPlugin/Source/GeometryCollectionNodes/Public/Dataflow/GeometryCollectionUtilityNodes.h
@@ -406,7 +406,7 @@ public:
 };
 
 
-/** Clear convex hulls from the selected transforms */
+/** Clear convex hulls from a collection */
 USTRUCT(meta = (DataflowGeometryCollection))
 struct FClearConvexHullsDataflowNode : public FDataflowNode
 {
@@ -417,8 +417,8 @@ public:
 	UPROPERTY(meta = (DataflowInput, DataflowOutput, DataflowPassthrough = "Collection", DataflowIntrinsic))
 	FManagedArrayCollection Collection;
 
-	/** Convex hulls will be cleared from these transforms */
-	UPROPERTY(meta = (DataflowInput, DataflowIntrinsic))
+	/** [Optional] selection of transforms to clear convex on, if not set all the transform will be used */
+	UPROPERTY(meta = (DataflowInput))
 	FDataflowTransformSelection TransformSelection;
 
 	FClearConvexHullsDataflowNode(const Dataflow::FNodeParameters& InParam, FGuid InGuid = FGuid::NewGuid())
