additional validation for selections on geometry collection convex processing

diff --git a/Engine/Plugins/Experimental/GeometryCollectionPlugin/Source/GeometryCollectionNodes/Private/Dataflow/GeometryCollectionUtilityNodes.cpp b/Engine/Plugins/Experimental/GeometryCollectionPlugin/Source/GeometryCollectionNodes/Private/Dataflow/GeometryCollectionUtilityNodes.cpp
--- a/Engine/Plugins/Experimental/GeometryCollectionPlugin/Source/GeometryCollectionNodes/Private/Dataflow/GeometryCollectionUtilityNodes.cpp
+++ b/Engine/Plugins/Experimental/GeometryCollectionPlugin/Source/GeometryCollectionNodes/Private/Dataflow/GeometryCollectionUtilityNodes.cpp
@@ -94,6 +94,8 @@ void FCreateLeafConvexHullsDataflowNode::Evaluate(Dataflow::FContext& Context, c
 				const FDataflowTransformSelection& InOptionalSelectionFilter = GetValue<FDataflowTransformSelection>(Context, &OptionalSelectionFilter);
 				bRestrictToSelection = true;
 				SelectedBones = InOptionalSelectionFilter.AsArray();
+				GeometryCollection::Facades::FCollectionTransformSelectionFacade SelectionFacade(InCollection);
+				SelectionFacade.Sanitize(SelectedBones, /* bFavorParent */false);
 			}
 
 			float InSimplificationDistanceThreshold = GetValue(Context, &SimplificationDistanceThreshold);
@@ -240,6 +242,8 @@ void FGenerateClusterConvexHullsFromLeafHullsDataflowNode::Evaluate(Dataflow::FC
 			{
 				const FDataflowTransformSelection& InOptionalSelectionFilter = GetValue<FDataflowTransformSelection>(Context, &OptionalSelectionFilter);
 				SelectionArray = InOptionalSelectionFilter.AsArray();
+				GeometryCollection::Facades::FCollectionTransformSelectionFacade SelectionFacade(InCollection);
+				SelectionFacade.Sanitize(SelectionArray, /* bFavorParent */false);
 			}
 
 			bool bHasNegativeSpace = false;
@@ -326,6 +330,8 @@ void FGenerateClusterConvexHullsFromChildrenHullsDataflowNode::Evaluate(Dataflow
 			{
 				const FDataflowTransformSelection& InOptionalSelectionFilter = GetValue<FDataflowTransformSelection>(Context, &OptionalSelectionFilter);
 				SelectionArray = InOptionalSelectionFilter.AsArray();
+				GeometryCollection::Facades::FCollectionTransformSelectionFacade SelectionFacade(InCollection);
+				SelectionFacade.Sanitize(SelectionArray, /* bFavorParent */false);
 			}
 
 			bool bHasNegativeSpace = false;
@@ -442,6 +448,8 @@ void FMergeConvexHullsDataflowNode::Evaluate(Dataflow::FContext& Context, const
 		{
 			const FDataflowTransformSelection& InOptionalSelectionFilter = GetValue<FDataflowTransformSelection>(Context, &OptionalSelectionFilter);
 			SelectionArray = InOptionalSelectionFilter.AsArray();
+			GeometryCollection::Facades::FCollectionTransformSelectionFacade SelectionFacade(InCollection);
+			SelectionFacade.Sanitize(SelectionArray, /* bFavorParent */false);
 		}
 
 		bool bHasPrecomputedNegativeSpace = false;
diff --git a/Engine/Source/Runtime/Experimental/Chaos/Private/GeometryCollection/GeometryCollectionConvexUtility.cpp b/Engine/Source/Runtime/Experimental/Chaos/Private/GeometryCollection/GeometryCollectionConvexUtility.cpp
--- a/Engine/Source/Runtime/Experimental/Chaos/Private/GeometryCollection/GeometryCollectionConvexUtility.cpp
+++ b/Engine/Source/Runtime/Experimental/Chaos/Private/GeometryCollection/GeometryCollectionConvexUtility.cpp
@@ -2027,6 +2027,11 @@ void FGeometryCollectionConvexUtility::GenerateClusterConvexHullsFromLeafOrChild
 
 	for (int32 TransformIndex : TransformsToProcess)
 	{
+		if (TransformIndex < 0 || TransformIndex >= SimulationTypeAttribute.Num())
+		{
+			UE_LOG(LogChaos, Warning, TEXT("Transform Index %d out of bounds [0,%d)"), TransformIndex, SimulationTypeAttribute.Num());
+			continue;
+		}
 		// only do this for clusters
 		const bool bIsCluster = (SimulationTypeAttribute.Get()[TransformIndex] == FGeometryCollection::ESimulationTypes::FST_Clustered);
 		if (bIsCluster)
