Dataflow : GeometryCollectionNodes
- Moved StaticMeshToCollection into its own file.
- Added StaticMesh Input to the node.

diff --git a/Engine/Plugins/Experimental/GeometryCollectionPlugin/Source/GeometryCollectionNodes/Private/Dataflow/GeometryCollectionAssetNodes.cpp b/Engine/Plugins/Experimental/GeometryCollectionPlugin/Source/GeometryCollectionNodes/Private/Dataflow/GeometryCollectionAssetNodes.cpp
new file mode 100644
--- /dev/null
+++ b/Engine/Plugins/Experimental/GeometryCollectionPlugin/Source/GeometryCollectionNodes/Private/Dataflow/GeometryCollectionAssetNodes.cpp
@@ -0,0 +1,302 @@
+// Copyright Epic Games, Inc. All Rights Reserved.
+
+#include "Dataflow/GeometryCollectionAssetNodes.h"
+#include "Dataflow/DataflowCore.h"
+
+#include "Engine/StaticMesh.h"
+#include "GeometryCollection/GeometryCollection.h"
+#include "GeometryCollection/GeometryCollectionEngineConversion.h"
+#include "GeometryCollection/GeometryCollectionObject.h"
+#include "GeometryCollection/GeometryCollectionClusteringUtility.h"
+#include "GeometryCollection/Facades/CollectionHierarchyFacade.h"
+#include "GeometryCollection/Facades/CollectionInstancedMeshFacade.h"
+#include "PreviewScene.h"
+
+namespace Dataflow
+{
+	void GeometryCollectionEngineAssetNodes()
+	{
+		static const FLinearColor CDefaultNodeBodyTintColor = FLinearColor(0.f, 0.f, 0.f, 0.5f);
+
+		DATAFLOW_NODE_REGISTER_CREATION_FACTORY(FGeometryCollectionTerminalDataflowNode);
+		DATAFLOW_NODE_REGISTER_CREATION_FACTORY(FGetGeometryCollectionAssetDataflowNode);
+		DATAFLOW_NODE_REGISTER_CREATION_FACTORY(FGetGeometryCollectionSourcesDataflowNode);
+		DATAFLOW_NODE_REGISTER_CREATION_FACTORY(FCreateGeometryCollectionFromSourcesDataflowNode);
+		DATAFLOW_NODE_REGISTER_CREATION_FACTORY(FGeometryCollectionToCollectionDataflowNode);
+		DATAFLOW_NODE_REGISTER_CREATION_FACTORY(FBlueprintToCollectionDataflowNode);
+
+		// Terminal
+		DATAFLOW_NODE_REGISTER_CREATION_FACTORY_NODE_COLORS_BY_CATEGORY("Terminal", FLinearColor(0.f, 0.f, 0.f), CDefaultNodeBodyTintColor);
+	}
+}
+
+
+// ===========================================================================================================================
+
+FGeometryCollectionTerminalDataflowNode::FGeometryCollectionTerminalDataflowNode(const Dataflow::FNodeParameters& InParam, FGuid InGuid)
+	: FDataflowTerminalNode(InParam, InGuid)
+{
+	RegisterInputConnection(&Collection);
+	RegisterOutputConnection(&Collection, &Collection);
+	RegisterInputConnection(&Materials);
+	RegisterOutputConnection(&Materials, &Materials);
+	RegisterInputConnection(&InstancedMeshes);
+	RegisterOutputConnection(&InstancedMeshes, &InstancedMeshes);
+}
+
+
+
+void FGeometryCollectionTerminalDataflowNode::SetAssetValue(TObjectPtr<UObject> Asset, Dataflow::FContext& Context) const
+{
+	using FGeometryCollectionPtr = TSharedPtr<FGeometryCollection, ESPMode::ThreadSafe>;
+	using FMaterialArray = TArray<TObjectPtr<UMaterial>>;
+	using FInstancedMeshesArray = TArray<FGeometryCollectionAutoInstanceMesh>;
+
+	if (UGeometryCollection* CollectionAsset = Cast<UGeometryCollection>(Asset.Get()))
+	{
+		if (FGeometryCollectionPtr GeometryCollection = CollectionAsset->GetGeometryCollection())
+		{
+			const FManagedArrayCollection& InCollection = GetValue(Context, &Collection);
+			const FMaterialArray& InMaterials = GetValue(Context, &Materials);
+			const FInstancedMeshesArray& InInstancedMeshes = GetValue(Context, &InstancedMeshes);
+
+			if (InCollection.NumElements(FGeometryCollection::TransformGroup) == 0)
+			{
+				UE_LOG(LogChaos, Warning, TEXT("Cannot set Geometry Collection asset: Geometry Collection must have at least one transform."));
+				return;
+			}
+
+			const bool bHasInternalMaterial = false; // with data flow there's no assumption of internal materials
+			CollectionAsset->ResetFrom(InCollection, InMaterials, false);
+			CollectionAsset->SetAutoInstanceMeshes(InInstancedMeshes);
+
+#if WITH_EDITOR
+			// make sure we rebuild the render data when we are done setting everything 
+			CollectionAsset->RebuildRenderData();
+			// also make sure all components using it are getting a notification about it
+			CollectionAsset->PropagateTransformUpdateToComponents();
+#endif
+		}
+	}
+}
+
+void FGeometryCollectionTerminalDataflowNode::Evaluate(Dataflow::FContext& Context) const
+{
+	using FMaterialArray = TArray<TObjectPtr<UMaterial>>;
+	using FInstancedMeshesArray = TArray<FGeometryCollectionAutoInstanceMesh>;
+
+	const FManagedArrayCollection& InCollection = GetValue(Context, &Collection);
+	const FMaterialArray& InMaterials = GetValue(Context, &Materials);
+	const FInstancedMeshesArray& InInstancedMeshes = GetValue(Context, &InstancedMeshes);
+
+	SetValue(Context, InCollection, &Collection);
+	SetValue(Context, InMaterials, &Materials);
+	SetValue(Context, InInstancedMeshes, &InstancedMeshes);
+}
+
+// ===========================================================================================================================
+
+FGetGeometryCollectionAssetDataflowNode::FGetGeometryCollectionAssetDataflowNode(const Dataflow::FNodeParameters& InParam, FGuid InGuid)
+	: FDataflowNode(InParam, InGuid)
+{
+	RegisterOutputConnection(&Asset);
+}
+
+void FGetGeometryCollectionAssetDataflowNode::Evaluate(Dataflow::FContext& Context, const FDataflowOutput* Out) const
+{
+	ensure(Out->IsA(&Asset));
+
+	TObjectPtr<UGeometryCollection> CollectionAsset(nullptr);
+	if (const Dataflow::FEngineContext* EngineContext = Context.AsType<Dataflow::FEngineContext>())
+	{
+		CollectionAsset = Cast<UGeometryCollection>(EngineContext->Owner);
+	}
+	SetValue(Context, CollectionAsset, &Asset);
+}
+
+// ===========================================================================================================================
+
+FGetGeometryCollectionSourcesDataflowNode::FGetGeometryCollectionSourcesDataflowNode(const Dataflow::FNodeParameters& InParam, FGuid InGuid)
+		: FDataflowNode(InParam, InGuid)
+{
+	RegisterInputConnection(&Asset);
+	RegisterOutputConnection(&Sources);
+}
+
+void FGetGeometryCollectionSourcesDataflowNode::Evaluate(Dataflow::FContext& Context, const FDataflowOutput* Out) const
+{
+	ensure(Out->IsA(&Sources));
+
+	TArray<FGeometryCollectionSource> OutSources;
+	
+	if (const TObjectPtr<const UGeometryCollection> InAsset = GetValue(Context, &Asset))
+	{
+#if WITH_EDITORONLY_DATA
+		OutSources = InAsset->GeometrySource; 
+#else
+		ensureMsgf(false, TEXT("FGetGeometryCollectionSourcesDataflowNode - GeometrySource is only available in editor, returning an empty array"));
+#endif
+
+	}
+
+	SetValue(Context, OutSources, &Sources);
+}
+
+// ===========================================================================================================================
+
+FCreateGeometryCollectionFromSourcesDataflowNode::FCreateGeometryCollectionFromSourcesDataflowNode(const Dataflow::FNodeParameters& InParam, FGuid InGuid)
+		: FDataflowNode(InParam, InGuid)
+{
+	RegisterInputConnection(&Sources);
+	RegisterOutputConnection(&Collection);
+	RegisterOutputConnection(&Materials);
+	RegisterOutputConnection(&InstancedMeshes);
+}
+
+void FCreateGeometryCollectionFromSourcesDataflowNode::Evaluate(Dataflow::FContext& Context, const FDataflowOutput* Out) const
+{
+	ensure(Out->IsA(&Collection) || Out->IsA(&Materials) || Out->IsA(&InstancedMeshes));
+	
+	const TArray<FGeometryCollectionSource>& InSources = GetValue(Context, &Sources);
+
+	FGeometryCollection OutCollection;
+	TArray<TObjectPtr<UMaterial>> OutMaterials;
+	TArray<FGeometryCollectionAutoInstanceMesh> OutInstancedMeshes;
+
+	// make sure we have an attribute for instanced meshes
+	GeometryCollection::Facades::FCollectionInstancedMeshFacade InstancedMeshFacade(OutCollection);
+	InstancedMeshFacade.DefineSchema();
+
+	constexpr bool bReindexMaterialsInLoop = false;
+	for (int32 SourceIndex = 0; SourceIndex < InSources.Num(); SourceIndex++)
+	{
+		const FGeometryCollectionSource& Source = InSources[SourceIndex];
+		const int32 NumTransformsBeforeAppending = OutCollection.NumElements(FGeometryCollection::TransformGroup);
+
+		// todo: change AppendGeometryCollectionSource to take a FManagedArrayCollection so we could move the collection when assigning it to the output
+		FGeometryCollectionEngineConversion::AppendGeometryCollectionSource(Source, OutCollection, MutableView(OutMaterials), bReindexMaterialsInLoop);
+
+		// todo(chaos) if the source is a geometry collection this will not work properly 
+		FGeometryCollectionAutoInstanceMesh InstancedMesh;
+		InstancedMesh.Mesh = Cast<UStaticMesh>(Source.SourceGeometryObject.TryLoad());
+		InstancedMesh.Materials = Source.SourceMaterial;
+		InstancedMesh.NumInstances = 0;
+		InstancedMesh.CustomData.Reset();
+
+		int32 InstancedMeshIndex = OutInstancedMeshes.Find(InstancedMesh);
+		if (InstancedMeshIndex == INDEX_NONE)
+		{
+			InstancedMeshIndex = OutInstancedMeshes.Add(InstancedMesh);
+		}
+		OutInstancedMeshes[InstancedMeshIndex].NumInstances++;
+		OutInstancedMeshes[InstancedMeshIndex].CustomData.Append(Source.InstanceCustomData);
+
+		// add the instanced mesh  for all the newly added transforms 
+		const int32 NumTransformsAfterAppending = OutCollection.NumElements(FGeometryCollection::TransformGroup);
+		//ensure((NumTransformsAfterAppending - NumTransformsBeforeAppending) == 1);
+		for (int32 TransformIndex = NumTransformsBeforeAppending; TransformIndex < NumTransformsAfterAppending; TransformIndex++)
+		{
+			InstancedMeshFacade.SetIndex(TransformIndex, InstancedMeshIndex);
+		}
+	}
+	if (bReindexMaterialsInLoop == false)
+	{
+		OutCollection.ReindexMaterials();
+	}
+
+	// add the instanced mesh indices
+
+	const int32 NumTransforms = InstancedMeshFacade.GetNumIndices();
+	for (int32 TransformIndex = 0; TransformIndex < NumTransforms; TransformIndex++)
+	{
+	}
+
+	// make sure we have only one root
+	if (FGeometryCollectionClusteringUtility::ContainsMultipleRootBones(&OutCollection))
+	{
+		FGeometryCollectionClusteringUtility::ClusterAllBonesUnderNewRoot(&OutCollection);
+	}
+
+	// make sure we have a level attribute
+	Chaos::Facades::FCollectionHierarchyFacade HierarchyFacade(OutCollection);
+	HierarchyFacade.GenerateLevelAttribute();
+
+	// we have to make a copy since we have generated a FGeometryCollection which is inherited from FManagedArrayCollection
+	SetValue(Context, static_cast<const FManagedArrayCollection&>(OutCollection), &Collection);
+	SetValue(Context, MoveTemp(OutMaterials), &Materials);
+	SetValue(Context, MoveTemp(OutInstancedMeshes), &InstancedMeshes);
+}
+
+// ===========================================================================================================================
+
+FGeometryCollectionToCollectionDataflowNode::FGeometryCollectionToCollectionDataflowNode(const Dataflow::FNodeParameters& InParam, FGuid InGuid)
+	: FDataflowNode(InParam, InGuid)
+{
+	RegisterOutputConnection(&Collection);
+	RegisterOutputConnection(&Materials);
+	RegisterOutputConnection(&InstancedMeshes);
+}
+
+void FGeometryCollectionToCollectionDataflowNode::Evaluate(Dataflow::FContext& Context, const FDataflowOutput* Out) const
+{
+	ensure(Out->IsA(&Collection) || Out->IsA(&Materials) || Out->IsA(&InstancedMeshes));
+
+	FManagedArrayCollection OutCollection;
+	TArray<TObjectPtr<UMaterial>> OutMaterials;
+	TArray<FGeometryCollectionAutoInstanceMesh> OutInstancedMeshes;
+
+	if (GeometryCollection)
+	{
+		FGeometryCollectionEngineConversion::ConvertGeometryCollectionToGeometryCollection(GeometryCollection, OutCollection, OutMaterials, OutInstancedMeshes);
+	}
+
+	// Set Outputs
+	SetValue(Context, MoveTemp(OutCollection), &Collection);
+	SetValue(Context, MoveTemp(OutMaterials), &Materials);
+	SetValue(Context, MoveTemp(OutInstancedMeshes), &InstancedMeshes);
+}
+
+// ===========================================================================================================================
+
+FBlueprintToCollectionDataflowNode::FBlueprintToCollectionDataflowNode(const Dataflow::FNodeParameters& InParam, FGuid InGuid)
+	: FDataflowNode(InParam, InGuid)
+{
+	RegisterOutputConnection(&Collection);
+	RegisterOutputConnection(&Materials);
+	RegisterOutputConnection(&InstancedMeshes);
+}
+
+void FBlueprintToCollectionDataflowNode::Evaluate(Dataflow::FContext& Context, const FDataflowOutput* Out) const
+{
+	ensure(Out->IsA(&Collection) || Out->IsA(&Materials) || Out->IsA(&InstancedMeshes));
+
+	FManagedArrayCollection OutCollection;
+	TArray<TObjectPtr<UMaterial>> OutMaterials;
+	TArray<FGeometryCollectionAutoInstanceMesh> OutInstancedMeshes;
+
+	if (Blueprint)
+	{
+		if (TUniquePtr<FPreviewScene> PreviewScene = MakeUnique<FPreviewScene>(FPreviewScene::ConstructionValues()))
+		{
+			if (UWorld* PreviewWorld = PreviewScene->GetWorld())
+			{
+				FActorSpawnParameters SpawnInfo;
+				SpawnInfo.SpawnCollisionHandlingOverride = ESpawnActorCollisionHandlingMethod::AlwaysSpawn;
+				SpawnInfo.bNoFail = true;
+				SpawnInfo.ObjectFlags = RF_Transient;
+
+				if (AActor* PreviewActor = PreviewWorld->SpawnActor(Blueprint->GeneratedClass, nullptr, SpawnInfo))
+				{
+					FGeometryCollectionEngineConversion::FSkeletalMeshToCollectionConversionParameters ConversionParameters;
+					FGeometryCollectionEngineConversion::ConvertActorToGeometryCollection(PreviewActor, OutCollection, OutMaterials, OutInstancedMeshes, ConversionParameters, bSplitComponents);
+				}
+			}
+		}
+	}
+
+	// Set Outputs
+	SetValue(Context, MoveTemp(OutCollection), &Collection);
+	SetValue(Context, MoveTemp(OutMaterials), &Materials);
+	SetValue(Context, MoveTemp(OutInstancedMeshes), &InstancedMeshes);
+}
+
diff --git a/Engine/Plugins/Experimental/GeometryCollectionPlugin/Source/GeometryCollectionNodes/Private/Dataflow/GeometryCollectionNodesPlugin.cpp b/Engine/Plugins/Experimental/GeometryCollectionPlugin/Source/GeometryCollectionNodes/Private/Dataflow/GeometryCollectionNodesPlugin.cpp
--- a/Engine/Plugins/Experimental/GeometryCollectionPlugin/Source/GeometryCollectionNodes/Private/Dataflow/GeometryCollectionNodesPlugin.cpp
+++ b/Engine/Plugins/Experimental/GeometryCollectionPlugin/Source/GeometryCollectionNodes/Private/Dataflow/GeometryCollectionNodesPlugin.cpp
@@ -5,29 +5,30 @@
 #include "CoreMinimal.h"
 
 #include "Dataflow/DataflowNodeFactory.h"
-#include "Dataflow/GeometryCollectionNodes.h"
-#include "Dataflow/Nodes/GeometryCollectionAssetNodes.h"
-#include "Dataflow/GeometryCollectionProcessingNodes.h"
-#include "Dataflow/GeometryCollectionSkeletalMeshToCollectionNode.h"
-#include "Dataflow/GeometryCollectionSkeletonToCollectionNode.h"
-#include "Dataflow/GeometryCollectionSelectionNodes.h"
-#include "Dataflow/GeometryCollectionMeshNodes.h"
+#include "Dataflow/GeometryCollectionArrayNodes.h"
+#include "Dataflow/GeometryCollectionAssetNodes.h"
 #include "Dataflow/GeometryCollectionClusteringNodes.h"
-#include "Dataflow/GeometryCollectionFracturingNodes.h"
+#include "Dataflow/GeometryCollectionConversionNodes.h"
+#include "Dataflow/GeometryCollectionDebugNodes.h"
 #include "Dataflow/GeometryCollectionEditNodes.h"
-#include "Dataflow/GeometryCollectionUtilityNodes.h"
-#include "Dataflow/GeometryCollectionMaterialNodes.h"
 #include "Dataflow/GeometryCollectionFieldNodes.h"
-#include "Dataflow/GeometryCollectionOverrideNodes.h"
+#include "Dataflow/GeometryCollectionFracturingNodes.h"
 #include "Dataflow/GeometryCollectionMakeNodes.h"
+#include "Dataflow/GeometryCollectionMaterialNodes.h"
 #include "Dataflow/GeometryCollectionMathNodes.h"
-#include "Dataflow/GeometryCollectionConversionNodes.h"
-#include "Dataflow/GeometryCollectionVerticesNodes.h"
-#include "Dataflow/GeometryCollectionArrayNodes.h"
-#include "Dataflow/GeometryCollectionDebugNodes.h"
-#include "Dataflow/GeometryCollectionVertexScalarToVertexIndicesNode.h"
+#include "Dataflow/GeometryCollectionMeshNodes.h"
+#include "Dataflow/GeometryCollectionNodes.h"
+#include "Dataflow/GeometryCollectionOverrideNodes.h"
+#include "Dataflow/GeometryCollectionProcessingNodes.h"
+#include "Dataflow/GeometryCollectionSelectionNodes.h"
+#include "Dataflow/GeometryCollectionSkeletalMeshToCollectionNode.h"
+#include "Dataflow/GeometryCollectionSkeletonToCollectionNode.h"
+#include "Dataflow/GeometryCollectionStaticMeshToCollectionNode.h"
 #include "Dataflow/GeometryCollectionTransferVertexScalarAttributeNode.h"
 #include "Dataflow/GeometryCollectionTriangleBoundaryIndicesNode.h"
+#include "Dataflow/GeometryCollectionUtilityNodes.h"
+#include "Dataflow/GeometryCollectionVerticesNodes.h"
+#include "Dataflow/GeometryCollectionVertexScalarToVertexIndicesNode.h"
 
 #define LOCTEXT_NAMESPACE "DataflowNodes"
 
@@ -56,6 +57,7 @@ void IGeometryCollectionNodesPlugin::StartupModule()
 	DATAFLOW_NODE_REGISTER_CREATION_FACTORY(FSkeletonToCollectionDataflowNode);
 	DATAFLOW_NODE_REGISTER_CREATION_FACTORY(FSkeletalMeshToCollectionDataflowNode);
 	DATAFLOW_NODE_REGISTER_CREATION_FACTORY(FGeometryCollectionTransferVertexScalarAttributeNode);
+	DATAFLOW_NODE_REGISTER_CREATION_FACTORY(FStaticMeshToCollectionDataflowNode);
 	DATAFLOW_NODE_REGISTER_CREATION_FACTORY(FTriangleBoundaryIndicesNode);
 }
 
diff --git a/Engine/Plugins/Experimental/GeometryCollectionPlugin/Source/GeometryCollectionNodes/Private/Dataflow/GeometryCollectionStaticMeshToCollectionNode.cpp b/Engine/Plugins/Experimental/GeometryCollectionPlugin/Source/GeometryCollectionNodes/Private/Dataflow/GeometryCollectionStaticMeshToCollectionNode.cpp
new file mode 100644
--- /dev/null
+++ b/Engine/Plugins/Experimental/GeometryCollectionPlugin/Source/GeometryCollectionNodes/Private/Dataflow/GeometryCollectionStaticMeshToCollectionNode.cpp
@@ -0,0 +1,44 @@
+// Copyright Epic Games, Inc. All Rights Reserved.
+
+#include "Dataflow/GeometryCollectionStaticMeshToCollectionNode.h"
+#include "Dataflow/DataflowCore.h"
+
+#include "Engine/StaticMesh.h"
+#include "GeometryCollection/GeometryCollection.h"
+#include "GeometryCollection/GeometryCollectionEngineConversion.h"
+#include "GeometryCollection/GeometryCollectionObject.h"
+#include "GeometryCollection/ManagedArrayCollection.h"
+
+
+// ===========================================================================================================================
+
+
+FStaticMeshToCollectionDataflowNode::FStaticMeshToCollectionDataflowNode(const Dataflow::FNodeParameters& InParam, FGuid InGuid)
+	: FDataflowNode(InParam, InGuid)
+{
+	RegisterInputConnection(&StaticMesh);
+	RegisterOutputConnection(&Collection);
+	RegisterOutputConnection(&Materials);
+	RegisterOutputConnection(&InstancedMeshes);
+}
+
+void FStaticMeshToCollectionDataflowNode::Evaluate(Dataflow::FContext& Context, const FDataflowOutput* Out) const
+{
+	ensure(Out->IsA(&Collection) || Out->IsA(&Materials) || Out->IsA(&InstancedMeshes));
+
+	FManagedArrayCollection OutCollection;
+	TArray<TObjectPtr<UMaterial>> OutMaterials;
+	TArray<FGeometryCollectionAutoInstanceMesh> OutInstancedMeshes;
+
+	TObjectPtr<UStaticMesh> StaticMeshVal = GetValue(Context, &StaticMesh, StaticMesh);
+	if (StaticMeshVal)
+	{
+		FGeometryCollectionEngineConversion::ConvertStaticMeshToGeometryCollection(StaticMeshVal, OutCollection, OutMaterials, OutInstancedMeshes, bSetInternalFromMaterialIndex, bSplitComponents);
+	}
+
+	// Set Outputs
+	SetValue(Context, MoveTemp(OutCollection), &Collection);
+	SetValue(Context, MoveTemp(OutMaterials), &Materials);
+	SetValue(Context, MoveTemp(OutInstancedMeshes), &InstancedMeshes);
+}
+
diff --git a/Engine/Plugins/Experimental/GeometryCollectionPlugin/Source/GeometryCollectionNodes/Private/Dataflow/Nodes/GeometryCollectionAssetNodes.cpp b/Engine/Plugins/Experimental/GeometryCollectionPlugin/Source/GeometryCollectionNodes/Private/Dataflow/Nodes/GeometryCollectionAssetNodes.cpp
deleted file mode 100644
diff --git a/Engine/Plugins/Experimental/GeometryCollectionPlugin/Source/GeometryCollectionNodes/Public/Dataflow/GeometryCollectionAssetNodes.h b/Engine/Plugins/Experimental/GeometryCollectionPlugin/Source/GeometryCollectionNodes/Public/Dataflow/GeometryCollectionAssetNodes.h
new file mode 100644
--- /dev/null
+++ b/Engine/Plugins/Experimental/GeometryCollectionPlugin/Source/GeometryCollectionNodes/Public/Dataflow/GeometryCollectionAssetNodes.h
@@ -0,0 +1,189 @@
+// Copyright Epic Games, Inc. All Rights Reserved.
+
+#pragma once 
+
+#include "CoreMinimal.h"
+#include "Dataflow/DataflowEngine.h"
+#include "GeometryCollection/ManagedArrayCollection.h"
+#include "GeometryCollection/GeometryCollectionObject.h"
+#include "Engine/Blueprint.h"
+
+#include "GeometryCollectionAssetNodes.generated.h"
+
+class UGeometryCollection;
+class UBlueprint;
+
+USTRUCT(meta = (DataflowGeometryCollection, DataflowTerminal))
+struct FGeometryCollectionTerminalDataflowNode : public FDataflowTerminalNode
+{
+	GENERATED_USTRUCT_BODY()
+	DATAFLOW_NODE_DEFINE_INTERNAL(FGeometryCollectionTerminalDataflowNode, "GeometryCollectionTerminal", "Terminal", "")
+	DATAFLOW_NODE_RENDER_TYPE(FGeometryCollection::StaticType(), "Collection")
+
+public:
+
+	UPROPERTY(meta = (DataflowInput, DataflowOutput, DataflowPassthrough = "Collection", DisplayName = "Collection"))
+	FManagedArrayCollection Collection;
+
+	/** Materials array to use for this asset */
+	UPROPERTY(meta = (DataflowInput, DataflowOutput, DataflowPassthrough = "Materials", DisplayName = "Materials"))
+	TArray<TObjectPtr<UMaterial>> Materials;
+
+	/** array of instanced meshes*/
+	UPROPERTY(meta = (DataflowInput, DataflowOutput, DataflowPassthrough = "InstancedMeshes", DisplayName = "InstancedMeshes"))
+	TArray<FGeometryCollectionAutoInstanceMesh> InstancedMeshes;
+
+	FGeometryCollectionTerminalDataflowNode(const Dataflow::FNodeParameters& InParam, FGuid InGuid = FGuid::NewGuid());
+
+	virtual void Evaluate(Dataflow::FContext& Context) const override;
+	virtual void SetAssetValue(TObjectPtr<UObject> Asset, Dataflow::FContext& Context) const override;
+};
+
+
+/**
+ * Get Current geometry collection asset 
+ * Note : Use with caution as this may get replaced in a near future for a more generic getAsset node
+ */
+USTRUCT(meta = (DataflowGeometryCollection))
+struct FGetGeometryCollectionAssetDataflowNode : public FDataflowNode
+{
+	GENERATED_USTRUCT_BODY()
+	DATAFLOW_NODE_DEFINE_INTERNAL(FGetGeometryCollectionAssetDataflowNode, "GetGeometryCollectionAsset", "GeometryCollection|Asset", "")
+
+public:
+	FGetGeometryCollectionAssetDataflowNode(const Dataflow::FNodeParameters& InParam, FGuid InGuid = FGuid::NewGuid());
+	virtual void Evaluate(Dataflow::FContext& Context, const FDataflowOutput* Out) const override;
+
+	/** Asset this data flow graph instance is assigned to */
+	UPROPERTY(meta = (DataflowOutput, DisplayName = "Asset"))
+	TObjectPtr<UGeometryCollection> Asset;
+};
+
+
+/**
+ * Get the list of the original mesh information used to create a specific geometryc collection asset
+ * each entry contains a mesh, a transform and a list of override materials
+ */
+USTRUCT(meta = (DataflowGeometryCollection))
+struct FGetGeometryCollectionSourcesDataflowNode : public FDataflowNode
+{
+	GENERATED_USTRUCT_BODY()
+	DATAFLOW_NODE_DEFINE_INTERNAL(FGetGeometryCollectionSourcesDataflowNode, "GetGeometryCollectionSources", "GeometryCollection|Asset", "")
+
+public:
+	FGetGeometryCollectionSourcesDataflowNode(const Dataflow::FNodeParameters& InParam, FGuid InGuid = FGuid::NewGuid());
+	virtual void Evaluate(Dataflow::FContext& Context, const FDataflowOutput* Out) const override;
+	
+	/** Asset to get geometry sources from */
+	UPROPERTY(meta = (DataflowInput, DisplayName = "Asset"))
+	TObjectPtr<UGeometryCollection> Asset;
+	
+	/** array of geometry sources */
+	UPROPERTY(meta = (DataflowOutput, DisplayName = "Sources"))
+	TArray<FGeometryCollectionSource> Sources;
+};
+
+
+/**
+ * create a geometry collection from a set of geometry sources    
+ */
+USTRUCT(meta = (DataflowGeometryCollection))
+struct FCreateGeometryCollectionFromSourcesDataflowNode : public FDataflowNode
+{
+	GENERATED_USTRUCT_BODY()
+	DATAFLOW_NODE_DEFINE_INTERNAL(FCreateGeometryCollectionFromSourcesDataflowNode, "CreateGeometryCollectionFromSources", "GeometryCollection|Asset", "")
+
+public:
+	FCreateGeometryCollectionFromSourcesDataflowNode(const Dataflow::FNodeParameters& InParam, FGuid InGuid = FGuid::NewGuid());
+	virtual void Evaluate(Dataflow::FContext& Context, const FDataflowOutput* Out) const override;
+	
+	/** array of geometry sources */
+	UPROPERTY(meta = (DataflowInput, DisplayName = "Sources"))
+	TArray<FGeometryCollectionSource> Sources;
+
+	/** Geometry collection newly created */
+	UPROPERTY(meta = (DataflowOutput, DisplayName = "Collection"))
+	FManagedArrayCollection Collection;
+
+	/** Materials array to use for this asset */
+	UPROPERTY(meta = (DataflowOutput, DisplayName = "Materials"))
+	TArray<TObjectPtr<UMaterial>> Materials;
+
+	/** array of instanced meshes*/
+	UPROPERTY(meta = (DataflowOutput, DisplayName = "InstancedMeshes"))
+	TArray<FGeometryCollectionAutoInstanceMesh> InstancedMeshes;
+};
+
+
+/**
+ * Converts a UGeometryCollection asset to an FManagedArrayCollection
+ */
+USTRUCT(meta = (DataflowContext = "GeometryCollection", DataflowGeometryCollection, DataflowTerminal))
+struct FGeometryCollectionToCollectionDataflowNode : public FDataflowNode
+{
+	GENERATED_USTRUCT_BODY()
+	DATAFLOW_NODE_DEFINE_INTERNAL(FGeometryCollectionToCollectionDataflowNode, "GeometryCollectionToCollection", "GeometryCollection|Asset", "")
+	DATAFLOW_NODE_RENDER_TYPE(FGeometryCollection::StaticType(), "Collection")
+
+public:
+	/** Asset input */
+	UPROPERTY(EditAnywhere, Category = "Asset");
+	TObjectPtr<UGeometryCollection> GeometryCollection;
+
+	/** Geometry collection newly created */
+	UPROPERTY(meta = (DataflowOutput, DisplayName = "Collection"))
+	FManagedArrayCollection Collection;
+
+	/** Materials array to use for this asset */
+	UPROPERTY(meta = (DataflowOutput, DisplayName = "Materials"))
+	TArray<TObjectPtr<UMaterial>> Materials;
+
+	/** Array of instanced meshes*/
+	UPROPERTY(meta = (DataflowOutput, DisplayName = "InstancedMeshes"))
+	TArray<FGeometryCollectionAutoInstanceMesh> InstancedMeshes;
+
+	FGeometryCollectionToCollectionDataflowNode(const Dataflow::FNodeParameters& InParam, FGuid InGuid = FGuid::NewGuid());
+	virtual void Evaluate(Dataflow::FContext& Context, const FDataflowOutput* Out) const override;
+};
+
+
+/**
+ * Create a geometry collection from an asset
+ */
+ USTRUCT(meta = (DataflowContext = "GeometryCollection", DataflowGeometryCollection, DataflowTerminal))
+ struct FBlueprintToCollectionDataflowNode : public FDataflowNode
+ {
+ 	GENERATED_USTRUCT_BODY()
+ 	DATAFLOW_NODE_DEFINE_INTERNAL(FBlueprintToCollectionDataflowNode, "BlueprintToCollection", "GeometryCollection|Asset", "")
+	DATAFLOW_NODE_RENDER_TYPE(FGeometryCollection::StaticType(), "Collection")
+
+ public:
+ 	/** Asset input */
+ 	UPROPERTY(EditAnywhere, Category = "Asset");
+	TObjectPtr<UBlueprint> Blueprint;
+ 
+ 	/** Split components */
+ 	UPROPERTY(EditAnywhere, Category = "Asset");
+ 	bool bSplitComponents = false;
+ 
+ 	/** Geometry collection newly created */
+ 	UPROPERTY(meta = (DataflowOutput, DisplayName = "Collection"))
+ 	FManagedArrayCollection Collection;
+ 
+ 	/** Materials array to use for this asset */
+ 	UPROPERTY(meta = (DataflowOutput, DisplayName = "Materials"))
+ 	TArray<TObjectPtr<UMaterial>> Materials;
+ 
+ 	/** Array of instanced meshes*/
+ 	UPROPERTY(meta = (DataflowOutput, DisplayName = "InstancedMeshes"))
+ 	TArray<FGeometryCollectionAutoInstanceMesh> InstancedMeshes;
+ 
+	FBlueprintToCollectionDataflowNode(const Dataflow::FNodeParameters& InParam, FGuid InGuid = FGuid::NewGuid());
+ 	virtual void Evaluate(Dataflow::FContext& Context, const FDataflowOutput* Out) const override;
+ };
+
+ 
+namespace Dataflow
+{
+	void GeometryCollectionEngineAssetNodes();
+}
\ No newline at end of file
diff --git a/Engine/Plugins/Experimental/GeometryCollectionPlugin/Source/GeometryCollectionNodes/Public/Dataflow/GeometryCollectionStaticMeshToCollectionNode.h b/Engine/Plugins/Experimental/GeometryCollectionPlugin/Source/GeometryCollectionNodes/Public/Dataflow/GeometryCollectionStaticMeshToCollectionNode.h
new file mode 100644
--- /dev/null
+++ b/Engine/Plugins/Experimental/GeometryCollectionPlugin/Source/GeometryCollectionNodes/Public/Dataflow/GeometryCollectionStaticMeshToCollectionNode.h
@@ -0,0 +1,52 @@
+// Copyright Epic Games, Inc. All Rights Reserved.
+
+#pragma once 
+
+#include "CoreMinimal.h"
+#include "Dataflow/DataflowEngine.h"
+#include "GeometryCollection/GeometryCollectionObject.h"
+
+#include "GeometryCollectionStaticMeshToCollectionNode.generated.h"
+
+class UGeometryCollection;
+class UStaticMesh;
+class UMaterial;
+
+/**
+ * Create a geometry collection from a UStaticMesh
+ */
+USTRUCT(meta = (DataflowContext = "GeometryCollection", DataflowGeometryCollection, DataflowTerminal))
+struct FStaticMeshToCollectionDataflowNode : public FDataflowNode
+{
+	GENERATED_USTRUCT_BODY()
+	DATAFLOW_NODE_DEFINE_INTERNAL(FStaticMeshToCollectionDataflowNode, "StaticMeshToCollection", "GeometryCollection|Asset", "")
+	DATAFLOW_NODE_RENDER_TYPE(FGeometryCollection::StaticType(), "Collection")
+
+public:
+	/** Asset input */
+	UPROPERTY(EditAnywhere, Category = "Asset", meta = (DataflowInput, DisplayName = "StaticMesh"));
+	TObjectPtr<UStaticMesh> StaticMesh;
+
+	/** Set the internal faces from material index */
+	UPROPERTY(EditAnywhere, Category = "Asset");
+	bool bSetInternalFromMaterialIndex = true;
+
+	/** Split components */
+	UPROPERTY(EditAnywhere, Category = "Asset");
+	bool bSplitComponents = false;
+
+	/** Geometry collection newly created */
+	UPROPERTY(meta = (DataflowOutput, DisplayName = "Collection"))
+	FManagedArrayCollection Collection;
+
+	/** Materials array to use for this asset */
+	UPROPERTY(meta = (DataflowOutput, DisplayName = "Materials"))
+	TArray<TObjectPtr<UMaterial>> Materials;
+
+	/** Array of instanced meshes*/
+	UPROPERTY(meta = (DataflowOutput, DisplayName = "InstancedMeshes"))
+	TArray<FGeometryCollectionAutoInstanceMesh> InstancedMeshes;
+
+	FStaticMeshToCollectionDataflowNode(const Dataflow::FNodeParameters& InParam, FGuid InGuid = FGuid::NewGuid());
+	virtual void Evaluate(Dataflow::FContext& Context, const FDataflowOutput* Out) const override;
+};
diff --git a/Engine/Plugins/Experimental/GeometryCollectionPlugin/Source/GeometryCollectionNodes/Public/Dataflow/Nodes/GeometryCollectionAssetNodes.h b/Engine/Plugins/Experimental/GeometryCollectionPlugin/Source/GeometryCollectionNodes/Public/Dataflow/Nodes/GeometryCollectionAssetNodes.h
deleted file mode 100644
