Dataflow : Add math float expression node

diff --git a/Engine/Plugins/Experimental/GeometryCollectionPlugin/Source/GeometryCollectionNodes/Private/Dataflow/GeometryCollectionMathNodes.cpp b/Engine/Plugins/Experimental/GeometryCollectionPlugin/Source/GeometryCollectionNodes/Private/Dataflow/GeometryCollectionMathNodes.cpp
--- a/Engine/Plugins/Experimental/GeometryCollectionPlugin/Source/GeometryCollectionNodes/Private/Dataflow/GeometryCollectionMathNodes.cpp
+++ b/Engine/Plugins/Experimental/GeometryCollectionPlugin/Source/GeometryCollectionNodes/Private/Dataflow/GeometryCollectionMathNodes.cpp
@@ -1,6 +1,7 @@
 // Copyright Epic Games, Inc. All Rights Reserved.
 
 #include "Dataflow/GeometryCollectionMathNodes.h"
+#include "Math/BasicMathExpressionEvaluator.h"
 #include "Dataflow/DataflowCore.h"
 
 
@@ -68,6 +69,7 @@ namespace Dataflow
 		DATAFLOW_NODE_REGISTER_CREATION_FACTORY(FDegreesToRadiansDataflowNode);
 		DATAFLOW_NODE_REGISTER_CREATION_FACTORY(FMathConstantsDataflowNode);
 		DATAFLOW_NODE_REGISTER_CREATION_FACTORY(FOneMinusDataflowNode);
+		DATAFLOW_NODE_REGISTER_CREATION_FACTORY(FFloatMathExpressionDataflowNode);
 
 		// Math
 		DATAFLOW_NODE_REGISTER_CREATION_FACTORY_NODE_COLORS_BY_CATEGORY("Math", FLinearColor(0.f, 0.4f, 0.8f), CDefaultNodeBodyTintColor);
@@ -858,5 +860,64 @@ void FOneMinusDataflowNode::Evaluate(Dataflow::FContext& Context, const FDataflo
 	}
 }
 
+FFloatMathExpressionDataflowNode::FFloatMathExpressionDataflowNode(const Dataflow::FNodeParameters& InParam, FGuid InGuid)
+	: FDataflowNode(InParam, InGuid)
+{
+	RegisterInputConnection(&A);
+	RegisterInputConnection(&B);
+	RegisterInputConnection(&C);
+	RegisterInputConnection(&D);
+	RegisterOutputConnection(&ReturnValue);
+}
+
+void FFloatMathExpressionDataflowNode::Evaluate(Dataflow::FContext& Context, const FDataflowOutput* Out) const
+{
+	check(Out->IsA(&ReturnValue));
+
+	float FloatResult = 0.0f;
+
+	FString ExpressionToEvaluate{ Expression };
+	ExpressionToEvaluate = ExpressionToEvaluate.TrimStartAndEnd();
+	if (!ExpressionToEvaluate.IsEmpty())
+	{
+		const FString VarA("{A}");
+		if (ExpressionToEvaluate.Contains(VarA))
+		{
+			const float InA = GetValue(Context, &A);
+			const FString StrA = FString::SanitizeFloat(InA);
+			ExpressionToEvaluate.ReplaceInline(*VarA, *StrA, ESearchCase::CaseSensitive);
+		}
+
+		const FString VarB("{B}");
+		if (ExpressionToEvaluate.Contains(VarB))
+		{
+			const float InB = GetValue(Context, &B);
+			const FString StrB = FString::SanitizeFloat(InB);
+			ExpressionToEvaluate.ReplaceInline(*VarB, *StrB, ESearchCase::CaseSensitive);
+		}
+
+		const FString VarC("{C}");
+		if (ExpressionToEvaluate.Contains(VarC))
+		{
+			const float InC = GetValue(Context, &C);
+			const FString StrC = FString::SanitizeFloat(InC);
+			ExpressionToEvaluate.ReplaceInline(*VarC, *StrC, ESearchCase::CaseSensitive);
+		}
 
+		const FString VarD("{D}");
+		if (ExpressionToEvaluate.Contains(VarD))
+		{
+			const float InD = GetValue(Context, &D);
+			const FString StrD = FString::SanitizeFloat(InD);
+			ExpressionToEvaluate.ReplaceInline(*VarD, *StrD, ESearchCase::CaseSensitive);
+		}
 
+		FBasicMathExpressionEvaluator Evaluator;
+		TValueOrError<double, FExpressionError> Result = Evaluator.Evaluate(*ExpressionToEvaluate);
+		if (Result.IsValid())
+		{
+			FloatResult = FMath::Clamp((float)Result.GetValue(), TNumericLimits<float>::Lowest(), TNumericLimits<float>::Max());
+		}
+	}
+	SetValue(Context, FloatResult, &ReturnValue);
+}
\ No newline at end of file
diff --git a/Engine/Plugins/Experimental/GeometryCollectionPlugin/Source/GeometryCollectionNodes/Public/Dataflow/GeometryCollectionMathNodes.h b/Engine/Plugins/Experimental/GeometryCollectionPlugin/Source/GeometryCollectionNodes/Public/Dataflow/GeometryCollectionMathNodes.h
--- a/Engine/Plugins/Experimental/GeometryCollectionPlugin/Source/GeometryCollectionNodes/Public/Dataflow/GeometryCollectionMathNodes.h
+++ b/Engine/Plugins/Experimental/GeometryCollectionPlugin/Source/GeometryCollectionNodes/Public/Dataflow/GeometryCollectionMathNodes.h
@@ -1745,6 +1745,40 @@ public:
 	virtual void Evaluate(Dataflow::FContext& Context, const FDataflowOutput* Out) const override;
 };
 
+/**
+*
+* Expression node for floats
+*
+*/
+USTRUCT()
+struct FFloatMathExpressionDataflowNode : public FDataflowNode
+{
+	GENERATED_USTRUCT_BODY()
+	DATAFLOW_NODE_DEFINE_INTERNAL(FFloatMathExpressionDataflowNode, "FloatMathExpression", "Math|Float", "")
+
+public:
+	UPROPERTY(EditAnywhere, Category = "Variables", meta = (DataflowInput));
+	float A = 0.f;
+
+	UPROPERTY(EditAnywhere, Category = "Variables", meta = (DataflowInput));
+	float B = 0.f;
+
+	UPROPERTY(EditAnywhere, Category = "Variables", meta = (DataflowInput));
+	float C = 0.f;
+
+	UPROPERTY(EditAnywhere, Category = "Variables", meta = (DataflowInput));
+	float D = 0.f;
+
+	UPROPERTY(EditAnywhere, Category = "Expression");
+	FString Expression;
+
+	UPROPERTY(meta = (DataflowOutput))
+	float ReturnValue = 0.f;
+
+	FFloatMathExpressionDataflowNode(const Dataflow::FNodeParameters& InParam, FGuid InGuid = FGuid::NewGuid());
+
+	virtual void Evaluate(Dataflow::FContext& Context, const FDataflowOutput* Out) const override;
+};
 
 namespace Dataflow
 {
