Dataflow : GeometryCollection
- Deprecated existing vertex color nodes.
- Renamed and moved Vertex Coloring nodes to own file
- Added a SetVertexColorFromVertexIndices node.

diff --git a/Engine/Plugins/Experimental/GeometryCollectionPlugin/GeometryCollectionPlugin.uplugin b/Engine/Plugins/Experimental/GeometryCollectionPlugin/GeometryCollectionPlugin.uplugin
--- a/Engine/Plugins/Experimental/GeometryCollectionPlugin/GeometryCollectionPlugin.uplugin
+++ b/Engine/Plugins/Experimental/GeometryCollectionPlugin/GeometryCollectionPlugin.uplugin
@@ -34,6 +34,11 @@
 			"Name": "GeometryCollectionNodes",
 			"Type": "Runtime",
 			"LoadingPhase": "Default"
+		},
+		{
+			"Name": "GeometryCollectionDeprecatedNodes",
+			"Type": "Runtime",
+			"LoadingPhase": "Default"
 		}
 	],
 	"Plugins": [
diff --git a/Engine/Plugins/Experimental/GeometryCollectionPlugin/Source/GeometryCollectionDeprecatedNodes/GeometryCollectionDeprecatedNodes.Build.cs b/Engine/Plugins/Experimental/GeometryCollectionPlugin/Source/GeometryCollectionDeprecatedNodes/GeometryCollectionDeprecatedNodes.Build.cs
new file mode 100644
--- /dev/null
+++ b/Engine/Plugins/Experimental/GeometryCollectionPlugin/Source/GeometryCollectionDeprecatedNodes/GeometryCollectionDeprecatedNodes.Build.cs
@@ -0,0 +1,30 @@
+// Copyright Epic Games, Inc. All Rights Reserved.
+
+namespace UnrealBuildTool.Rules
+{
+	public class GeometryCollectionDeprecatedNodes : ModuleRules
+	{
+        public GeometryCollectionDeprecatedNodes(ReadOnlyTargetRules Target) : base(Target)
+		{
+			PublicDependencyModuleNames.AddRange(
+				new string[]
+				{
+					"Core",
+					"CoreUObject",
+					"ChaosCore",
+					"Chaos",
+					"DataflowCore",
+					"DataflowEngine",
+					"GeometryCollectionEngine",
+					"Engine"
+				}
+			);
+
+			PrivateDependencyModuleNames.AddRange(
+				new string[] 
+				{
+				}
+			);
+		}
+	}
+}
diff --git a/Engine/Plugins/Experimental/GeometryCollectionPlugin/Source/GeometryCollectionDeprecatedNodes/Private/Dataflow/GeometryCollectionDeprecatedNodesPlugin.cpp b/Engine/Plugins/Experimental/GeometryCollectionPlugin/Source/GeometryCollectionDeprecatedNodes/Private/Dataflow/GeometryCollectionDeprecatedNodesPlugin.cpp
new file mode 100644
--- /dev/null
+++ b/Engine/Plugins/Experimental/GeometryCollectionPlugin/Source/GeometryCollectionDeprecatedNodes/Private/Dataflow/GeometryCollectionDeprecatedNodesPlugin.cpp
@@ -0,0 +1,30 @@
+// Copyright Epic Games, Inc. All Rights Reserved.
+
+#include "Dataflow/GeometryCollectionDeprecatedNodesPlugin.h"
+
+#include "CoreMinimal.h"
+
+#include "Dataflow/DataflowNodeFactory.h"
+#include "Dataflow/SetVertexColorFromFloatArrayDepNode.h"
+#include "Dataflow/SetVertexColorFromVertexSelectionDepNode.h"
+
+
+
+#define LOCTEXT_NAMESPACE "DataflowNodes"
+
+
+void IGeometryCollectionDeprecatedNodesPlugin::StartupModule()
+{
+	DATAFLOW_NODE_REGISTER_CREATION_FACTORY(FSetVertexColorInCollectionFromVertexSelectionDataflowNode);
+	DATAFLOW_NODE_REGISTER_CREATION_FACTORY(FSetVertexColorInCollectionFromFloatArrayDataflowNode);
+}
+
+void IGeometryCollectionDeprecatedNodesPlugin::ShutdownModule()
+{
+}
+
+
+IMPLEMENT_MODULE(IGeometryCollectionDeprecatedNodesPlugin, GeometryCollectionDeprecatedNodes)
+
+
+#undef LOCTEXT_NAMESPACE
diff --git a/Engine/Plugins/Experimental/GeometryCollectionPlugin/Source/GeometryCollectionDeprecatedNodes/Private/Dataflow/SetVertexColorFromFloatArrayDepNode.cpp b/Engine/Plugins/Experimental/GeometryCollectionPlugin/Source/GeometryCollectionDeprecatedNodes/Private/Dataflow/SetVertexColorFromFloatArrayDepNode.cpp
new file mode 100644
--- /dev/null
+++ b/Engine/Plugins/Experimental/GeometryCollectionPlugin/Source/GeometryCollectionDeprecatedNodes/Private/Dataflow/SetVertexColorFromFloatArrayDepNode.cpp
@@ -0,0 +1,31 @@
+// Copyright Epic Games, Inc. All Rights Reserved.
+
+#include "Dataflow/SetVertexColorFromFloatArrayDepNode.h"
+#include "Dataflow/DataflowCore.h"
+
+
+#include UE_INLINE_GENERATED_CPP_BY_NAME(SetVertexColorFromFloatArrayDepNode)
+
+void FSetVertexColorInCollectionFromFloatArrayDataflowNode::Evaluate(Dataflow::FContext& Context, const FDataflowOutput* Out) const
+{
+	if (Out->IsA<FManagedArrayCollection>(&Collection))
+	{
+		FManagedArrayCollection InCollection = GetValue<FManagedArrayCollection>(Context, &Collection);
+		const TArray<float>& InFloatArray = GetValue<TArray<float>>(Context, &FloatArray);
+
+		const int32 NumVertices = InCollection.NumElements(FGeometryCollection::VerticesGroup);
+
+		if (InFloatArray.Num() == NumVertices)
+		{
+			if (TManagedArray<FLinearColor>* VertexColors = InCollection.FindAttribute<FLinearColor>("Color", FGeometryCollection::VerticesGroup))
+			{
+				for (int32 Idx = 0; Idx < NumVertices; ++Idx)
+				{
+					(*VertexColors)[Idx] = FLinearColor(Scale * InFloatArray[Idx], Scale * InFloatArray[Idx], Scale * InFloatArray[Idx]);
+				}
+			}
+		}
+
+		SetValue(Context, MoveTemp(InCollection), &Collection);
+	}
+}
\ No newline at end of file
diff --git a/Engine/Plugins/Experimental/GeometryCollectionPlugin/Source/GeometryCollectionDeprecatedNodes/Private/Dataflow/SetVertexColorFromVertexSelectionDepNode.cpp b/Engine/Plugins/Experimental/GeometryCollectionPlugin/Source/GeometryCollectionDeprecatedNodes/Private/Dataflow/SetVertexColorFromVertexSelectionDepNode.cpp
new file mode 100644
--- /dev/null
+++ b/Engine/Plugins/Experimental/GeometryCollectionPlugin/Source/GeometryCollectionDeprecatedNodes/Private/Dataflow/SetVertexColorFromVertexSelectionDepNode.cpp
@@ -0,0 +1,37 @@
+// Copyright Epic Games, Inc. All Rights Reserved.
+
+#include "Dataflow/SetVertexColorFromVertexSelectionDepNode.h"
+
+#include UE_INLINE_GENERATED_CPP_BY_NAME(SetVertexColorFromVertexSelectionDepNode)
+
+void FSetVertexColorInCollectionFromVertexSelectionDataflowNode::Evaluate(Dataflow::FContext& Context, const FDataflowOutput* Out) const
+{
+	if (Out->IsA<FManagedArrayCollection>(&Collection))
+	{
+		FManagedArrayCollection InCollection = GetValue<FManagedArrayCollection>(Context, &Collection);
+		const FDataflowVertexSelection& InVertexSelection = GetValue<FDataflowVertexSelection>(Context, &VertexSelection);
+
+		if (InCollection.NumElements(FGeometryCollection::VerticesGroup) == InVertexSelection.Num())
+		{
+			const int32 NumVertices = InCollection.NumElements(FGeometryCollection::VerticesGroup);
+
+			//			TManagedArray<FLinearColor>& VertexColors = InCollection.ModifyAttribute<FLinearColor>("Color", FGeometryCollection::VerticesGroup);
+			if (TManagedArray<FLinearColor>* VertexColors = InCollection.FindAttribute<FLinearColor>("Color", FGeometryCollection::VerticesGroup))
+			{
+				for (int32 Idx = 0; Idx < NumVertices; ++Idx)
+				{
+					if (InVertexSelection.IsSelected(Idx))
+					{
+						(*VertexColors)[Idx] = SelectedColor;
+					}
+					else
+					{
+						(*VertexColors)[Idx] = NonSelectedColor;
+					}
+				}
+			}
+		}
+
+		SetValue(Context, MoveTemp(InCollection), &Collection);
+	}
+}
diff --git a/Engine/Plugins/Experimental/GeometryCollectionPlugin/Source/GeometryCollectionDeprecatedNodes/Public/Dataflow/GeometryCollectionDeprecatedNodesPlugin.h b/Engine/Plugins/Experimental/GeometryCollectionPlugin/Source/GeometryCollectionDeprecatedNodes/Public/Dataflow/GeometryCollectionDeprecatedNodesPlugin.h
new file mode 100644
--- /dev/null
+++ b/Engine/Plugins/Experimental/GeometryCollectionPlugin/Source/GeometryCollectionDeprecatedNodes/Public/Dataflow/GeometryCollectionDeprecatedNodesPlugin.h
@@ -0,0 +1,18 @@
+// Copyright Epic Games, Inc. All Rights Reserved.
+
+#pragma once
+
+#include "CoreMinimal.h"
+#include "Modules/ModuleInterface.h"
+#include "Modules/ModuleManager.h"
+
+/**
+ * The public interface to this module
+ */
+class IGeometryCollectionDeprecatedNodesPlugin : public IModuleInterface
+{
+public:
+	virtual void StartupModule();
+	virtual void ShutdownModule();
+};
+
diff --git a/Engine/Plugins/Experimental/GeometryCollectionPlugin/Source/GeometryCollectionDeprecatedNodes/Public/Dataflow/SetVertexColorFromFloatArrayDepNode.h b/Engine/Plugins/Experimental/GeometryCollectionPlugin/Source/GeometryCollectionDeprecatedNodes/Public/Dataflow/SetVertexColorFromFloatArrayDepNode.h
new file mode 100644
--- /dev/null
+++ b/Engine/Plugins/Experimental/GeometryCollectionPlugin/Source/GeometryCollectionDeprecatedNodes/Public/Dataflow/SetVertexColorFromFloatArrayDepNode.h
@@ -0,0 +1,42 @@
+// Copyright Epic Games, Inc. All Rights Reserved.
+
+#pragma once 
+
+#include "CoreMinimal.h"
+#include "Dataflow/DataflowEngine.h"
+#include "GeometryCollection/ManagedArrayCollection.h"
+
+
+#include "SetVertexColorFromFloatArrayDepNode.generated.h"
+
+/** Set collection color based on input float array */
+USTRUCT(meta = (Deprecated = "5.5"))
+struct FSetVertexColorInCollectionFromFloatArrayDataflowNode : public FDataflowNode
+{
+	GENERATED_USTRUCT_BODY()
+		DATAFLOW_NODE_DEFINE_INTERNAL(FSetVertexColorInCollectionFromFloatArrayDataflowNode, "SetVertexColorInCollectionFromFloatArray", "Collection|Utilities", "")
+
+public:
+	/** Collection */
+	UPROPERTY(meta = (DataflowInput, DataflowOutput, DataflowPassthrough = "Collection", DataflowIntrinsic))
+		FManagedArrayCollection Collection;
+
+	/**  */
+	UPROPERTY(meta = (DataflowInput, DataflowIntrinsic))
+		TArray<float> FloatArray;
+
+	/**  */
+	UPROPERTY(EditAnywhere, Category = "Color")
+		float Scale = 1.f;
+
+	FSetVertexColorInCollectionFromFloatArrayDataflowNode(const Dataflow::FNodeParameters& InParam, FGuid InGuid = FGuid::NewGuid())
+		: FDataflowNode(InParam, InGuid)
+	{
+		RegisterInputConnection(&Collection);
+		RegisterInputConnection(&FloatArray);
+		RegisterOutputConnection(&Collection, &Collection);
+	}
+
+	virtual void Evaluate(Dataflow::FContext& Context, const FDataflowOutput* Out) const override;
+
+};
\ No newline at end of file
diff --git a/Engine/Plugins/Experimental/GeometryCollectionPlugin/Source/GeometryCollectionDeprecatedNodes/Public/Dataflow/SetVertexColorFromVertexSelectionDepNode.h b/Engine/Plugins/Experimental/GeometryCollectionPlugin/Source/GeometryCollectionDeprecatedNodes/Public/Dataflow/SetVertexColorFromVertexSelectionDepNode.h
new file mode 100644
--- /dev/null
+++ b/Engine/Plugins/Experimental/GeometryCollectionPlugin/Source/GeometryCollectionDeprecatedNodes/Public/Dataflow/SetVertexColorFromVertexSelectionDepNode.h
@@ -0,0 +1,46 @@
+// Copyright Epic Games, Inc. All Rights Reserved.
+
+#pragma once 
+
+#include "CoreMinimal.h"
+#include "Dataflow/DataflowEngine.h"
+#include "Dataflow/DataflowSelection.h"
+#include "GeometryCollection/ManagedArrayCollection.h"
+
+#include "SetVertexColorFromVertexSelectionDepNode.generated.h"
+
+/** Set the collections vertex color from the selection set. */
+USTRUCT(meta = (Deprecated = "5.5"))
+struct FSetVertexColorInCollectionFromVertexSelectionDataflowNode : public FDataflowNode
+{
+	GENERATED_USTRUCT_BODY()
+	DATAFLOW_NODE_DEFINE_INTERNAL(FSetVertexColorInCollectionFromVertexSelectionDataflowNode, "SetVertexColorInCollectionFromVertexSelection", "Collection|Utilities", "")
+
+public:
+	/** Collection */
+	UPROPERTY(meta = (DataflowInput, DataflowOutput, DataflowPassthrough = "Collection", DataflowIntrinsic))
+		FManagedArrayCollection Collection;
+
+	/**  */
+	UPROPERTY(meta = (DataflowInput, DisplayName = "VertexSelection", DataflowIntrinsic))
+		FDataflowVertexSelection VertexSelection;
+
+	/**  */
+	UPROPERTY(EditAnywhere, Category = "Color")
+		FLinearColor SelectedColor = FLinearColor(FColor::Yellow);
+
+	/**  */
+	UPROPERTY(EditAnywhere, Category = "Color", meta = (DisplayName = "NonSelected Color"))
+		FLinearColor NonSelectedColor = FLinearColor(FColor::Blue);
+
+	FSetVertexColorInCollectionFromVertexSelectionDataflowNode(const Dataflow::FNodeParameters& InParam, FGuid InGuid = FGuid::NewGuid())
+		: FDataflowNode(InParam, InGuid)
+	{
+		RegisterInputConnection(&Collection);
+		RegisterInputConnection(&VertexSelection);
+		RegisterOutputConnection(&Collection, &Collection);
+	}
+
+	virtual void Evaluate(Dataflow::FContext& Context, const FDataflowOutput* Out) const override;
+
+};
diff --git a/Engine/Plugins/Experimental/GeometryCollectionPlugin/Source/GeometryCollectionNodes/Private/Dataflow/CreateColorArrayFromFloatArrayNode.cpp b/Engine/Plugins/Experimental/GeometryCollectionPlugin/Source/GeometryCollectionNodes/Private/Dataflow/CreateColorArrayFromFloatArrayNode.cpp
new file mode 100644
--- /dev/null
+++ b/Engine/Plugins/Experimental/GeometryCollectionPlugin/Source/GeometryCollectionNodes/Private/Dataflow/CreateColorArrayFromFloatArrayNode.cpp
@@ -0,0 +1,46 @@
+// Copyright Epic Games, Inc. All Rights Reserved.
+
+#include "Dataflow/CreateColorArrayFromFloatArrayNode.h"
+#include "Dataflow/DataflowCore.h"
+
+#include UE_INLINE_GENERATED_CPP_BY_NAME(CreateColorArrayFromFloatArrayNode)
+
+void FCreateColorArrayFromFloatArrayDataflowNode::Evaluate(Dataflow::FContext& Context, const FDataflowOutput* Out) const
+{
+	if (Out->IsA(&ColorArray))
+	{
+		TArray<FLinearColor> ColorArrayVal;
+		const TArray<float>& InFloatArray = GetValue<TArray<float>>(Context, &FloatArray);
+		int Num  = InFloatArray.Num();
+		if (Num)
+		{
+			float MaxValue = 0;
+			if (bNormalizeInput)
+			{
+				for (int32 Idx = 0; Idx < Num; ++Idx)
+				{
+					MaxValue = FMath::Max(MaxValue, InFloatArray[Idx]);
+				}
+			}
+			else
+			{
+				MaxValue = 1.0;
+			}
+
+			if (FMath::IsNearlyZero(MaxValue))
+			{
+				ColorArrayVal.Init(FLinearColor::Black, Num);
+			}
+			else
+			{
+				ColorArrayVal.Init(Color, Num);
+				for (int32 Idx = 0; Idx < Num; ++Idx)
+				{
+					ColorArrayVal[Idx] *= InFloatArray[Idx] / MaxValue;
+				}
+			}
+		}
+
+		SetValue(Context, MoveTemp(ColorArrayVal), &ColorArray);
+	}
+}
diff --git a/Engine/Plugins/Experimental/GeometryCollectionPlugin/Source/GeometryCollectionNodes/Private/Dataflow/GeometryCollectionNodes.cpp b/Engine/Plugins/Experimental/GeometryCollectionPlugin/Source/GeometryCollectionNodes/Private/Dataflow/GeometryCollectionNodes.cpp
--- a/Engine/Plugins/Experimental/GeometryCollectionPlugin/Source/GeometryCollectionNodes/Private/Dataflow/GeometryCollectionNodes.cpp
+++ b/Engine/Plugins/Experimental/GeometryCollectionPlugin/Source/GeometryCollectionNodes/Private/Dataflow/GeometryCollectionNodes.cpp
@@ -77,8 +77,6 @@ namespace Dataflow
 		DATAFLOW_NODE_REGISTER_CREATION_FACTORY(FGetNumElementsInCollectionGroupDataflowNode);
 		DATAFLOW_NODE_REGISTER_CREATION_FACTORY(FGetCollectionAttributeDataTypedDataflowNode);
 		DATAFLOW_NODE_REGISTER_CREATION_FACTORY(FSetCollectionAttributeDataTypedDataflowNode);
-		DATAFLOW_NODE_REGISTER_CREATION_FACTORY(FSetVertexColorInCollectionFromVertexSelectionDataflowNode);
-		DATAFLOW_NODE_REGISTER_CREATION_FACTORY(FSetVertexColorInCollectionFromFloatArrayDataflowNode);
 		DATAFLOW_NODE_REGISTER_CREATION_FACTORY(FMultiplyTransformDataflowNode);
 		DATAFLOW_NODE_REGISTER_CREATION_FACTORY(FInvertTransformDataflowNode);
 		DATAFLOW_NODE_REGISTER_CREATION_FACTORY(FSelectionToVertexListDataflowNode);
@@ -953,39 +951,60 @@ void FGetCollectionAttributeDataTypedDataflowNode::Evaluate(Dataflow::FContext&
 
 					if (TypeStr == FString("Bool"))
 					{
-						const TManagedArray<bool>& AttributeArr = InCollection.GetAttribute<bool>(AttributeNameVal, GroupNameVal);
-						TArray<bool> BoolArray = AttributeArr.GetAsBoolArray();
-						SetValue(Context, MoveTemp(BoolArray), &BoolAttributeData);
+						if (const TManagedArray<bool>* AttributeArr = InCollection.FindAttribute<bool>(AttributeNameVal, GroupNameVal))
+						{
+							TArray<bool> BoolArray = AttributeArr->GetAsBoolArray();
+							SetValue(Context, MoveTemp(BoolArray), &BoolAttributeData);
+						}
 					}
 					else if (TypeStr == FString("Float"))
 					{
-						const TManagedArray<float>& AttributeArr = InCollection.GetAttribute<float>(AttributeNameVal, GroupNameVal);
-						SetValue(Context, AttributeArr.GetConstArray(), &FloatAttributeData);
+						if (const TManagedArray<float>* AttributeArr = InCollection.FindAttribute<float>(AttributeNameVal, GroupNameVal))
+						{
+							SetValue(Context, AttributeArr->GetConstArray(), &FloatAttributeData);
+						}
 					}
 					else if (TypeStr == FString("Double"))
 					{
-						const TManagedArray<double>& AttributeArr = InCollection.GetAttribute<double>(AttributeNameVal, GroupNameVal);
-						SetValue(Context, AttributeArr.GetConstArray(), &DoubleAttributeData);
+						if (const TManagedArray<double>* AttributeArr = InCollection.FindAttribute<double>(AttributeNameVal, GroupNameVal))
+						{
+							SetValue(Context, AttributeArr->GetConstArray(), &DoubleAttributeData);
+						}
 					}
 					else if (TypeStr == FString("Int32"))
 					{
-						const TManagedArray<int32>& AttributeArr = InCollection.GetAttribute<int32>(AttributeNameVal, GroupNameVal);
-						SetValue(Context, AttributeArr.GetConstArray(), &Int32AttributeData);
+						if (const TManagedArray<int32>* AttributeArr = InCollection.FindAttribute<int32>(AttributeNameVal, GroupNameVal))
+						{
+							SetValue(Context, AttributeArr->GetConstArray(), &Int32AttributeData);
+						}
 					}
 					else if (TypeStr == FString("String"))
 					{
-						const TManagedArray<FString>& AttributeArr = InCollection.GetAttribute<FString>(AttributeNameVal, GroupNameVal);
-						SetValue(Context, AttributeArr.GetConstArray(), &StringAttributeData);
+						if (const TManagedArray<FString>* AttributeArr = InCollection.FindAttribute<FString>(AttributeNameVal, GroupNameVal))
+						{
+							SetValue(Context, AttributeArr->GetConstArray(), &StringAttributeData);
+						}
 					}
 					else if (TypeStr == FString("Vector"))
 					{
-						const TManagedArray<FVector3f>& AttributeArr = InCollection.GetAttribute<FVector3f>(AttributeNameVal, GroupNameVal);
-						SetValue(Context, AttributeArr.GetConstArray(), &Vector3fAttributeData);
+						if (const TManagedArray<FVector3f>* AttributeArr = InCollection.FindAttribute<FVector3f>(AttributeNameVal, GroupNameVal))
+						{
+							SetValue(Context, AttributeArr->GetConstArray(), &Vector3fAttributeData);
+						}
 					}
 					else if (TypeStr == FString("Vector3d"))
 					{
-						const TManagedArray<FVector3d>& AttributeArr = InCollection.GetAttribute<FVector3d>(AttributeNameVal, GroupNameVal);
-						SetValue(Context, AttributeArr.GetConstArray(), &Vector3dAttributeData);
+						if (const TManagedArray<FVector3d>* AttributeArr = InCollection.FindAttribute<FVector3d>(AttributeNameVal, GroupNameVal))
+						{
+							SetValue(Context, AttributeArr->GetConstArray(), &Vector3dAttributeData);
+						}
+					}
+					else if (TypeStr == FString("LinearColor"))
+					{
+						if (const TManagedArray<FLinearColor>* AttributeArr = InCollection.FindAttribute<FLinearColor>(AttributeNameVal, GroupNameVal))
+						{
+							SetValue(Context, AttributeArr->GetConstArray(), &LinearColorAttributeData);
+						}
 					}
 				}
 			}
@@ -998,14 +1017,17 @@ static void SetAttributeData(const FDataflowNode* DataflowNode, Dataflow::FConte
 {
 	if (DataflowNode && DataflowNode->IsConnected<TArray<T>>(&Property))
 	{
-		TArray<T> AttributeData = DataflowNode->GetValue<TArray<T>>(Context, &Property);
-		TManagedArray<T>& AttributeArray = InCollection.ModifyAttribute<T>(AttributeName, GroupName);
-
-		if (AttributeData.Num() == AttributeArray.Num())
+		const TArray<T> & AttributeData = DataflowNode->GetValue<TArray<T>>(Context, &Property);
+		if (InCollection.FindAttributeTyped<T>(AttributeName, GroupName))
 		{
-			for (int32 Idx = 0; Idx < AttributeArray.Num(); ++Idx)
+			TManagedArray<T>& AttributeArray = InCollection.ModifyAttribute<T>(AttributeName, GroupName);
+
+			if (AttributeData.Num() == AttributeArray.Num())
 			{
-				AttributeArray[Idx] = AttributeData[Idx];
+				for (int32 Idx = 0; Idx < AttributeArray.Num(); ++Idx)
+				{
+					AttributeArray[Idx] = AttributeData[Idx];
+				}
 			}
 		}
 	}
@@ -1017,87 +1039,60 @@ void FSetCollectionAttributeDataTypedDataflowNode::Evaluate(Dataflow::FContext&
 	{
 		FManagedArrayCollection InCollection = GetValue<FManagedArrayCollection>(Context, &Collection);
 
-		FName GroupNameToUse;
+		FName InputGroupName;
 		if (GroupName != EStandardGroupNameEnum::Dataflow_EStandardGroupNameEnum_Custom)
 		{
-			GroupNameToUse = GetGroupName(GroupName);
+			InputGroupName = GetGroupName(GroupName);
 		}
 		else
 		{
-			GroupNameToUse = FName(*CustomGroupName);
+			InputGroupName = FName(*CustomGroupName);
 		}
 
-		if (GroupNameToUse.GetStringLength() > 0 && AttrName.Len() > 0)
+		FCollectionAttributeKey DefaultAttributeKey(AttrName, InputGroupName.ToString());
+		FCollectionAttributeKey AttributeKeyVal = GetValue(Context, &AttributeKey, DefaultAttributeKey);
+		FName GroupNameVal = FName(AttributeKeyVal.Group);
+		FName AttributeNameVal = FName(AttributeKeyVal.Attribute);
+
+		if (GroupNameVal.GetStringLength() && AttributeNameVal.GetStringLength() )
 		{
-			if (InCollection.HasGroup(GroupNameToUse))
+			if (InCollection.HasGroup(GroupNameVal))
 			{
-				if (InCollection.HasAttribute(FName(*AttrName), GroupNameToUse))
+				if (InCollection.HasAttribute(AttributeNameVal, GroupNameVal))
 				{
-					FName AttributeName = FName(*AttrName);
-					FString TypeStr = GetArrayTypeString(InCollection.GetAttributeType(AttributeName, GroupNameToUse)).ToString();
+					FString TypeStr = GetArrayTypeString(InCollection.GetAttributeType(AttributeNameVal, GroupNameVal)).ToString();
 					
 					if (TypeStr == FString("Bool"))
 					{
-						SetAttributeData<bool>(this, Context, InCollection, BoolAttributeData, AttributeName, GroupNameToUse);
+						SetAttributeData<bool>(this, Context, InCollection, BoolAttributeData, AttributeNameVal, GroupNameVal);
 					}
 					else if (TypeStr == FString("Float"))
 					{
-						SetAttributeData<float>(this, Context, InCollection, FloatAttributeData, AttributeName, GroupNameToUse);
+						SetAttributeData<float>(this, Context, InCollection, FloatAttributeData, AttributeNameVal, GroupNameVal);
 					}
 					else if (TypeStr == FString("Double"))
 					{
-						SetAttributeData<double>(this, Context, InCollection, DoubleAttributeData, AttributeName, GroupNameToUse);
+						SetAttributeData<double>(this, Context, InCollection, DoubleAttributeData, AttributeNameVal, GroupNameVal);
 					}
 					else if (TypeStr == FString("Int32"))
 					{
-						SetAttributeData<int32>(this, Context, InCollection, Int32AttributeData, AttributeName, GroupNameToUse);
+						SetAttributeData<int32>(this, Context, InCollection, Int32AttributeData, AttributeNameVal, GroupNameVal);
 					}
 					else if (TypeStr == FString("String"))
 					{
-						SetAttributeData<FString>(this, Context, InCollection, StringAttributeData, AttributeName, GroupNameToUse);
+						SetAttributeData<FString>(this, Context, InCollection, StringAttributeData, AttributeNameVal, GroupNameVal);
 					}
 					else if (TypeStr == FString("Vector"))
 					{
-						SetAttributeData<FVector3f>(this, Context, InCollection, Vector3fAttributeData, AttributeName, GroupNameToUse);
+						SetAttributeData<FVector3f>(this, Context, InCollection, Vector3fAttributeData, AttributeNameVal, GroupNameVal);
 					}
 					else if (TypeStr == FString("Vector3d"))
 					{
-						SetAttributeData<FVector3d>(this, Context, InCollection, Vector3dAttributeData, AttributeName, GroupNameToUse);
+						SetAttributeData<FVector3d>(this, Context, InCollection, Vector3dAttributeData, AttributeNameVal, GroupNameVal);
 					}
-				}
-			}
-		}
-
-		SetValue(Context, MoveTemp(InCollection), &Collection);
-	}
-}
-
-
-
-
-void FSetVertexColorInCollectionFromVertexSelectionDataflowNode::Evaluate(Dataflow::FContext& Context, const FDataflowOutput* Out) const
-{
-	if (Out->IsA<FManagedArrayCollection>(&Collection))
-	{
-		FManagedArrayCollection InCollection = GetValue<FManagedArrayCollection>(Context, &Collection);
-		const FDataflowVertexSelection& InVertexSelection = GetValue<FDataflowVertexSelection>(Context, &VertexSelection);
-
-		if (InCollection.NumElements(FGeometryCollection::VerticesGroup) == InVertexSelection.Num())
-		{
-			const int32 NumVertices = InCollection.NumElements(FGeometryCollection::VerticesGroup);
-
-//			TManagedArray<FLinearColor>& VertexColors = InCollection.ModifyAttribute<FLinearColor>("Color", FGeometryCollection::VerticesGroup);
-			if (TManagedArray<FLinearColor>* VertexColors = InCollection.FindAttribute<FLinearColor>("Color", FGeometryCollection::VerticesGroup))
-			{
-				for (int32 Idx = 0; Idx < NumVertices; ++Idx)
-				{
-					if (InVertexSelection.IsSelected(Idx))
-					{
-						(*VertexColors)[Idx] = SelectedColor;
-					}
-					else
+					else if (TypeStr == FString("LinearColor"))
 					{
-						(*VertexColors)[Idx] = NonSelectedColor;
+						SetAttributeData<FLinearColor>(this, Context, InCollection, LinearColorAttributeData, AttributeNameVal, GroupNameVal);
 					}
 				}
 			}
@@ -1113,32 +1108,6 @@ void FSelectionToVertexListDataflowNode::Evaluate(Dataflow::FContext& Context, c
 	SetValue(Context, InVertexSelection.AsArray(), &VertexList);
 }
 
-void FSetVertexColorInCollectionFromFloatArrayDataflowNode::Evaluate(Dataflow::FContext& Context, const FDataflowOutput* Out) const
-{
-	if (Out->IsA<FManagedArrayCollection>(&Collection))
-	{
-		FManagedArrayCollection InCollection = GetValue<FManagedArrayCollection>(Context, &Collection);
-		const TArray<float>& InFloatArray = GetValue<TArray<float>>(Context, &FloatArray);
-
-		const int32 NumVertices = InCollection.NumElements(FGeometryCollection::VerticesGroup);
-
-		if (InFloatArray.Num() == NumVertices)
-		{
-			if (TManagedArray<FLinearColor>* VertexColors = InCollection.FindAttribute<FLinearColor>("Color", FGeometryCollection::VerticesGroup))
-			{
-				for (int32 Idx = 0; Idx < NumVertices; ++Idx)
-				{
-					(*VertexColors)[Idx] = FLinearColor(Scale * InFloatArray[Idx], Scale * InFloatArray[Idx], Scale * InFloatArray[Idx]);
-				}
-			}
-		}
-
-		SetValue(Context, MoveTemp(InCollection), &Collection);
-	}
-}
-
-
-
 void FMultiplyTransformDataflowNode::Evaluate(Dataflow::FContext& Context, const FDataflowOutput* Out) const
 {
 	if (Out->IsA<FTransform>(&OutTransform))
diff --git a/Engine/Plugins/Experimental/GeometryCollectionPlugin/Source/GeometryCollectionNodes/Private/Dataflow/GeometryCollectionNodesPlugin.cpp b/Engine/Plugins/Experimental/GeometryCollectionPlugin/Source/GeometryCollectionNodes/Private/Dataflow/GeometryCollectionNodesPlugin.cpp
--- a/Engine/Plugins/Experimental/GeometryCollectionPlugin/Source/GeometryCollectionNodes/Private/Dataflow/GeometryCollectionNodesPlugin.cpp
+++ b/Engine/Plugins/Experimental/GeometryCollectionPlugin/Source/GeometryCollectionNodes/Private/Dataflow/GeometryCollectionNodesPlugin.cpp
@@ -5,6 +5,7 @@
 #include "CoreMinimal.h"
 
 #include "Dataflow/DataflowNodeFactory.h"
+#include "Dataflow/CreateColorArrayFromFloatArrayNode.h"
 #include "Dataflow/GeometryCollectionArrayNodes.h"
 #include "Dataflow/GeometryCollectionAssetNodes.h"
 #include "Dataflow/GeometryCollectionClusteringNodes.h"
@@ -29,6 +30,9 @@
 #include "Dataflow/GeometryCollectionUtilityNodes.h"
 #include "Dataflow/GeometryCollectionVerticesNodes.h"
 #include "Dataflow/GeometryCollectionVertexScalarToVertexIndicesNode.h"
+#include "Dataflow/SetVertexColorFromFloatArrayNode.h"
+#include "Dataflow/SetVertexColorFromVertexIndicesNode.h"
+#include "Dataflow/SetVertexColorFromVertexSelectionNode.h"
 
 #define LOCTEXT_NAMESPACE "DataflowNodes"
 
@@ -53,7 +57,11 @@ void IGeometryCollectionNodesPlugin::StartupModule()
 	Dataflow::GeometryCollectionVerticesNodes();
 	Dataflow::GeometryCollectionArrayNodes();
 	Dataflow::GeometryCollectionDebugNodes();
+	DATAFLOW_NODE_REGISTER_CREATION_FACTORY(FCreateColorArrayFromFloatArrayDataflowNode);
 	DATAFLOW_NODE_REGISTER_CREATION_FACTORY(FGeometryCollectionVertexScalarToVertexIndicesNode);
+	DATAFLOW_NODE_REGISTER_CREATION_FACTORY(FSetVertexColorFromFloatArrayDataflowNode);
+	DATAFLOW_NODE_REGISTER_CREATION_FACTORY(FSetVertexColorFromVertexIndicesDataflowNode);
+	DATAFLOW_NODE_REGISTER_CREATION_FACTORY(FSetVertexColorFromVertexSelectionDataflowNode);
 	DATAFLOW_NODE_REGISTER_CREATION_FACTORY(FSkeletonToCollectionDataflowNode);
 	DATAFLOW_NODE_REGISTER_CREATION_FACTORY(FSkeletalMeshToCollectionDataflowNode);
 	DATAFLOW_NODE_REGISTER_CREATION_FACTORY(FGeometryCollectionTransferVertexScalarAttributeNode);
diff --git a/Engine/Plugins/Experimental/GeometryCollectionPlugin/Source/GeometryCollectionNodes/Private/Dataflow/SetVertexColorFromFloatArrayNode.cpp b/Engine/Plugins/Experimental/GeometryCollectionPlugin/Source/GeometryCollectionNodes/Private/Dataflow/SetVertexColorFromFloatArrayNode.cpp
new file mode 100644
--- /dev/null
+++ b/Engine/Plugins/Experimental/GeometryCollectionPlugin/Source/GeometryCollectionNodes/Private/Dataflow/SetVertexColorFromFloatArrayNode.cpp
@@ -0,0 +1,50 @@
+// Copyright Epic Games, Inc. All Rights Reserved.
+
+#include "Dataflow/SetVertexColorFromFloatArrayNode.h"
+#include "Dataflow/DataflowCore.h"
+
+#include UE_INLINE_GENERATED_CPP_BY_NAME(SetVertexColorFromFloatArrayNode)
+
+void FSetVertexColorFromFloatArrayDataflowNode::Evaluate(Dataflow::FContext& Context, const FDataflowOutput* Out) const
+{
+	if (Out->IsA<FManagedArrayCollection>(&Collection))
+	{
+		FManagedArrayCollection InCollection = GetValue<FManagedArrayCollection>(Context, &Collection);
+		const TArray<float>& InFloatArray = GetValue<TArray<float>>(Context, &FloatArray);
+
+		const int32 NumVertices = InCollection.NumElements(FGeometryCollection::VerticesGroup);
+
+		if (InFloatArray.Num() == NumVertices)
+		{
+			float MaxValue = 0;
+			if (bNormalizeInput)
+			{
+				for (int32 Idx = 0; Idx < NumVertices; ++Idx)
+				{
+					MaxValue = FMath::Max(MaxValue, InFloatArray[Idx]);
+				}
+			}
+			else
+			{
+				MaxValue = 1.0;
+			}
+
+			if (TManagedArray<FLinearColor>* VertexColors = InCollection.FindAttribute<FLinearColor>("Color", FGeometryCollection::VerticesGroup))
+			{
+				if (FMath::IsNearlyZero(MaxValue))
+				{
+					VertexColors->Fill(FLinearColor::Black);
+				}
+				else
+				{
+					for (int32 Idx = 0; Idx < NumVertices; ++Idx)
+					{
+						(*VertexColors)[Idx] = InFloatArray[Idx]/MaxValue * Color;
+					}
+				}
+			}
+		}
+
+		SetValue(Context, MoveTemp(InCollection), &Collection);
+	}
+}
diff --git a/Engine/Plugins/Experimental/GeometryCollectionPlugin/Source/GeometryCollectionNodes/Private/Dataflow/SetVertexColorFromVertexIndicesNode.cpp b/Engine/Plugins/Experimental/GeometryCollectionPlugin/Source/GeometryCollectionNodes/Private/Dataflow/SetVertexColorFromVertexIndicesNode.cpp
new file mode 100644
--- /dev/null
+++ b/Engine/Plugins/Experimental/GeometryCollectionPlugin/Source/GeometryCollectionNodes/Private/Dataflow/SetVertexColorFromVertexIndicesNode.cpp
@@ -0,0 +1,32 @@
+// Copyright Epic Games, Inc. All Rights Reserved.
+
+#include "Dataflow/SetVertexColorFromVertexIndicesNode.h"
+#include "Dataflow/DataflowCore.h"
+
+#include UE_INLINE_GENERATED_CPP_BY_NAME(SetVertexColorFromVertexIndicesNode)
+
+void FSetVertexColorFromVertexIndicesDataflowNode::Evaluate(Dataflow::FContext& Context, const FDataflowOutput* Out) const
+{
+	if (Out->IsA<FManagedArrayCollection>(&Collection))
+	{
+		auto IsValidVertex = [](int32 Index, int32 Num) {return 0 <= Index && Index < Num; };
+
+		FManagedArrayCollection InCollection = GetValue<FManagedArrayCollection>(Context, &Collection);
+
+		const int32 NumVertices = InCollection.NumElements(FGeometryCollection::VerticesGroup);
+		const TArray<int32>& VertexIndices = GetValue<TArray<int32>>(Context, &VertexIndicesIn, VertexIndicesIn);
+
+		if (TManagedArray<FLinearColor>* VertexColors = InCollection.FindAttribute<FLinearColor>("Color", FGeometryCollection::VerticesGroup))
+		{
+			for (int32 Idx : VertexIndices)
+			{
+				if (IsValidVertex(Idx, NumVertices))
+				{
+					(*VertexColors)[Idx] = SelectedColor;
+				}
+			}
+		}
+
+		SetValue(Context, MoveTemp(InCollection), &Collection);
+	}
+}
diff --git a/Engine/Plugins/Experimental/GeometryCollectionPlugin/Source/GeometryCollectionNodes/Private/Dataflow/SetVertexColorFromVertexSelectionNode.cpp b/Engine/Plugins/Experimental/GeometryCollectionPlugin/Source/GeometryCollectionNodes/Private/Dataflow/SetVertexColorFromVertexSelectionNode.cpp
new file mode 100644
--- /dev/null
+++ b/Engine/Plugins/Experimental/GeometryCollectionPlugin/Source/GeometryCollectionNodes/Private/Dataflow/SetVertexColorFromVertexSelectionNode.cpp
@@ -0,0 +1,32 @@
+// Copyright Epic Games, Inc. All Rights Reserved.
+
+#include "Dataflow/SetVertexColorFromVertexSelectionNode.h"
+#include "Dataflow/DataflowCore.h"
+
+#include UE_INLINE_GENERATED_CPP_BY_NAME(SetVertexColorFromVertexSelectionNode)
+
+void FSetVertexColorFromVertexSelectionDataflowNode::Evaluate(Dataflow::FContext& Context, const FDataflowOutput* Out) const
+{
+	if (Out->IsA<FManagedArrayCollection>(&Collection))
+	{
+		FManagedArrayCollection InCollection = GetValue<FManagedArrayCollection>(Context, &Collection);
+		const FDataflowVertexSelection& InVertexSelection = GetValue<FDataflowVertexSelection>(Context, &VertexSelection);
+
+		if (InCollection.NumElements(FGeometryCollection::VerticesGroup) == InVertexSelection.Num())
+		{
+			const int32 NumVertices = InCollection.NumElements(FGeometryCollection::VerticesGroup);
+			if (TManagedArray<FLinearColor>* VertexColors = InCollection.FindAttribute<FLinearColor>("Color", FGeometryCollection::VerticesGroup))
+			{
+				for (int32 Idx = 0; Idx < NumVertices; ++Idx)
+				{
+					if (InVertexSelection.IsSelected(Idx))
+					{
+						(*VertexColors)[Idx] = SelectedColor;
+					}
+				}
+			}
+		}
+
+		SetValue(Context, MoveTemp(InCollection), &Collection);
+	}
+}
diff --git a/Engine/Plugins/Experimental/GeometryCollectionPlugin/Source/GeometryCollectionNodes/Public/Dataflow/CreateColorArrayFromFloatArrayNode.h b/Engine/Plugins/Experimental/GeometryCollectionPlugin/Source/GeometryCollectionNodes/Public/Dataflow/CreateColorArrayFromFloatArrayNode.h
new file mode 100644
--- /dev/null
+++ b/Engine/Plugins/Experimental/GeometryCollectionPlugin/Source/GeometryCollectionNodes/Public/Dataflow/CreateColorArrayFromFloatArrayNode.h
@@ -0,0 +1,47 @@
+// Copyright Epic Games, Inc. All Rights Reserved.
+
+#pragma once 
+
+#include "CoreMinimal.h"
+#include "Dataflow/DataflowEngine.h"
+#include "GeometryCollection/ManagedArrayCollection.h"
+
+#include "CreateColorArrayFromFloatArrayNode.generated.h"
+
+class FGeometryCollection;
+
+/** Set the vertex color on the collection based on the normalized float array. */
+USTRUCT()
+struct FCreateColorArrayFromFloatArrayDataflowNode : public FDataflowNode
+{
+	GENERATED_USTRUCT_BODY()
+	DATAFLOW_NODE_DEFINE_INTERNAL(FCreateColorArrayFromFloatArrayDataflowNode, "CreateColorArrayFromFloatArray", "Collection|Utilities", "")
+
+public:
+	/** Float array to use as a scalar for the color */
+	UPROPERTY(meta = (DataflowInput, DisplayName  = "Float Array"))
+	TArray<float> FloatArray;
+
+	/** Color array output */
+	UPROPERTY(meta = (DataflowOutput, DisplayName = "Linear Color Array"))
+	TArray<FLinearColor> ColorArray;
+
+	/** Enable normalization of input array */
+	UPROPERTY(EditAnywhere, Category = "Color", meta = (DisplayName = "Normalize Input"))
+	bool bNormalizeInput = true;
+	
+	/** Base color for the normalized float array */
+	UPROPERTY(EditAnywhere, Category = "Color", meta = (DisplayName = "Color"))
+	FLinearColor Color = FLinearColor(FColor::White);
+
+
+	FCreateColorArrayFromFloatArrayDataflowNode(const Dataflow::FNodeParameters& InParam, FGuid InGuid = FGuid::NewGuid())
+		: FDataflowNode(InParam, InGuid)
+	{
+		RegisterInputConnection(&FloatArray);
+		RegisterOutputConnection(&ColorArray);
+	}
+
+	virtual void Evaluate(Dataflow::FContext& Context, const FDataflowOutput* Out) const override;
+
+};
diff --git a/Engine/Plugins/Experimental/GeometryCollectionPlugin/Source/GeometryCollectionNodes/Public/Dataflow/GeometryCollectionNodes.h b/Engine/Plugins/Experimental/GeometryCollectionPlugin/Source/GeometryCollectionNodes/Public/Dataflow/GeometryCollectionNodes.h
--- a/Engine/Plugins/Experimental/GeometryCollectionPlugin/Source/GeometryCollectionNodes/Public/Dataflow/GeometryCollectionNodes.h
+++ b/Engine/Plugins/Experimental/GeometryCollectionPlugin/Source/GeometryCollectionNodes/Public/Dataflow/GeometryCollectionNodes.h
@@ -1336,33 +1336,37 @@ public:
 	FString AttrName = FString("");
 
 	/** Bool type attribute data */
-	UPROPERTY(meta = (DataflowOutput));
+	UPROPERTY(meta = (DataflowOutput, DisplayName = "Bool Array"));
 	TArray<bool> BoolAttributeData;
 
 	/** Float type attribute data */
-	UPROPERTY(meta = (DataflowOutput));
+	UPROPERTY(meta = (DataflowOutput, DisplayName = "Float Array"));
 	TArray<float> FloatAttributeData;
 
 	/** Float type attribute data */
-	UPROPERTY(meta = (DataflowOutput));
+	UPROPERTY(meta = (DataflowOutput, DisplayName = "Double Array"));
 	TArray<double> DoubleAttributeData;
 
 	/** Int type attribute data */
-	UPROPERTY(meta = (DataflowOutput));
+	UPROPERTY(meta = (DataflowOutput, DisplayName = "Int32 Array"));
 	TArray<int32> Int32AttributeData;
 
 	/** Int type attribute data */
-	UPROPERTY(meta = (DataflowOutput));
+	UPROPERTY(meta = (DataflowOutput, DisplayName = "String Array"));
 	TArray<FString> StringAttributeData;
 
 	/** Vector3f type attribute data */
-	UPROPERTY(meta = (DataflowOutput));
+	UPROPERTY(meta = (DataflowOutput, DisplayName = "Vector3f Array"));
 	TArray<FVector3f> Vector3fAttributeData;
 
 	/** Vector3d type attribute data */
-	UPROPERTY(meta = (DataflowOutput));
+	UPROPERTY(meta = (DataflowOutput, DisplayName = "Vector3d Array"));
 	TArray<FVector3d> Vector3dAttributeData;
 
+	/** Vector3d type attribute data */
+	UPROPERTY(meta = (DataflowOutput, DisplayName = "Linear Color Array"));
+	TArray<FLinearColor> LinearColorAttributeData;
+
 	FGetCollectionAttributeDataTypedDataflowNode(const Dataflow::FNodeParameters& InParam, FGuid InGuid = FGuid::NewGuid())
 		: FDataflowNode(InParam, InGuid)
 	{
@@ -1375,6 +1379,7 @@ public:
 		RegisterOutputConnection(&StringAttributeData);
 		RegisterOutputConnection(&Vector3fAttributeData);
 		RegisterOutputConnection(&Vector3dAttributeData);
+		RegisterOutputConnection(&LinearColorAttributeData);
 	}
 
 	virtual void Evaluate(Dataflow::FContext& Context, const FDataflowOutput* Out) const override;
@@ -1398,6 +1403,10 @@ public:
 	UPROPERTY(meta = (DataflowInput, DataflowOutput, DataflowPassthrough = "Collection", DataflowIntrinsic))
 	FManagedArrayCollection Collection;
 
+	/** Input to drive the Attribute and Group name */
+	UPROPERTY(meta = (DataflowInput, DisplayName = "AttributeKey"))
+	FCollectionAttributeKey AttributeKey;
+
 	/** Standard group names */
 	UPROPERTY(EditAnywhere, Category = "Attribute", meta = (DisplayName = "Group"))
 	EStandardGroupNameEnum GroupName = EStandardGroupNameEnum::Dataflow_EStandardGroupNameEnum_Transform;
@@ -1411,37 +1420,42 @@ public:
 	FString AttrName = FString("");
 
 	/** Bool type attribute data */
-	UPROPERTY(meta = (DataflowInput));
+	UPROPERTY(meta = (DataflowInput, DisplayName = "Bool Array"));
 	TArray<bool> BoolAttributeData;
 
 	/** Float type attribute data */
-	UPROPERTY(meta = (DataflowInput));
+	UPROPERTY(meta = (DataflowInput, DisplayName = "Float Array"));
 	TArray<float> FloatAttributeData;
 
 	/** Float type attribute data */
-	UPROPERTY(meta = (DataflowInput));
+	UPROPERTY(meta = (DataflowInput, DisplayName = "Double Array"));
 	TArray<double> DoubleAttributeData;
 
 	/** Int type attribute data */
-	UPROPERTY(meta = (DataflowInput));
+	UPROPERTY(meta = (DataflowInput, DisplayName = "Int32 Array"));
 	TArray<int32> Int32AttributeData;
 
 	/** Int type attribute data */
-	UPROPERTY(meta = (DataflowInput));
+	UPROPERTY(meta = (DataflowInput, DisplayName = "String Array"));
 	TArray<FString> StringAttributeData;
 
 	/** Vector3f type attribute data */
-	UPROPERTY(meta = (DataflowInput));
+	UPROPERTY(meta = (DataflowInput, DisplayName = "Vector3f Array"));
 	TArray<FVector3f> Vector3fAttributeData;
 
 	/** Vector3d type attribute data */
-	UPROPERTY(meta = (DataflowInput));
+	UPROPERTY(meta = (DataflowInput, DisplayName = "Vector3d Array"));
 	TArray<FVector3d> Vector3dAttributeData;
 
+	/** LinearColor type attribute data */
+	UPROPERTY(meta = (DataflowInput, DisplayName = "Linear Color Array"));
+	TArray<FLinearColor> LinearColorAttributeData;
+
 	FSetCollectionAttributeDataTypedDataflowNode(const Dataflow::FNodeParameters& InParam, FGuid InGuid = FGuid::NewGuid())
 		: FDataflowNode(InParam, InGuid)
 	{
 		RegisterInputConnection(&Collection);
+		RegisterInputConnection(&AttributeKey);
 		RegisterInputConnection(&BoolAttributeData);
 		RegisterInputConnection(&FloatAttributeData);
 		RegisterInputConnection(&DoubleAttributeData);
@@ -1449,47 +1463,7 @@ public:
 		RegisterInputConnection(&StringAttributeData);
 		RegisterInputConnection(&Vector3fAttributeData);
 		RegisterInputConnection(&Vector3dAttributeData);
-		RegisterOutputConnection(&Collection, &Collection);
-	}
-
-	virtual void Evaluate(Dataflow::FContext& Context, const FDataflowOutput* Out) const override;
-
-};
-
-
-/**
- *
- * 
- *
- */
-USTRUCT()
-struct FSetVertexColorInCollectionFromVertexSelectionDataflowNode : public FDataflowNode
-{
-	GENERATED_USTRUCT_BODY()
-	DATAFLOW_NODE_DEFINE_INTERNAL(FSetVertexColorInCollectionFromVertexSelectionDataflowNode, "SetVertexColorInCollectionFromVertexSelection", "Collection|Utilities", "")
-
-public:
-	/** Collection */
-	UPROPERTY(meta = (DataflowInput, DataflowOutput, DataflowPassthrough = "Collection", DataflowIntrinsic))
-	FManagedArrayCollection Collection;
-
-	/**  */
-	UPROPERTY(meta = (DataflowInput, DisplayName = "VertexSelection", DataflowIntrinsic))
-	FDataflowVertexSelection VertexSelection;
-
-	/**  */
-	UPROPERTY(EditAnywhere, Category = "Color")
-	FLinearColor SelectedColor = FLinearColor(FColor::Yellow);
-
-	/**  */
-	UPROPERTY(EditAnywhere, Category = "Color", meta = (DisplayName = "NonSelected Color"))
-	FLinearColor NonSelectedColor = FLinearColor(FColor::Blue);
-
-	FSetVertexColorInCollectionFromVertexSelectionDataflowNode(const Dataflow::FNodeParameters& InParam, FGuid InGuid = FGuid::NewGuid())
-		: FDataflowNode(InParam, InGuid)
-	{
-		RegisterInputConnection(&Collection);
-		RegisterInputConnection(&VertexSelection);
+		RegisterInputConnection(&LinearColorAttributeData);
 		RegisterOutputConnection(&Collection, &Collection);
 	}
 
@@ -1530,43 +1504,6 @@ public:
 
 };
 
-/**
- *
- *
- *
- */
-USTRUCT()
-struct FSetVertexColorInCollectionFromFloatArrayDataflowNode : public FDataflowNode
-{
-	GENERATED_USTRUCT_BODY()
-	DATAFLOW_NODE_DEFINE_INTERNAL(FSetVertexColorInCollectionFromFloatArrayDataflowNode, "SetVertexColorInCollectionFromFloatArray", "Collection|Utilities", "")
-
-public:
-	/** Collection */
-	UPROPERTY(meta = (DataflowInput, DataflowOutput, DataflowPassthrough = "Collection", DataflowIntrinsic))
-	FManagedArrayCollection Collection;
-
-	/**  */
-	UPROPERTY(meta = (DataflowInput, DataflowIntrinsic))
-	TArray<float> FloatArray;
-
-	/**  */
-	UPROPERTY(EditAnywhere, Category = "Color")
-	float Scale = 1.f;
-
-	FSetVertexColorInCollectionFromFloatArrayDataflowNode(const Dataflow::FNodeParameters& InParam, FGuid InGuid = FGuid::NewGuid())
-		: FDataflowNode(InParam, InGuid)
-	{
-		RegisterInputConnection(&Collection);
-		RegisterInputConnection(&FloatArray);
-		RegisterOutputConnection(&Collection, &Collection);
-	}
-
-	virtual void Evaluate(Dataflow::FContext& Context, const FDataflowOutput* Out) const override;
-
-};
-
-
 /**
  *
  * Description for this node
diff --git a/Engine/Plugins/Experimental/GeometryCollectionPlugin/Source/GeometryCollectionNodes/Public/Dataflow/SetVertexColorFromFloatArrayNode.h b/Engine/Plugins/Experimental/GeometryCollectionPlugin/Source/GeometryCollectionNodes/Public/Dataflow/SetVertexColorFromFloatArrayNode.h
new file mode 100644
--- /dev/null
+++ b/Engine/Plugins/Experimental/GeometryCollectionPlugin/Source/GeometryCollectionNodes/Public/Dataflow/SetVertexColorFromFloatArrayNode.h
@@ -0,0 +1,48 @@
+// Copyright Epic Games, Inc. All Rights Reserved.
+
+#pragma once 
+
+#include "CoreMinimal.h"
+#include "Dataflow/DataflowEngine.h"
+#include "GeometryCollection/ManagedArrayCollection.h"
+
+#include "SetVertexColorFromFloatArrayNode.generated.h"
+
+class FGeometryCollection;
+
+/** Set the vertex color on the collection based on the normalized float array. */
+USTRUCT()
+struct FSetVertexColorFromFloatArrayDataflowNode : public FDataflowNode
+{
+	GENERATED_USTRUCT_BODY()
+	DATAFLOW_NODE_DEFINE_INTERNAL(FSetVertexColorFromFloatArrayDataflowNode, "SetVertexColorFromFloatArray", "Collection|Utilities", "")
+
+public:
+	/** Collection Passthrough*/
+	UPROPERTY(meta = (DataflowInput, DataflowOutput, DataflowPassthrough = "Collection", DataflowIntrinsic))
+	FManagedArrayCollection Collection;
+
+	/** Float array to use as a scalar for the color */
+	UPROPERTY(meta = (DataflowInput, DataflowIntrinsic))
+	TArray<float> FloatArray;
+
+	/** Enable normalization of input array */
+	UPROPERTY(EditAnywhere, Category = "Color", meta = (DisplayName = "Normalize Input"))
+	bool bNormalizeInput = true;
+	
+	/** Base color for the normalized float array */
+	UPROPERTY(EditAnywhere, Category = "Color", meta = (DisplayName = "Color"))
+	FLinearColor Color = FLinearColor(FColor::White);
+
+
+	FSetVertexColorFromFloatArrayDataflowNode(const Dataflow::FNodeParameters& InParam, FGuid InGuid = FGuid::NewGuid())
+		: FDataflowNode(InParam, InGuid)
+	{
+		RegisterInputConnection(&Collection);
+		RegisterInputConnection(&FloatArray);
+		RegisterOutputConnection(&Collection, &Collection);
+	}
+
+	virtual void Evaluate(Dataflow::FContext& Context, const FDataflowOutput* Out) const override;
+
+};
diff --git a/Engine/Plugins/Experimental/GeometryCollectionPlugin/Source/GeometryCollectionNodes/Public/Dataflow/SetVertexColorFromVertexIndicesNode.h b/Engine/Plugins/Experimental/GeometryCollectionPlugin/Source/GeometryCollectionNodes/Public/Dataflow/SetVertexColorFromVertexIndicesNode.h
new file mode 100644
--- /dev/null
+++ b/Engine/Plugins/Experimental/GeometryCollectionPlugin/Source/GeometryCollectionNodes/Public/Dataflow/SetVertexColorFromVertexIndicesNode.h
@@ -0,0 +1,41 @@
+// Copyright Epic Games, Inc. All Rights Reserved.
+
+#pragma once 
+
+#include "CoreMinimal.h"
+#include "Dataflow/DataflowEngine.h"
+#include "GeometryCollection/ManagedArrayCollection.h"
+
+#include "SetVertexColorFromVertexIndicesNode.generated.h"
+
+/** Set the vertex color of the collection based on the selection set. */
+USTRUCT()
+struct FSetVertexColorFromVertexIndicesDataflowNode : public FDataflowNode
+{
+	GENERATED_USTRUCT_BODY()
+	DATAFLOW_NODE_DEFINE_INTERNAL(FSetVertexColorFromVertexIndicesDataflowNode, "SetVertexColorFromVertexIndices", "Collection|Utilities", "")
+
+public:
+	/** Collection Passthrough*/
+	UPROPERTY(meta = (DataflowInput, DataflowOutput, DataflowPassthrough = "Collection", DataflowIntrinsic))
+	FManagedArrayCollection Collection;
+
+	/** Vertex indices set */
+	UPROPERTY(EditAnywhere, Category = "Vertices", meta = (DataflowInput, DisplayName = "VertexIndices"))
+	TArray<int32> VertexIndicesIn;
+
+	/** Selected vertex color */
+	UPROPERTY(EditAnywhere, Category = "Color")
+	FLinearColor SelectedColor = FLinearColor(FColor::Yellow);
+
+	FSetVertexColorFromVertexIndicesDataflowNode(const Dataflow::FNodeParameters& InParam, FGuid InGuid = FGuid::NewGuid())
+		: FDataflowNode(InParam, InGuid)
+	{
+		RegisterInputConnection(&Collection);
+		RegisterInputConnection(&VertexIndicesIn);
+		RegisterOutputConnection(&Collection, &Collection);
+	}
+
+	virtual void Evaluate(Dataflow::FContext& Context, const FDataflowOutput* Out) const override;
+
+};
diff --git a/Engine/Plugins/Experimental/GeometryCollectionPlugin/Source/GeometryCollectionNodes/Public/Dataflow/SetVertexColorFromVertexSelectionNode.h b/Engine/Plugins/Experimental/GeometryCollectionPlugin/Source/GeometryCollectionNodes/Public/Dataflow/SetVertexColorFromVertexSelectionNode.h
new file mode 100644
--- /dev/null
+++ b/Engine/Plugins/Experimental/GeometryCollectionPlugin/Source/GeometryCollectionNodes/Public/Dataflow/SetVertexColorFromVertexSelectionNode.h
@@ -0,0 +1,42 @@
+// Copyright Epic Games, Inc. All Rights Reserved.
+
+#pragma once 
+
+#include "CoreMinimal.h"
+#include "Dataflow/DataflowEngine.h"
+#include "Dataflow/DataflowSelection.h"
+#include "GeometryCollection/ManagedArrayCollection.h"
+
+#include "SetVertexColorFromVertexSelectionNode.generated.h"
+
+/** Set the vertex color of the collection based on the selection set. */
+USTRUCT()
+struct FSetVertexColorFromVertexSelectionDataflowNode : public FDataflowNode
+{
+	GENERATED_USTRUCT_BODY()
+	DATAFLOW_NODE_DEFINE_INTERNAL(FSetVertexColorFromVertexSelectionDataflowNode, "SetVertexColorFromVertexSelection", "Collection|Utilities", "")
+
+public:
+	/** Collection Passthrough*/
+	UPROPERTY(meta = (DataflowInput, DataflowOutput, DataflowPassthrough = "Collection", DataflowIntrinsic))
+	FManagedArrayCollection Collection;
+
+	/** Vertex selection set */
+	UPROPERTY(meta = (DataflowInput, DisplayName = "VertexSelection", DataflowIntrinsic))
+	FDataflowVertexSelection VertexSelection;
+
+	/** Selected vertex color */
+	UPROPERTY(EditAnywhere, Category = "Color")
+	FLinearColor SelectedColor = FLinearColor(FColor::Yellow);
+
+	FSetVertexColorFromVertexSelectionDataflowNode(const Dataflow::FNodeParameters& InParam, FGuid InGuid = FGuid::NewGuid())
+		: FDataflowNode(InParam, InGuid)
+	{
+		RegisterInputConnection(&Collection);
+		RegisterInputConnection(&VertexSelection);
+		RegisterOutputConnection(&Collection, &Collection);
+	}
+
+	virtual void Evaluate(Dataflow::FContext& Context, const FDataflowOutput* Out) const override;
+
+};
