Chaos Flesh : Kinematic Node
- Update inputs and outputs to be VertexIndices

diff --git a/Engine/Plugins/Experimental/ChaosFlesh/Source/ChaosFleshNodes/Public/Dataflow/ChaosFleshKinematicConstraintNode.h b/Engine/Plugins/Experimental/ChaosFlesh/Source/ChaosFleshNodes/Public/Dataflow/ChaosFleshKinematicConstraintNode.h
--- a/Engine/Plugins/Experimental/ChaosFlesh/Source/ChaosFleshNodes/Public/Dataflow/ChaosFleshKinematicConstraintNode.h
+++ b/Engine/Plugins/Experimental/ChaosFlesh/Source/ChaosFleshNodes/Public/Dataflow/ChaosFleshKinematicConstraintNode.h
@@ -33,7 +33,7 @@ public:
 	UPROPERTY(meta = (DataflowInput, DisplayName = "SkeletalMesh"))
 	TObjectPtr<USkeletalMesh> SkeletalMeshIn;
 		
-	UPROPERTY(meta = (DataflowInput, DisplayName = "SelectionSet"))
+	UPROPERTY(meta = (DataflowInput, DisplayName = "VertexIndices"))
 	TArray<int32> VertexIndicesIn;
 
 	UPROPERTY(meta = (DataflowInput, DisplayName = "BoneIndex"))
diff --git a/Engine/Plugins/Experimental/GeometryCollectionPlugin/Source/GeometryCollectionNodes/Private/Dataflow/GeometryCollectionVertexScalarToVertexIndicesNode.cpp b/Engine/Plugins/Experimental/GeometryCollectionPlugin/Source/GeometryCollectionNodes/Private/Dataflow/GeometryCollectionVertexScalarToVertexIndicesNode.cpp
--- a/Engine/Plugins/Experimental/GeometryCollectionPlugin/Source/GeometryCollectionNodes/Private/Dataflow/GeometryCollectionVertexScalarToVertexIndicesNode.cpp
+++ b/Engine/Plugins/Experimental/GeometryCollectionPlugin/Source/GeometryCollectionNodes/Private/Dataflow/GeometryCollectionVertexScalarToVertexIndicesNode.cpp
@@ -12,12 +12,12 @@ FGeometryCollectionVertexScalarToVertexIndicesNode::FGeometryCollectionVertexSca
 {
 	RegisterInputConnection(&Collection);
 	RegisterInputConnection(&AttributeKey);
-	RegisterOutputConnection(&Indices);
+	RegisterOutputConnection(&VertexIndices);
 }
 
 void FGeometryCollectionVertexScalarToVertexIndicesNode::Evaluate(Dataflow::FContext& Context, const FDataflowOutput* Out) const
 {
-	if (Out->IsA< TArray<int32> >(&Indices))
+	if (Out->IsA< TArray<int32> >(&VertexIndices))
 	{
 		TArray<int32> IndicesOut;
 
@@ -34,7 +34,7 @@ void FGeometryCollectionVertexScalarToVertexIndicesNode::Evaluate(Dataflow::FCon
 				}
 			}
 		}
-		SetValue< TArray<int32> >(Context, MoveTemp(IndicesOut), &Indices);
+		SetValue< TArray<int32> >(Context, MoveTemp(IndicesOut), &VertexIndices);
 	}
 }
 
diff --git a/Engine/Plugins/Experimental/GeometryCollectionPlugin/Source/GeometryCollectionNodes/Public/Dataflow/GeometryCollectionVertexScalarToVertexIndicesNode.h b/Engine/Plugins/Experimental/GeometryCollectionPlugin/Source/GeometryCollectionNodes/Public/Dataflow/GeometryCollectionVertexScalarToVertexIndicesNode.h
--- a/Engine/Plugins/Experimental/GeometryCollectionPlugin/Source/GeometryCollectionNodes/Public/Dataflow/GeometryCollectionVertexScalarToVertexIndicesNode.h
+++ b/Engine/Plugins/Experimental/GeometryCollectionPlugin/Source/GeometryCollectionNodes/Public/Dataflow/GeometryCollectionVertexScalarToVertexIndicesNode.h
@@ -30,8 +30,8 @@ public:
 	float SelectionThreshold = 0.f;
 
 	/** Output list of indices */
-	UPROPERTY(Meta = (DataflowOutput))
-	TArray<int32> Indices = {};
+	UPROPERTY(Meta = (DataflowOutput, DisplayName = "VertexIndices"))
+	TArray<int32> VertexIndices = {};
 
 	FGeometryCollectionVertexScalarToVertexIndicesNode(const Dataflow::FNodeParameters& InParam, FGuid InGuid = FGuid::NewGuid());
 
