ChaosFlesh: Change point-triangle transfer scheme

diff --git a/Engine/Plugins/Experimental/ChaosFlesh/Source/ChaosFleshNodes/Private/Dataflow/ChaosFleshSetVertexTrianglePositionTargetBindingNode.cpp b/Engine/Plugins/Experimental/ChaosFlesh/Source/ChaosFleshNodes/Private/Dataflow/ChaosFleshSetVertexTrianglePositionTargetBindingNode.cpp
--- a/Engine/Plugins/Experimental/ChaosFlesh/Source/ChaosFleshNodes/Private/Dataflow/ChaosFleshSetVertexTrianglePositionTargetBindingNode.cpp
+++ b/Engine/Plugins/Experimental/ChaosFlesh/Source/ChaosFleshNodes/Private/Dataflow/ChaosFleshSetVertexTrianglePositionTargetBindingNode.cpp
@@ -2,7 +2,8 @@
 
 #include "Dataflow/ChaosFleshSetVertexTrianglePositionTargetBindingNode.h"
 
-#include "Chaos/BoundingVolumeHierarchy.h"
+#include "Chaos/HierarchicalSpatialHash.h"
+#include "Chaos/TriangleCollisionPoint.h"
 #include "GeometryCollection/Facades/CollectionMeshFacade.h"
 #include "GeometryCollection/Facades/CollectionPositionTargetFacade.h"
 #include "ChaosFlesh/TetrahedralCollection.h"
@@ -70,162 +71,86 @@ void FSetVertexTrianglePositionTargetBindingDataflowNode::Evaluate(Dataflow::FCo
 							CurrentParticleIndex += 1;
 						}
 					}
-					TArray<Chaos::TSphere<Chaos::FReal, 3>*> VertexSpherePtrs;
-					TArray<Chaos::TSphere<Chaos::FReal, 3>> VertexSpheres;
-					Chaos::FReal SphereRadius = (Chaos::FReal)0.;
+
+					//Compute max bound and determine detection radius
 					Chaos::TVec3<float> CoordMaxs(-FLT_MAX);
 					Chaos::TVec3<float> CoordMins(FLT_MAX);
-					if (FindInput(&VertexSelectionSetIn) && FindInput(&VertexSelectionSetIn)->GetConnection())
+					for (int32 i = 0; i < Vertex.Num(); i++)
 					{
-						TArray<int32> VertexSelectionSet = GetValue<TArray<int32>>(Context, &VertexSelectionSetIn);
-						TArray<Chaos::TVector<float, 3>> VertexPositions;
-						VertexPositions.SetNum(VertexSelectionSet.Num());
-						for (int32 i = 0; i < VertexSelectionSet.Num(); i++)
+						for (int32 j = 0; j < 3; j++)
 						{
-							if (VertexSelectionSet[i] > INDEX_NONE && VertexSelectionSet[i] < Vertex.Num())
+							if (Vertex[i][j] > CoordMaxs[j])
 							{
-								VertexPositions[i] = Vertex[VertexSelectionSet[i]];
+								CoordMaxs[j] = Vertex[i][j];
 							}
-						}
-						for (int32 i = 0; i < VertexPositions.Num(); i++)
-						{
-							for (int32 j = 0; j < 3; j++)
+							if (Vertex[i][j] < CoordMins[j])
 							{
-								if (VertexPositions[i][j] > CoordMaxs[j])
-								{
-									CoordMaxs[j] = VertexPositions[i][j];
-								}
-								if (VertexPositions[i][j] < CoordMins[j])
-								{
-									CoordMins[j] = VertexPositions[i][j];
-								}
+								CoordMins[j] = Vertex[i][j];
 							}
 						}
-						Chaos::TVec3<float> CoordDiff = (CoordMaxs - CoordMins) * VertexRadiusRatio;
-						SphereRadius = Chaos::FReal(FGenericPlatformMath::Min(CoordDiff[0], FGenericPlatformMath::Min(CoordDiff[1], CoordDiff[2])));
-
-						VertexSpheres.Init(Chaos::TSphere<Chaos::FReal, 3>(Chaos::TVec3<Chaos::FReal>(0), SphereRadius), VertexPositions.Num());
-						VertexSpherePtrs.SetNum(VertexPositions.Num());
-
-						for (int32 i = 0; i < VertexPositions.Num(); i++)
-						{
-							Chaos::TVec3<Chaos::FReal> SphereCenter(VertexPositions[i]);
-							Chaos::TSphere<Chaos::FReal, 3> VertexSphere(SphereCenter, SphereRadius);
-							VertexSpheres[i] = Chaos::TSphere<Chaos::FReal, 3>(SphereCenter, SphereRadius);
-							VertexSpherePtrs[i] = &VertexSpheres[i];
-						}
-						IndicesMap = VertexSelectionSet;
 					}
-					else
+					Chaos::TVec3<float> CoordDiff = (CoordMaxs - CoordMins) * VertexRadiusRatio;
+					Chaos::FReal SphereRadius = Chaos::FReal(FGenericPlatformMath::Max(CoordDiff[0], FGenericPlatformMath::Max(CoordDiff[1], CoordDiff[2])));
+					
+					if (FindInput(&VertexSelectionSetIn) && FindInput(&VertexSelectionSetIn)->GetConnection())
 					{
-						for (int32 i = 0; i < IndicesPositions.Num(); i++)
-						{
-							for (int32 j = 0; j < 3; j++)
-							{
-								if (IndicesPositions[i][j] > CoordMaxs[j])
-								{
-									CoordMaxs[j] = IndicesPositions[i][j];
-								}
-								if (IndicesPositions[i][j] < CoordMins[j])
-								{
-									CoordMins[j] = IndicesPositions[i][j];
-								}
-							}
-						}
-						Chaos::TVec3<float> CoordDiff = (CoordMaxs - CoordMins) * VertexRadiusRatio;
-						SphereRadius = Chaos::FReal(FGenericPlatformMath::Min(CoordDiff[0], FGenericPlatformMath::Min(CoordDiff[1], CoordDiff[2])));
-
-						VertexSpheres.Init(Chaos::TSphere<Chaos::FReal, 3>(Chaos::TVec3<Chaos::FReal>(0), SphereRadius), IndicesPositions.Num());
-						VertexSpherePtrs.SetNum(IndicesPositions.Num());
-
-						for (int32 i = 0; i < IndicesPositions.Num(); i++)
-						{
-							Chaos::TVec3<Chaos::FReal> SphereCenter(IndicesPositions[i]);
-							Chaos::TSphere<Chaos::FReal, 3> VertexSphere(SphereCenter, SphereRadius);
-							VertexSpheres[i] = Chaos::TSphere<Chaos::FReal, 3>(SphereCenter, SphereRadius);
-							VertexSpherePtrs[i] = &VertexSpheres[i];
-						}
+						IndicesMap = GetValue<TArray<int32>>(Context, &VertexSelectionSetIn);
 					}
 
-					Chaos::TBoundingVolumeHierarchy<
-						TArray<Chaos::TSphere<Chaos::FReal, 3>*>,
-						TArray<int32>,
-						Chaos::FReal,
-						3> VertexBVH(VertexSpherePtrs);
-
 					GeometryCollection::Facades::FPositionTargetFacade PositionTargets(InCollection);
 					PositionTargets.DefineSchema();
 
-					for (int32 i = 0; i < Indices->Num(); i++)
+					TArray<Chaos::TVec3<Chaos::FReal>> VertexTVec3;
+					VertexTVec3.SetNum(Vertex.Num());
+					for (int32 VertexIdx = 0; VertexIdx < Vertex.Num(); VertexIdx++)
 					{
-						TArray<int32> TriangleIntersections0 = VertexBVH.FindAllIntersections(Vertex[(*Indices)[i][0]]);
-						TArray<int32> TriangleIntersections1 = VertexBVH.FindAllIntersections(Vertex[(*Indices)[i][1]]);
-						TArray<int32> TriangleIntersections2 = VertexBVH.FindAllIntersections(Vertex[(*Indices)[i][2]]);
-						TriangleIntersections0.Sort();
-						TriangleIntersections1.Sort();
-						TriangleIntersections2.Sort();
-
-						TArray<int32> TriangleIntersections({});
-						for (int32 k = 0; k < TriangleIntersections0.Num(); k++)
-						{
-							if (TriangleIntersections1.Contains(TriangleIntersections0[k]) 
-								&& TriangleIntersections2.Contains(TriangleIntersections0[k]))
+						VertexTVec3[VertexIdx] = Chaos::TVec3<Chaos::FReal>(Vertex[VertexIdx]);
+					}
+					Chaos::FTriangleMesh TriangleMesh;
+					TriangleMesh.Init(IndicesArray);
+					Chaos::FTriangleMesh::TSpatialHashType<Chaos::FReal> SpatialHash;
+					TConstArrayView<Chaos::TVec3<Chaos::FReal>> ConstArrayViewVertex(VertexTVec3);
+					TriangleMesh.BuildSpatialHash(ConstArrayViewVertex, SpatialHash, SphereRadius);
+					for (int32 PointIndex: IndicesMap)
+					{
+						TArray<Chaos::TTriangleCollisionPoint<Chaos::FReal>> Result;
+						if (TriangleMesh.PointClosestTriangleQuery(SpatialHash, ConstArrayViewVertex,
+							PointIndex, VertexTVec3[PointIndex], SphereRadius / 2.f, SphereRadius / 2.f,
+							[this, &ComponentIndex, &IndicesArray](const int32 PointIndex, const int32 TriangleIndex)->bool
 							{
-								TriangleIntersections.Emplace(TriangleIntersections0[k]);
-							}
-						}
-
-						int32 TriangleIndex = ComponentIndex[(*Indices)[i][0]];
-						int32 MinIndex = -1;
-						float MinDis = SphereRadius;
-						Chaos::TVector<float, 3> ClosestBary(0.f);
-						for (int32 j = 0; j < TriangleIntersections.Num(); j++)
+								return ComponentIndex[PointIndex] != ComponentIndex[IndicesArray[TriangleIndex][0]];
+							},
+							Result))
 						{
-							if (ComponentIndex[IndicesMap[TriangleIntersections[j]]] > INDEX_NONE && TriangleIndex > INDEX_NONE && ComponentIndex[IndicesMap[TriangleIntersections[j]]] != TriangleIndex)
+							for (const Chaos::TTriangleCollisionPoint<Chaos::FReal>& CollisionPoint : Result)
 							{
-								Chaos::TVector<float, 3> Bary, TriPos0(Vertex[(*Indices)[i][0]]), TriPos1(Vertex[(*Indices)[i][1]]), TriPos2(Vertex[(*Indices)[i][2]]), ParticlePos(Vertex[IndicesMap[TriangleIntersections[j]]]);
-								Chaos::TVector<Chaos::FRealSingle, 3> ClosestPoint = Chaos::FindClosestPointAndBaryOnTriangle(TriPos0, TriPos1, TriPos2, ParticlePos, Bary);
-								Chaos::FRealSingle CurrentDistance = (Vertex[IndicesMap[TriangleIntersections[j]]] - ClosestPoint).Size();
-								if (CurrentDistance < MinDis)
+								GeometryCollection::Facades::FPositionTargetsData DataPackage;
+								DataPackage.TargetIndex.Init(PointIndex, 1);
+								DataPackage.TargetWeights.Init(1.f, 1);
+								DataPackage.SourceWeights.Init(1.f, 3);
+								DataPackage.SourceIndex.Init(-1, 3);
+								DataPackage.SourceIndex[0] = IndicesArray[CollisionPoint.Indices[1]][0];
+								DataPackage.SourceIndex[1] = IndicesArray[CollisionPoint.Indices[1]][1];
+								DataPackage.SourceIndex[2] = IndicesArray[CollisionPoint.Indices[1]][2];
+								DataPackage.SourceWeights[0] = CollisionPoint.Bary[1]; //convention: Bary[0] point, Bary[1:3] triangle
+								DataPackage.SourceWeights[1] = CollisionPoint.Bary[2];
+								DataPackage.SourceWeights[2] = CollisionPoint.Bary[3];
+								if (TManagedArray<float>* Mass = InCollection.FindAttribute<float>("Mass", FGeometryCollection::VerticesGroup))
 								{
-									MinDis = CurrentDistance;
-									MinIndex = IndicesMap[TriangleIntersections[j]];
-									ClosestBary = Bary;
+									DataPackage.Stiffness = 0.f;
+									for (int32 k = 0; k < 3; k++)
+									{
+										DataPackage.Stiffness += DataPackage.SourceWeights[k] * PositionTargetStiffness * (*Mass)[DataPackage.SourceIndex[k]];
+									}
+									DataPackage.Stiffness += DataPackage.TargetWeights[0] * PositionTargetStiffness * (*Mass)[DataPackage.TargetIndex[0]];
+									DataPackage.Stiffness /= 2.f;
 								}
-
-							}
-						}
-						if (MinIndex != -1
-							&& MinIndex != (*Indices)[i][0]
-							&& MinIndex != (*Indices)[i][1]
-							&& MinIndex != (*Indices)[i][2])
-						{
-							GeometryCollection::Facades::FPositionTargetsData DataPackage;
-							DataPackage.TargetIndex.Init(MinIndex, 1);
-							DataPackage.TargetWeights.Init(1.f, 1);
-							DataPackage.SourceWeights.Init(1.f, 3);
-							DataPackage.SourceIndex.Init(-1, 3);
-							DataPackage.SourceIndex[0] = (*Indices)[i][0];
-							DataPackage.SourceIndex[1] = (*Indices)[i][1];
-							DataPackage.SourceIndex[2] = (*Indices)[i][2];
-							DataPackage.SourceWeights[0] = ClosestBary[0];
-							DataPackage.SourceWeights[1] = ClosestBary[1];
-							DataPackage.SourceWeights[2] = ClosestBary[2];
-							if (TManagedArray<float>* Mass = InCollection.FindAttribute<float>("Mass", FGeometryCollection::VerticesGroup))
-							{
-								DataPackage.Stiffness = 0.f;
-								for (int32 k = 0; k < 3; k++)
+								else
 								{
-									DataPackage.Stiffness += DataPackage.SourceWeights[k] * PositionTargetStiffness * (*Mass)[DataPackage.SourceIndex[k]];
+									DataPackage.Stiffness = PositionTargetStiffness;
 								}
-								DataPackage.Stiffness += DataPackage.TargetWeights[0] * PositionTargetStiffness * (*Mass)[DataPackage.TargetIndex[0]];
-								DataPackage.Stiffness /= 2.f;
-							}
-							else
-							{
-								DataPackage.Stiffness = PositionTargetStiffness;
+								PositionTargets.AddPositionTarget(DataPackage);
 							}
-							PositionTargets.AddPositionTarget(DataPackage);
 						}
 					}
 				}
diff --git a/Engine/Plugins/Experimental/GeometryCollectionPlugin/Source/GeometryCollectionNodes/Private/Dataflow/GeometryCollectionTransferVertexScalarAttributeNode.cpp b/Engine/Plugins/Experimental/GeometryCollectionPlugin/Source/GeometryCollectionNodes/Private/Dataflow/GeometryCollectionTransferVertexScalarAttributeNode.cpp
--- a/Engine/Plugins/Experimental/GeometryCollectionPlugin/Source/GeometryCollectionNodes/Private/Dataflow/GeometryCollectionTransferVertexScalarAttributeNode.cpp
+++ b/Engine/Plugins/Experimental/GeometryCollectionPlugin/Source/GeometryCollectionNodes/Private/Dataflow/GeometryCollectionTransferVertexScalarAttributeNode.cpp
@@ -3,13 +3,15 @@
 #include "Dataflow/GeometryCollectionTransferVertexScalarAttributeNode.h"
 
 #include "Chaos/Triangle.h"
+#include "Chaos/TriangleMesh.h"
+#include "Chaos/HierarchicalSpatialHash.h"
+#include "Chaos/TriangleCollisionPoint.h"
 #include "Dataflow/DataflowInputOutput.h"
 #include "GeometryCollection/GeometryCollectionAlgo.h"
 #include "GeometryCollection/TransformCollection.h"
 #include "GeometryCollection/GeometryCollection.h"
 #include "GeometryCollection/ManagedArrayAccessor.h"
 
-
 #include UE_INLINE_GENERATED_CPP_BY_NAME(GeometryCollectionTransferVertexScalarAttributeNode)
 #define LOCTEXT_NAMESPACE "FGeometryCollectionTransferVertexScalarAttributeNode"
 
@@ -154,11 +156,51 @@ void FGeometryCollectionTransferVertexScalarAttributeNode::PairedGeometryTransfe
 {
 	if (const TManagedArray<float>* FloatArray = Sample.GetFloatArray(Key.Attribute, Key.Group))
 	{
+		Chaos::FReal SphereFullRadius;
+
+		if (SampleScale == EDataflowTransferNodeSampleScale::Dataflow_Transfer_Asset_Edge || SampleScale == EDataflowTransferNodeSampleScale::Dataflow_Transfer_Asset_Bound)
+		{
+			// Build component space vertices for TargetCollection
+			TArray<FVector> ComponentSpaceFullTargetVertices; // size of num vertices of the geometry entry. 
+			BuildComponentSpaceVertices(&Target.Transform.Get(), &Target.Parent.Get(), &Target.BoneMap.Get(), &Target.Vertex.Get(),
+				0, Target.Vertex.Num(), ComponentSpaceFullTargetVertices);
+
+			// Build component space vertices for TargetCollection
+			TArray<FVector> ComponentSpaceFullVertices; // size of num vertices of the geometry entry
+			BuildComponentSpaceVertices(&Sample.Transform.Get(), &Sample.Parent.Get(), &Sample.BoneMap.Get(), &Sample.Vertex.Get(),
+				0, Sample.Vertex.Num(), ComponentSpaceFullVertices);
+			if (SampleScale == EDataflowTransferNodeSampleScale::Dataflow_Transfer_Asset_Edge)
+			{
+				SphereFullRadius = Chaos::FReal(EdgeMultiplier * FMath::Max(
+					MaxEdgeLength(ComponentSpaceFullTargetVertices, Target.Indices.Get(), 0, 0, Target.Indices.Num()),
+					MaxEdgeLength(ComponentSpaceFullVertices, Sample.Indices.Get(), 0, 0, Sample.Indices.Num())));
+			}
+			else if (SampleScale == EDataflowTransferNodeSampleScale::Dataflow_Transfer_Asset_Bound)
+			{
+				Chaos::TVec3<float> CoordMaxs(-FLT_MAX);
+				Chaos::TVec3<float> CoordMins(FLT_MAX);
+				for (int32 i = 0; i < ComponentSpaceFullVertices.Num(); i++)
+				{
+					for (int32 j = 0; j < 3; j++)
+					{
+						if (ComponentSpaceFullVertices[i][j] > CoordMaxs[j])
+						{
+							CoordMaxs[j] = ComponentSpaceFullVertices[i][j];
+						}
+						if (ComponentSpaceFullVertices[i][j] < CoordMins[j])
+						{
+							CoordMins[j] = ComponentSpaceFullVertices[i][j];
+						}
+					}
+				}
+				Chaos::TVec3<float> CoordDiff = (CoordMaxs - CoordMins) * BoundMultiplier;
+				SphereFullRadius = Chaos::FReal(FGenericPlatformMath::Min(CoordDiff[0], FGenericPlatformMath::Min(CoordDiff[1], CoordDiff[2])));
+			}
+		}
 		ParallelFor(PairedGeometry.Num(), [&](int32 Pdx)
 		{
 			int32 AttributeGeometryIndex = PairedGeometry[Pdx][0];
 			int32 TargetGeometryIndex = PairedGeometry[Pdx][1];
-
 			if (ensure(0 <= AttributeGeometryIndex && AttributeGeometryIndex < Sample.VertexStart.Num()))
 			{
 				if (ensure(0 <= TargetGeometryIndex && TargetGeometryIndex < Target.VertexStart.Num()))
@@ -168,54 +210,110 @@ void FGeometryCollectionTransferVertexScalarAttributeNode::PairedGeometryTransfe
 					BuildComponentSpaceVertices(&Target.Transform.Get(), &Target.Parent.Get(), &Target.BoneMap.Get(), &Target.Vertex.Get(),
 						Target.VertexStart[TargetGeometryIndex], Target.VertexCount[TargetGeometryIndex], ComponentSpaceTargetVertices);
 
-					// Build component space vertices for TargetCollection
+					// Build component space vertices for SampleCollection
 					TArray<FVector> ComponentSpaceVertices; // size of num vertices of the geometry entry
 					BuildComponentSpaceVertices(&Sample.Transform.Get(), &Sample.Parent.Get(), &Sample.BoneMap.Get(), &Sample.Vertex.Get(),
 						Sample.VertexStart[AttributeGeometryIndex], Sample.VertexCount[AttributeGeometryIndex], ComponentSpaceVertices);
 
 					// build Sphere based BVH
-					Chaos::FReal SphereRadius = Chaos::FReal(EdgeMultiplier * FMath::Max(
-						MaxEdgeLength(ComponentSpaceTargetVertices, Target.Indices.Get(), Target.VertexStart[TargetGeometryIndex], Target.FaceStart[TargetGeometryIndex], Target.FaceCount[TargetGeometryIndex]),
-						MaxEdgeLength(ComponentSpaceVertices, Sample.Indices.Get(), Sample.VertexStart[AttributeGeometryIndex], Sample.FaceStart[AttributeGeometryIndex], Sample.FaceCount[AttributeGeometryIndex])));
-
-					TUniquePtr<UE::Private::BVH> VertexBVH(BuildParticleSphereBVH(ComponentSpaceTargetVertices, SphereRadius));
+					Chaos::FReal SphereRadius = SphereFullRadius;
+					if (SampleScale == EDataflowTransferNodeSampleScale::Dataflow_Transfer_Component_Edge)
+					{
+						SphereRadius = Chaos::FReal(EdgeMultiplier * FMath::Max(
+							MaxEdgeLength(ComponentSpaceTargetVertices, Target.Indices.Get(), Target.VertexStart[TargetGeometryIndex], Target.FaceStart[TargetGeometryIndex], Target.FaceCount[TargetGeometryIndex]),
+							MaxEdgeLength(ComponentSpaceVertices, Sample.Indices.Get(), Sample.VertexStart[AttributeGeometryIndex], Sample.FaceStart[AttributeGeometryIndex], Sample.FaceCount[AttributeGeometryIndex])));
+					}
 
 					int32 TargetVertexStartVal = Target.VertexStart[TargetGeometryIndex];
 					int32 TargetVertexCountVal = Target.VertexCount[TargetGeometryIndex];
 					int32 VertexStartVal = Sample.VertexStart[AttributeGeometryIndex];
 					int32 FaceStartVal = Sample.FaceStart[AttributeGeometryIndex];
 					int32 FaceCountVal = Sample.FaceCount[AttributeGeometryIndex];
-					for (int32 i = 0; i < FaceCountVal; i++)
+
+					if (BoundingVolumeType == EDataflowTransferNodeBoundingVolume::Dataflow_Transfer_Triangle)
 					{
-						FIntVector3 ComponentTriangle = Sample.Indices[FaceStartVal + i] - FIntVector3(VertexStartVal);
-						if (TriangleHasWeightsToTransfer(Sample.Indices[FaceStartVal + i], *FloatArray))
+						TArray<Chaos::TVec3<Chaos::FReal>> ComponentSpaceVerticesTVec3;
+						ComponentSpaceVerticesTVec3.SetNum(ComponentSpaceVertices.Num());
+						for (int32 SourceIndex = 0; SourceIndex < ComponentSpaceVerticesTVec3.Num(); SourceIndex++)
 						{
-							TArray<int32> TargetVertexIntersection({});
-							TriangleToVertexIntersections(*VertexBVH, ComponentSpaceVertices, ComponentTriangle, TargetVertexIntersection);
-
-							for (int32 j = 0; j < TargetVertexIntersection.Num(); j++)
+							ComponentSpaceVerticesTVec3[SourceIndex] = Chaos::TVec3<Chaos::FReal>(ComponentSpaceVertices[SourceIndex]);
+						}
+						TConstArrayView<Chaos::TVec3<Chaos::FReal>> ConstComponentSpaceVertices(ComponentSpaceVerticesTVec3);
+						Chaos::FTriangleMesh TriangleMesh;
+						TArray<Chaos::TVec3<int32>> SourceElements;
+						SourceElements.SetNum(FaceCountVal);
+						for (int32 ElementIndex = 0; ElementIndex < FaceCountVal; ++ElementIndex)
+						{
+							FIntVector3 Element = Sample.Indices[FaceStartVal + ElementIndex];
+							SourceElements[ElementIndex] = Chaos::TVec3<int32>(Element[0]-VertexStartVal, Element[1]-VertexStartVal, Element[2]-VertexStartVal);
+						}
+						TriangleMesh.Init(SourceElements);
+						Chaos::FTriangleMesh::TSpatialHashType<Chaos::FReal> SpatialHash;
+						TriangleMesh.BuildSpatialHash(ConstComponentSpaceVertices, SpatialHash, SphereRadius);
+						for (int32 TargetIndex = 0; TargetIndex < TargetVertexCountVal; TargetIndex++)
+						{
+							TArray<Chaos::TTriangleCollisionPoint<Chaos::FReal>> Result;
+							if (TriangleMesh.PointClosestTriangleQuery(SpatialHash, ConstComponentSpaceVertices,
+								TargetIndex, Chaos::TVec3<Chaos::FReal>(ComponentSpaceTargetVertices[TargetIndex]), SphereRadius / 2.f, SphereRadius / 2.f,
+								[](const int32 PointIndex, const int32 TriangleIndex)->bool {return true; }, Result))
+							{
+								for (const Chaos::TTriangleCollisionPoint<Chaos::FReal>& CollisionPoint : Result)
+								{
+									Chaos::FRealSingle CurrentDistance = abs(CollisionPoint.Phi);
+									float TriRadius = FalloffThreshold * MaxEdgeLength(ComponentSpaceVertices, Sample.Indices.Get(), VertexStartVal, FaceStartVal+CollisionPoint.Indices[1], 1);
+									float FalloffScale = CalculateFalloffScale(Falloff, TriRadius, CurrentDistance);
+									if (!FMath::IsNearlyZero(FalloffScale))
+									{
+										int32 TargetCandidateIndex = CollisionPoint.Indices[0]+ TargetVertexStartVal;
+										if (ensure(0 <= TargetCandidateIndex && TargetCandidateIndex < TargetFloatArray->Num()))
+										{
+											float Value = 0.f;
+											for (int32 k = 0; k < 3; k++)
+											{
+												Value += FalloffScale * (CollisionPoint.Bary[k] * (*FloatArray)[Sample.Indices[FaceStartVal+CollisionPoint.Indices[1]][k]]);
+											}
+											(*TargetFloatArray)[TargetCandidateIndex] = FMath::Max((*TargetFloatArray)[TargetCandidateIndex], Value);
+										}
+									}
+								}
+							}
+						}
+					}
+					else if (BoundingVolumeType == EDataflowTransferNodeBoundingVolume::Dataflow_Transfer_Vertex)
+					{
+						TUniquePtr<UE::Private::BVH> VertexBVH(BuildParticleSphereBVH(ComponentSpaceTargetVertices, SphereRadius));
+						for (int32 i = 0; i < FaceCountVal; i++)
+						{
+							FIntVector3 ComponentTriangle = Sample.Indices[FaceStartVal + i] - FIntVector3(VertexStartVal);
+							if (TriangleHasWeightsToTransfer(Sample.Indices[FaceStartVal + i], *FloatArray))
 							{
-								Chaos::TVector<float, 3> Bary,
-									TriPos0(ComponentSpaceVertices[ComponentTriangle[0]]),
-									TriPos1(ComponentSpaceVertices[ComponentTriangle[1]]),
-									TriPos2(ComponentSpaceVertices[ComponentTriangle[2]]),
-									ParticlePos(ComponentSpaceTargetVertices[TargetVertexIntersection[j]]);
-
-								Chaos::TVector<Chaos::FRealSingle, 3> ClosestPoint = Chaos::FindClosestPointAndBaryOnTriangle(TriPos0, TriPos1, TriPos2, ParticlePos, Bary);
-								Chaos::FRealSingle CurrentDistance = (ParticlePos - ClosestPoint).Size();
-								float TriRadius = FalloffThreshold * MaxEdgeLength(ComponentSpaceVertices, Sample.Indices.Get(), 0, i, 1);
-								float FalloffScale = CalculateFalloffScale(Falloff, TriRadius, CurrentDistance);
-								if (!FMath::IsNearlyZero(FalloffScale))
+								TArray<int32> TargetVertexIntersection({});
+								TriangleToVertexIntersections(*VertexBVH, ComponentSpaceVertices, ComponentTriangle, TargetVertexIntersection);
+
+								for (int32 j = 0; j < TargetVertexIntersection.Num(); j++)
 								{
-									int32 TargetIndex = TargetVertexIntersection[j] + TargetVertexStartVal;
-									if (ensure(0 <= TargetIndex && TargetIndex < TargetFloatArray->Num()))
+									Chaos::TVector<float, 3> Bary,
+										TriPos0(ComponentSpaceVertices[ComponentTriangle[0]]),
+										TriPos1(ComponentSpaceVertices[ComponentTriangle[1]]),
+										TriPos2(ComponentSpaceVertices[ComponentTriangle[2]]),
+										ParticlePos(ComponentSpaceTargetVertices[TargetVertexIntersection[j]]);
+
+									Chaos::TVector<Chaos::FRealSingle, 3> ClosestPoint = Chaos::FindClosestPointAndBaryOnTriangle(TriPos0, TriPos1, TriPos2, ParticlePos, Bary);
+									Chaos::FRealSingle CurrentDistance = (ParticlePos - ClosestPoint).Size();
+									float TriRadius = FalloffThreshold * MaxEdgeLength(ComponentSpaceVertices, Sample.Indices.Get(), VertexStartVal, FaceStartVal+i, 1);
+									float FalloffScale = CalculateFalloffScale(Falloff, TriRadius, CurrentDistance);
+									if (!FMath::IsNearlyZero(FalloffScale))
 									{
-										float Value = 0.f;
-										for (int32 k = 0; k < 3; k++)
+										int32 TargetIndex = TargetVertexIntersection[j] + TargetVertexStartVal;
+										if (ensure(0 <= TargetIndex && TargetIndex < TargetFloatArray->Num()))
 										{
-											Value += FalloffScale * (Bary[k] * (*FloatArray)[ComponentTriangle[k] + VertexStartVal]);
+											float Value = 0.f;
+											for (int32 k = 0; k < 3; k++)
+											{
+												Value += FalloffScale * (Bary[k] * (*FloatArray)[ComponentTriangle[k] + VertexStartVal]);
+											}
+											(*TargetFloatArray)[TargetIndex] = FMath::Max((*TargetFloatArray)[TargetIndex], Value);
 										}
-										(*TargetFloatArray)[TargetIndex] = FMath::Max((*TargetFloatArray)[TargetIndex], Value);
 									}
 								}
 							}
@@ -285,28 +383,28 @@ void FGeometryCollectionTransferVertexScalarAttributeNode::NearestVertexTransfer
 }
 
 
-float FGeometryCollectionTransferVertexScalarAttributeNode::MaxEdgeLength(TArray<FVector>& Vert, const TManagedArray<FIntVector3>& Tri, int VertexOFfset, int TriStart, int TriCount)
+float FGeometryCollectionTransferVertexScalarAttributeNode::MaxEdgeLength(TArray<FVector>& Vert, const TManagedArray<FIntVector3>& Tri, int VertexOffset, int TriStart, int TriCount)
 {
 	auto TriInRange = [](const FIntVector3& T, int Max) {
 		for (int k = 0; k < 3; k++)
 		{
-			if (ensure(T[k] < 0 || Max <= T[k]))
+			if (ensure(0 <= T[k] && T[k] < Max))
 			{
-				return false;
+				return true;
 			}
 		}
-		return true;
+		return false;
 	};
 
 	float Max = 0;
 	int TriStop = TriStart + TriCount;
 	for (int i = TriStart; i < TriStop; i++)
 	{
-		if (TriInRange(Tri[i] - FIntVector3(VertexOFfset), Vert.Num()))
+		if (TriInRange(Tri[i] - FIntVector3(VertexOffset), Vert.Num()))
 		{
-			Max = FMath::Max(Max, FVector3f(Vert[Tri[i][0] - VertexOFfset] - Vert[Tri[i][1] - VertexOFfset]).SquaredLength());
-			Max = FMath::Max(Max, FVector3f(Vert[Tri[i][0] - VertexOFfset] - Vert[Tri[i][2] - VertexOFfset]).SquaredLength());
-			Max = FMath::Max(Max, FVector3f(Vert[Tri[i][1] - VertexOFfset] - Vert[Tri[i][2] - VertexOFfset]).SquaredLength());
+			Max = FMath::Max(Max, FVector3f(Vert[Tri[i][0] - VertexOffset] - Vert[Tri[i][1] - VertexOffset]).SquaredLength());
+			Max = FMath::Max(Max, FVector3f(Vert[Tri[i][0] - VertexOffset] - Vert[Tri[i][2] - VertexOffset]).SquaredLength());
+			Max = FMath::Max(Max, FVector3f(Vert[Tri[i][1] - VertexOffset] - Vert[Tri[i][2] - VertexOffset]).SquaredLength());
 		}
 	}
 	return FMath::Sqrt(Max);
@@ -381,19 +479,19 @@ void FGeometryCollectionTransferVertexScalarAttributeNode::TriangleToVertexInter
 
 float FGeometryCollectionTransferVertexScalarAttributeNode::CalculateFalloffScale(EDataflowTransferNodeFalloff FalloffSetting, float Threshold, float Distance)
 {
-	float Demonator = 1.0;
+	float Denominator = 1.0;
 	if (Distance > Threshold && !FMath::IsNearlyZero(Threshold))
 	{
-		Demonator = Distance / Threshold;
+		Denominator = Distance / Threshold;
 	}
 	switch (FalloffSetting)
 	{
 	case EDataflowTransferNodeFalloff::Dataflow_Transfer_Linear:
-		return 1. / Demonator;
+		return 1. / Denominator;
 	case EDataflowTransferNodeFalloff::Dataflow_Transfer_Squared:
-		return 1. / FMath::Square(Demonator);
+		return 1. / FMath::Square(Denominator);
 	}
-	return Demonator;
+	return 1.0;
 }
 
 
diff --git a/Engine/Plugins/Experimental/GeometryCollectionPlugin/Source/GeometryCollectionNodes/Public/Dataflow/GeometryCollectionTransferVertexScalarAttributeNode.h b/Engine/Plugins/Experimental/GeometryCollectionPlugin/Source/GeometryCollectionNodes/Public/Dataflow/GeometryCollectionTransferVertexScalarAttributeNode.h
--- a/Engine/Plugins/Experimental/GeometryCollectionPlugin/Source/GeometryCollectionNodes/Public/Dataflow/GeometryCollectionTransferVertexScalarAttributeNode.h
+++ b/Engine/Plugins/Experimental/GeometryCollectionPlugin/Source/GeometryCollectionNodes/Public/Dataflow/GeometryCollectionTransferVertexScalarAttributeNode.h
@@ -34,6 +34,36 @@ enum class EDataflowTransferNodeFalloff : uint8
 	Dataflow_Transfer_Dataflow_Max UMETA(Hidden)
 };
 
+UENUM(BlueprintType)
+enum class EDataflowTransferNodeSampleScale : uint8
+{
+	/** Bounding volume hierarchy cell size based on max edge length of each geometry group*/
+	Dataflow_Transfer_Component_Edge  UMETA(DisplayName = "Component Max Edge"),
+
+	/** Bounding volume hierarchy cell size based on max edge length of the whole asset*/
+	Dataflow_Transfer_Asset_Edge UMETA(DisplayName = "Asset Max Edge"),
+
+	/** Bounding volume hierarchy cell size based on max length of the bounding box of the whole asset*/
+	Dataflow_Transfer_Asset_Bound UMETA(DisplayName = "Asset Max Bound"),
+
+	//~~~
+	//256th entry
+	Dataflow_Transfer_Dataflow_Max UMETA(Hidden)
+};
+
+UENUM(BlueprintType)
+enum class EDataflowTransferNodeBoundingVolume : uint8
+{
+	/** Bounding volume on vertices of the source triangle mesh*/
+	Dataflow_Transfer_Vertex  UMETA(DisplayName = "Vertex"),
+
+	/** Bounding volume on triangles of the source triangle mesh*/
+	Dataflow_Transfer_Triangle UMETA(DisplayName = "Triangle"),
+
+	//~~~
+	//256th entry
+	Dataflow_Transfer_Dataflow_Max UMETA(Hidden)
+};
 
 /**
  * Transfer float properties from a sample collection to a target collection. 
@@ -58,18 +88,31 @@ public:
 	UPROPERTY(EditAnywhere, Category = "Dataflow", Meta = (DataflowInput, DataflowOutput, DisplayName = "AttributeKey", DataflowPassthrough = "AttributeKey"))
 	FCollectionAttributeKey AttributeKey;
 
+	/* Bounding volume type for source assets[default: Triangle] */
+	UPROPERTY(EditAnywhere, Category = "Thresholds")
+	EDataflowTransferNodeBoundingVolume BoundingVolumeType = EDataflowTransferNodeBoundingVolume::Dataflow_Transfer_Triangle;
+
+	/* Bounding volume hierarchy cell size for neighboring vertices to transfer into[default: Asset] */
+	UPROPERTY(EditAnywhere, Category = "Thresholds")
+	EDataflowTransferNodeSampleScale SampleScale = EDataflowTransferNodeSampleScale::Dataflow_Transfer_Asset_Bound;
+
 	/* Falloff of sample value based on distance from sample triangle[default: Squared] */
 	UPROPERTY(EditAnywhere, Category = "Thresholds")
 	EDataflowTransferNodeFalloff Falloff = EDataflowTransferNodeFalloff::Dataflow_Transfer_Squared;
 
-	/* Threshold based on distance from sample triangle.Values sampled past the threshold will falloff.[Defualts to 1 percent of triangle size(0.01)] */
+	/* Threshold based on distance from sample triangle.Values sampled past the threshold will falloff.[Defaults to 1 percent of triangle size(0.01)] */
 	UPROPERTY(EditAnywhere, Category = "Thresholds")
-	float  FalloffThreshold = 0.01f;
+	float FalloffThreshold = 0.01f;
 
-	/* Scalar for the Bounding Volume Hierarchy(BVH) target's particle search radius. */
-	UPROPERTY(EditAnywhere, Category = "Thresholds")
-	float  EdgeMultiplier = 1.0f;
+	/* Edge multiplier for the Bounding Volume Hierarchy(BVH) target's particle search radius. */
+	UPROPERTY(EditAnywhere, Category = "Thresholds", meta = (EditCondition = "SampleScale == EDataflowTransferNodeSampleScale::Dataflow_Transfer_Asset_Edge || SampleScale == EDataflowTransferNodeSampleScale::Dataflow_Transfer_Component_Edge", EditConditionHides))
+	float EdgeMultiplier = 0.5f;
 
+	/* Max bound multiplier for the Bounding Volume Hierarchy(BVH) target's particle search radius. */
+	UPROPERTY(EditAnywhere, Category = "Thresholds", meta = (
+		EditCondition = "SampleScale == EDataflowTransferNodeSampleScale::Dataflow_Transfer_Asset_Bound", 
+		EditConditionHides))
+	float BoundMultiplier = 0.01f;
 
 	FGeometryCollectionTransferVertexScalarAttributeNode(const Dataflow::FNodeParameters& InParam, FGuid InGuid = FGuid::NewGuid())
 		: FDataflowNode(InParam, InGuid)
diff --git a/Engine/Source/Runtime/Experimental/Chaos/Private/Chaos/TriangleMesh.cpp b/Engine/Source/Runtime/Experimental/Chaos/Private/Chaos/TriangleMesh.cpp
--- a/Engine/Source/Runtime/Experimental/Chaos/Private/Chaos/TriangleMesh.cpp
+++ b/Engine/Source/Runtime/Experimental/Chaos/Private/Chaos/TriangleMesh.cpp
@@ -1763,8 +1763,8 @@ void FTriangleMesh::BuildSpatialHash(const TConstArrayView<TVec3<T>>& Points, TS
 
 	SpatialHash.Initialize(BvData, MinSpatialLodSize);
 }
-template void FTriangleMesh::BuildSpatialHash<FRealSingle>(const TConstArrayView<TVec3<FRealSingle>>& Points, TSpatialHashType<FRealSingle>& SpatialHash, const FRealSingle MinSpatialLodSize) const;
-template void FTriangleMesh::BuildSpatialHash<FRealDouble>(const TConstArrayView<TVec3<FRealDouble>>& Points, TSpatialHashType<FRealDouble>& SpatialHash, const FRealDouble MinSpatialLodSize) const;
+template CHAOS_API void FTriangleMesh::BuildSpatialHash<FRealSingle>(const TConstArrayView<TVec3<FRealSingle>>& Points, TSpatialHashType<FRealSingle>& SpatialHash, const FRealSingle MinSpatialLodSize) const;
+template CHAOS_API void FTriangleMesh::BuildSpatialHash<FRealDouble>(const TConstArrayView<TVec3<FRealDouble>>& Points, TSpatialHashType<FRealDouble>& SpatialHash, const FRealDouble MinSpatialLodSize) const;
 
 void FTriangleMesh::BuildSpatialHash(const TConstArrayView<Softs::FSolverVec3>& Points, TSpatialHashType<Softs::FSolverReal>& SpatialHash, const Softs::FPBDFlatWeightMap& PointThicknesses, int32 ThicknessMapIndexOffset, const Softs::FSolverReal MinSpatialLodSize) const
 {
@@ -1820,8 +1820,8 @@ bool FTriangleMesh::PointProximityQuery(const TSpatialHashType<T>& SpatialHash,
 	}
 	return Result.Num() > 0;
 }
-template bool FTriangleMesh::PointProximityQuery<FRealSingle>(const TSpatialHashType<FRealSingle>& SpatialHash, const TConstArrayView<TVector<FRealSingle, 3>>& Points, const int32 PointIndex, const TVector<FRealSingle, 3>& PointPosition, const FRealSingle PointThickness, const FRealSingle ThisThickness, TFunctionRef<bool(const int32 PointIndex, const int32 TriangleIndex)> BroadphaseTest, TArray<TTriangleCollisionPoint<FRealSingle>>& Result) const;
-template bool FTriangleMesh::PointProximityQuery<FRealDouble>(const TSpatialHashType<FRealDouble>& SpatialHash, const TConstArrayView<TVector<FRealDouble, 3>>& Points, const int32 PointIndex, const TVector<FRealDouble, 3>& PointPosition, const FRealDouble PointThickness, const FRealDouble ThisThickness, TFunctionRef<bool(const int32 PointIndex, const int32 TriangleIndex)> BroadphaseTest, TArray<TTriangleCollisionPoint<FRealDouble>>& Result) const;
+template CHAOS_API bool FTriangleMesh::PointProximityQuery<FRealSingle>(const TSpatialHashType<FRealSingle>& SpatialHash, const TConstArrayView<TVector<FRealSingle, 3>>& Points, const int32 PointIndex, const TVector<FRealSingle, 3>& PointPosition, const FRealSingle PointThickness, const FRealSingle ThisThickness, TFunctionRef<bool(const int32 PointIndex, const int32 TriangleIndex)> BroadphaseTest, TArray<TTriangleCollisionPoint<FRealSingle>>& Result) const;
+template CHAOS_API bool FTriangleMesh::PointProximityQuery<FRealDouble>(const TSpatialHashType<FRealDouble>& SpatialHash, const TConstArrayView<TVector<FRealDouble, 3>>& Points, const int32 PointIndex, const TVector<FRealDouble, 3>& PointPosition, const FRealDouble PointThickness, const FRealDouble ThisThickness, TFunctionRef<bool(const int32 PointIndex, const int32 TriangleIndex)> BroadphaseTest, TArray<TTriangleCollisionPoint<FRealDouble>>& Result) const;
 
 bool FTriangleMesh::PointProximityQuery(const TSpatialHashType<Softs::FSolverReal>& SpatialHash, const TConstArrayView<Softs::FSolverVec3>& Points, const int32 PointIndex, const Softs::FSolverVec3& PointPosition, const Softs::FSolverReal PointThickness, const Softs::FPBDFlatWeightMap& ThisThicknesses,
 	const Softs::FSolverReal ThisThicknessExtraMultiplier, int32 ThicknessMapIndexOffset, TFunctionRef<bool(const int32 PointIndex, const int32 TriangleIndex)> BroadphaseTest, TArray<TTriangleCollisionPoint<Softs::FSolverReal>>& Result) const
@@ -1924,11 +1924,7 @@ bool FTriangleMesh::PointClosestTriangleQuery(const TSpatialHashType<T>& Spatial
 		TVec3<T> Bary;
 		const TVec3<T> ClosestPoint = FindClosestPointAndBaryOnTriangle(A, B, C, PointPosition, Bary);
 		const T DistSq = (PointPosition - ClosestPoint).SizeSquared();
-		TVec3<T> Normal = TVec3<T>::CrossProduct(B - A, C - A).GetSafeNormal();
-		if (TVec3<T>::DotProduct(Normal, PointPosition - A) < 0) //Point is outside of triangle
-		{
-			return false;
-		}	
+		TVec3<T> Normal = TVec3<T>::CrossProduct(B - A, C - A).GetSafeNormal();	
 
 		TTriangleCollisionPoint<T> CollisionPoint;
 		CollisionPoint.ContactType = TTriangleCollisionPoint<T>::EContactType::PointFace;
@@ -1938,12 +1934,16 @@ bool FTriangleMesh::PointClosestTriangleQuery(const TSpatialHashType<T>& Spatial
 		CollisionPoint.Location = ClosestPoint;
 		CollisionPoint.Normal = Normal;
 		CollisionPoint.Phi = FMath::Sqrt(DistSq);
+		if (TVec3<T>::DotProduct(Normal, PointPosition - A) > 0) //Point is inside of triangle
+		{
+			CollisionPoint.Phi = -CollisionPoint.Phi;
+		}
 		Result.Add(CollisionPoint);
 	}
 	return Result.Num() > 0;
 }
-template bool FTriangleMesh::PointClosestTriangleQuery<FRealSingle>(const TSpatialHashType<FRealSingle>& SpatialHash, const TConstArrayView<TVector<FRealSingle, 3>>& Points, const int32 PointIndex, const TVector<FRealSingle, 3>& PointPosition, const FRealSingle PointThickness, const FRealSingle ThisThickness, TFunctionRef<bool(const int32 PointIndex, const int32 TriangleIndex)> BroadphaseTest, TArray<TTriangleCollisionPoint<FRealSingle>>& Result) const;
-template bool FTriangleMesh::PointClosestTriangleQuery<FRealDouble>(const TSpatialHashType<FRealDouble>& SpatialHash, const TConstArrayView<TVector<FRealDouble, 3>>& Points, const int32 PointIndex, const TVector<FRealDouble, 3>& PointPosition, const FRealDouble PointThickness, const FRealDouble ThisThickness, TFunctionRef<bool(const int32 PointIndex, const int32 TriangleIndex)> BroadphaseTest, TArray<TTriangleCollisionPoint<FRealDouble>>& Result) const;
+template CHAOS_API bool FTriangleMesh::PointClosestTriangleQuery<FRealSingle>(const TSpatialHashType<FRealSingle>& SpatialHash, const TConstArrayView<TVector<FRealSingle, 3>>& Points, const int32 PointIndex, const TVector<FRealSingle, 3>& PointPosition, const FRealSingle PointThickness, const FRealSingle ThisThickness, TFunctionRef<bool(const int32 PointIndex, const int32 TriangleIndex)> BroadphaseTest, TArray<TTriangleCollisionPoint<FRealSingle>>& Result) const;
+template CHAOS_API bool FTriangleMesh::PointClosestTriangleQuery<FRealDouble>(const TSpatialHashType<FRealDouble>& SpatialHash, const TConstArrayView<TVector<FRealDouble, 3>>& Points, const int32 PointIndex, const TVector<FRealDouble, 3>& PointPosition, const FRealDouble PointThickness, const FRealDouble ThisThickness, TFunctionRef<bool(const int32 PointIndex, const int32 TriangleIndex)> BroadphaseTest, TArray<TTriangleCollisionPoint<FRealDouble>>& Result) const;
 
 
 template<typename T>
diff --git a/Engine/Source/Runtime/Experimental/Chaos/Public/Chaos/Deformable/GaussSeidelWeakConstraints.h b/Engine/Source/Runtime/Experimental/Chaos/Public/Chaos/Deformable/GaussSeidelWeakConstraints.h
--- a/Engine/Source/Runtime/Experimental/Chaos/Public/Chaos/Deformable/GaussSeidelWeakConstraints.h
+++ b/Engine/Source/Runtime/Experimental/Chaos/Public/Chaos/Deformable/GaussSeidelWeakConstraints.h
@@ -495,35 +495,30 @@ namespace Chaos::Softs
 					if (TriangleMesh.PointClosestTriangleQuery(Spatial, static_cast<const TArrayView<const FSolverVec3>&>(Particles.XArray()), Index, Particles.GetX(Index), DetectRadius, DetectRadius,
 						[this, &ComponentIndex, &Elements](const int32 PointIndex, const int32 TriangleIndex)->bool
 						{
-							if (ComponentIndex[PointIndex] == ComponentIndex[Elements[TriangleIndex][0]])
-							{
-								return false;
-							}
-							return true;
+							return ComponentIndex[PointIndex] != ComponentIndex[Elements[TriangleIndex][0]];
 						},
 						Result))
 					{
 						for (const TTriangleCollisionPoint<FSolverReal>& CollisionPoint : Result)
 						{
-							const TVector<int32, 3>& Elem = Elements[CollisionPoint.Indices[1]];
-							// // NOTE: CollisionPoint.Normal has already been flipped to point toward the Point, so need to recalculate here.
-							// const TTriangle<FSolverReal> Triangle(Particles.GetX(Elem[0]), Particles.GetX(Elem[1]), Particles.GetX(Elem[2]));
-							// if ((Particles.GetX(Index) - CollisionPoint.Location).Dot(-Triangle.GetNormal()) < 0) //Is point inside boundary? Normal should point outwards
-							// {
-							const int32 IndexToWrite = ConstraintIndex.fetch_add(1);
-							Indices[IndexToWrite] = { Elem[0], Elem[1] ,Elem[2] };
-							SecondIndices[IndexToWrite] = { Index };
-							Weights[IndexToWrite] = { CollisionPoint.Bary[1], CollisionPoint.Bary[2], CollisionPoint.Bary[3] };
-							SecondWeights[IndexToWrite] = { 1.f };
-							float SpringStiffness = 0.f;
-							for (int32 k = 0; k < 3; k++)
+							if (CollisionPoint.Phi < 0)
 							{
-								SpringStiffness += Weights[IndexToWrite][k] * PositionTargetStiffness * Particles.M(Elem[k]);
+								const TVector<int32, 3>& Elem = Elements[CollisionPoint.Indices[1]];
+								const int32 IndexToWrite = ConstraintIndex.fetch_add(1);
+								Indices[IndexToWrite] = { Elem[0], Elem[1] ,Elem[2] };
+								SecondIndices[IndexToWrite] = { Index };
+								Weights[IndexToWrite] = { CollisionPoint.Bary[1], CollisionPoint.Bary[2], CollisionPoint.Bary[3] };
+								SecondWeights[IndexToWrite] = { 1.f };
+								float SpringStiffness = 0.f;
+								for (int32 k = 0; k < 3; k++)
+								{
+									SpringStiffness += Weights[IndexToWrite][k] * PositionTargetStiffness * Particles.M(Elem[k]);
+								}
+								SpringStiffness += PositionTargetStiffness * Particles.M(Index);
+								Stiffness[IndexToWrite] = SpringStiffness;
+								IsAnisotropic[IndexToWrite] = UseAnisotropicSpring;
+								Normals[IndexToWrite] = CollisionPoint.Normal;
 							}
-							SpringStiffness += PositionTargetStiffness * Particles.M(Index);
-							Stiffness[IndexToWrite] = SpringStiffness;
-							IsAnisotropic[IndexToWrite] = UseAnisotropicSpring;
-							Normals[IndexToWrite] = CollisionPoint.Normal;						
 						}
 					}
 				}
@@ -550,35 +545,30 @@ namespace Chaos::Softs
 					if (TriangleMesh.PointClosestTriangleQuery(Spatial, static_cast<const TArrayView<const FSolverVec3>&>(Particles.XArray()), Index, Particles.GetX(Index), DetectRadius, DetectRadius,
 						[this, &Elements, &ExcludeMap](const int32 PointIndex, const int32 TriangleIndex)->bool
 						{	
-							if (ExcludeMap.Find(PointIndex) && ExcludeMap[PointIndex].Contains(TriangleIndex))
-							{
-								return false;
-							}
-							return true;
+							return  !(ExcludeMap.Find(PointIndex) && ExcludeMap[PointIndex].Contains(TriangleIndex));
 						},
 						Result))
 					{
 						for (const TTriangleCollisionPoint<FSolverReal>& CollisionPoint : Result)
 						{
-							const TVector<int32, 3>& Elem = Elements[CollisionPoint.Indices[1]];
-							// // NOTE: CollisionPoint.Normal has already been flipped to point toward the Point, so need to recalculate here.
-							// const TTriangle<FSolverReal> Triangle(Particles.GetX(Elem[0]), Particles.GetX(Elem[1]), Particles.GetX(Elem[2]));
-							// if ((Particles.GetX(Index) - CollisionPoint.Location).Dot(-Triangle.GetNormal()) < 0) //Is point inside boundary? Normal should point outwards
-							// {
-							const int32 IndexToWrite = ConstraintIndex.fetch_add(1);
-							Indices[IndexToWrite] = { Elem[0], Elem[1] ,Elem[2] };
-							SecondIndices[IndexToWrite] = { Index };
-							Weights[IndexToWrite] = { CollisionPoint.Bary[1], CollisionPoint.Bary[2], CollisionPoint.Bary[3] };
-							SecondWeights[IndexToWrite] = { 1.f };
-							float SpringStiffness = 0.f;
-							for (int32 k = 0; k < 3; k++)
+							if (CollisionPoint.Phi < 0)
 							{
-								SpringStiffness += Weights[IndexToWrite][k] * PositionTargetStiffness * Particles.M(Elem[k]);
+								const TVector<int32, 3>& Elem = Elements[CollisionPoint.Indices[1]];
+								const int32 IndexToWrite = ConstraintIndex.fetch_add(1);
+								Indices[IndexToWrite] = { Elem[0], Elem[1] ,Elem[2] };
+								SecondIndices[IndexToWrite] = { Index };
+								Weights[IndexToWrite] = { CollisionPoint.Bary[1], CollisionPoint.Bary[2], CollisionPoint.Bary[3] };
+								SecondWeights[IndexToWrite] = { 1.f };
+								float SpringStiffness = 0.f;
+								for (int32 k = 0; k < 3; k++)
+								{
+									SpringStiffness += Weights[IndexToWrite][k] * PositionTargetStiffness * Particles.M(Elem[k]);
+								}
+								SpringStiffness += PositionTargetStiffness * Particles.M(Index);
+								Stiffness[IndexToWrite] = SpringStiffness;
+								IsAnisotropic[IndexToWrite] = UseAnisotropicSpring;
+								Normals[IndexToWrite] = CollisionPoint.Normal;
 							}
-							SpringStiffness += PositionTargetStiffness * Particles.M(Index);
-							Stiffness[IndexToWrite] = SpringStiffness;
-							IsAnisotropic[IndexToWrite] = UseAnisotropicSpring;
-							Normals[IndexToWrite] = CollisionPoint.Normal;						
 						}
 					}
 				}
diff --git a/Engine/Source/Runtime/Experimental/Chaos/Public/Chaos/TriangleMesh.h b/Engine/Source/Runtime/Experimental/Chaos/Public/Chaos/TriangleMesh.h
--- a/Engine/Source/Runtime/Experimental/Chaos/Public/Chaos/TriangleMesh.h
+++ b/Engine/Source/Runtime/Experimental/Chaos/Public/Chaos/TriangleMesh.h
@@ -315,17 +315,17 @@ namespace Chaos
 
 		// Hierarchy will only go down to lods as small as MinSpatialLodSize.
 		template<typename T>
-		void BuildSpatialHash(const TConstArrayView<TVec3<T>>& Points, TSpatialHashType<T>& SpatialHash, const T MinSpatialLodSize = (T)0.) const;
-		void BuildSpatialHash(const TConstArrayView<TVec3<FRealSingle>>& Points, TSpatialHashType<FRealSingle>& SpatialHash, const Softs::FPBDFlatWeightMap& PointThicknesses, int32 ThicknessMapIndexOffset, const FRealSingle MinSpatialLodSize = 0.f) const;
+		CHAOS_API void BuildSpatialHash(const TConstArrayView<TVec3<T>>& Points, TSpatialHashType<T>& SpatialHash, const T MinSpatialLodSize = (T)0.) const;
+		CHAOS_API void BuildSpatialHash(const TConstArrayView<TVec3<FRealSingle>>& Points, TSpatialHashType<FRealSingle>& SpatialHash, const Softs::FPBDFlatWeightMap& PointThicknesses, int32 ThicknessMapIndexOffset, const FRealSingle MinSpatialLodSize = 0.f) const;
 
 		template<typename T>
-		bool PointProximityQuery(const TSpatialHashType<T>& SpatialHash, const TConstArrayView<TVec3<T>>& Points, const int32 PointIndex, const TVec3<T>& PointPosition, const T PointThickness, const T ThisThickness,
+		CHAOS_API bool PointProximityQuery(const TSpatialHashType<T>& SpatialHash, const TConstArrayView<TVec3<T>>& Points, const int32 PointIndex, const TVec3<T>& PointPosition, const T PointThickness, const T ThisThickness,
 			TFunctionRef<bool(const int32 PointIndex, const int32 TriangleIndex)> BroadphaseTest, TArray<TTriangleCollisionPoint<T>>& Result) const;
-		bool PointProximityQuery(const TSpatialHashType<FRealSingle>& SpatialHash, const TConstArrayView<TVec3<FRealSingle>>& Points, const int32 PointIndex, const TVec3<FRealSingle>& PointPosition, const FRealSingle PointThickness, const Softs::FPBDFlatWeightMap& ThisThicknesses,
+		CHAOS_API bool PointProximityQuery(const TSpatialHashType<FRealSingle>& SpatialHash, const TConstArrayView<TVec3<FRealSingle>>& Points, const int32 PointIndex, const TVec3<FRealSingle>& PointPosition, const FRealSingle PointThickness, const Softs::FPBDFlatWeightMap& ThisThicknesses,
 			const FRealSingle ThisThicknessExtraMultiplier, int32 ThicknessMapIndexOffset, TFunctionRef<bool(const int32 PointIndex, const int32 TriangleIndex)> BroadphaseTest, TArray<TTriangleCollisionPoint<FRealSingle>>& Result) const;
 
 		template<typename T>
-		bool PointClosestTriangleQuery(const TSpatialHashType<T>& SpatialHash, const TConstArrayView<TVec3<T>>& Points, const int32 PointIndex, const TVec3<T>& PointPosition, const T PointThickness, const T ThisThickness,
+		CHAOS_API bool PointClosestTriangleQuery(const TSpatialHashType<T>& SpatialHash, const TConstArrayView<TVec3<T>>& Points, const int32 PointIndex, const TVec3<T>& PointPosition, const T PointThickness, const T ThisThickness,
 			TFunctionRef<bool(const int32 PointIndex, const int32 TriangleIndex)> BroadphaseTest, TArray<TTriangleCollisionPoint<T>>& Result) const;
 
 		template<typename T>
