GeometryCollection : Renamed DeprecatedNodes because of filename too long errors.

diff --git a/Engine/Plugins/Experimental/GeometryCollectionPlugin/GeometryCollectionPlugin.uplugin b/Engine/Plugins/Experimental/GeometryCollectionPlugin/GeometryCollectionPlugin.uplugin
--- a/Engine/Plugins/Experimental/GeometryCollectionPlugin/GeometryCollectionPlugin.uplugin
+++ b/Engine/Plugins/Experimental/GeometryCollectionPlugin/GeometryCollectionPlugin.uplugin
@@ -36,7 +36,7 @@
 			"LoadingPhase": "Default"
 		},
 		{
-			"Name": "GeometryCollectionDeprecatedNodes",
+			"Name": "GeometryCollectionDepNodes",
 			"Type": "Runtime",
 			"LoadingPhase": "Default"
 		}
diff --git a/Engine/Plugins/Experimental/GeometryCollectionPlugin/Source/GeometryCollectionDepNodes/GeometryCollectionDepNodes.Build.cs b/Engine/Plugins/Experimental/GeometryCollectionPlugin/Source/GeometryCollectionDepNodes/GeometryCollectionDepNodes.Build.cs
new file mode 100644
--- /dev/null
+++ b/Engine/Plugins/Experimental/GeometryCollectionPlugin/Source/GeometryCollectionDepNodes/GeometryCollectionDepNodes.Build.cs
@@ -0,0 +1,30 @@
+// Copyright Epic Games, Inc. All Rights Reserved.
+
+namespace UnrealBuildTool.Rules
+{
+	public class GeometryCollectionDepNodes : ModuleRules
+	{
+        public GeometryCollectionDepNodes(ReadOnlyTargetRules Target) : base(Target)
+		{
+			PublicDependencyModuleNames.AddRange(
+				new string[]
+				{
+					"Core",
+					"CoreUObject",
+					"ChaosCore",
+					"Chaos",
+					"DataflowCore",
+					"DataflowEngine",
+					"GeometryCollectionEngine",
+					"Engine"
+				}
+			);
+
+			PrivateDependencyModuleNames.AddRange(
+				new string[] 
+				{
+				}
+			);
+		}
+	}
+}
diff --git a/Engine/Plugins/Experimental/GeometryCollectionPlugin/Source/GeometryCollectionDepNodes/Private/Dataflow/GeometryCollectionDepNodesPlugin.cpp b/Engine/Plugins/Experimental/GeometryCollectionPlugin/Source/GeometryCollectionDepNodes/Private/Dataflow/GeometryCollectionDepNodesPlugin.cpp
new file mode 100644
--- /dev/null
+++ b/Engine/Plugins/Experimental/GeometryCollectionPlugin/Source/GeometryCollectionDepNodes/Private/Dataflow/GeometryCollectionDepNodesPlugin.cpp
@@ -0,0 +1,30 @@
+// Copyright Epic Games, Inc. All Rights Reserved.
+
+#include "Dataflow/GeometryCollectionDepNodesPlugin.h"
+
+#include "CoreMinimal.h"
+
+#include "Dataflow/DataflowNodeFactory.h"
+#include "Dataflow/SetVertexColorFromFloatArrayDepNode.h"
+#include "Dataflow/SetVertexColorFromVertexSelectionDepNode.h"
+
+
+
+#define LOCTEXT_NAMESPACE "DataflowNodes"
+
+
+void IGeometryCollectionDepNodesPlugin::StartupModule()
+{
+	DATAFLOW_NODE_REGISTER_CREATION_FACTORY(FSetVertexColorInCollectionFromVertexSelectionDataflowNode);
+	DATAFLOW_NODE_REGISTER_CREATION_FACTORY(FSetVertexColorInCollectionFromFloatArrayDataflowNode);
+}
+
+void IGeometryCollectionDepNodesPlugin::ShutdownModule()
+{
+}
+
+
+IMPLEMENT_MODULE(IGeometryCollectionDepNodesPlugin, GeometryCollectionDepNodes)
+
+
+#undef LOCTEXT_NAMESPACE
diff --git a/Engine/Plugins/Experimental/GeometryCollectionPlugin/Source/GeometryCollectionDepNodes/Private/Dataflow/SetVertexColorFromFloatArrayDepNode.cpp b/Engine/Plugins/Experimental/GeometryCollectionPlugin/Source/GeometryCollectionDepNodes/Private/Dataflow/SetVertexColorFromFloatArrayDepNode.cpp
new file mode 100644
--- /dev/null
+++ b/Engine/Plugins/Experimental/GeometryCollectionPlugin/Source/GeometryCollectionDepNodes/Private/Dataflow/SetVertexColorFromFloatArrayDepNode.cpp
@@ -0,0 +1,31 @@
+// Copyright Epic Games, Inc. All Rights Reserved.
+
+#include "Dataflow/SetVertexColorFromFloatArrayDepNode.h"
+#include "Dataflow/DataflowCore.h"
+
+
+#include UE_INLINE_GENERATED_CPP_BY_NAME(SetVertexColorFromFloatArrayDepNode)
+
+void FSetVertexColorInCollectionFromFloatArrayDataflowNode::Evaluate(Dataflow::FContext& Context, const FDataflowOutput* Out) const
+{
+	if (Out->IsA<FManagedArrayCollection>(&Collection))
+	{
+		FManagedArrayCollection InCollection = GetValue<FManagedArrayCollection>(Context, &Collection);
+		const TArray<float>& InFloatArray = GetValue<TArray<float>>(Context, &FloatArray);
+
+		const int32 NumVertices = InCollection.NumElements(FGeometryCollection::VerticesGroup);
+
+		if (InFloatArray.Num() == NumVertices)
+		{
+			if (TManagedArray<FLinearColor>* VertexColors = InCollection.FindAttribute<FLinearColor>("Color", FGeometryCollection::VerticesGroup))
+			{
+				for (int32 Idx = 0; Idx < NumVertices; ++Idx)
+				{
+					(*VertexColors)[Idx] = FLinearColor(Scale * InFloatArray[Idx], Scale * InFloatArray[Idx], Scale * InFloatArray[Idx]);
+				}
+			}
+		}
+
+		SetValue(Context, MoveTemp(InCollection), &Collection);
+	}
+}
\ No newline at end of file
diff --git a/Engine/Plugins/Experimental/GeometryCollectionPlugin/Source/GeometryCollectionDepNodes/Private/Dataflow/SetVertexColorFromVertexSelectionDepNode.cpp b/Engine/Plugins/Experimental/GeometryCollectionPlugin/Source/GeometryCollectionDepNodes/Private/Dataflow/SetVertexColorFromVertexSelectionDepNode.cpp
new file mode 100644
--- /dev/null
+++ b/Engine/Plugins/Experimental/GeometryCollectionPlugin/Source/GeometryCollectionDepNodes/Private/Dataflow/SetVertexColorFromVertexSelectionDepNode.cpp
@@ -0,0 +1,37 @@
+// Copyright Epic Games, Inc. All Rights Reserved.
+
+#include "Dataflow/SetVertexColorFromVertexSelectionDepNode.h"
+
+#include UE_INLINE_GENERATED_CPP_BY_NAME(SetVertexColorFromVertexSelectionDepNode)
+
+void FSetVertexColorInCollectionFromVertexSelectionDataflowNode::Evaluate(Dataflow::FContext& Context, const FDataflowOutput* Out) const
+{
+	if (Out->IsA<FManagedArrayCollection>(&Collection))
+	{
+		FManagedArrayCollection InCollection = GetValue<FManagedArrayCollection>(Context, &Collection);
+		const FDataflowVertexSelection& InVertexSelection = GetValue<FDataflowVertexSelection>(Context, &VertexSelection);
+
+		if (InCollection.NumElements(FGeometryCollection::VerticesGroup) == InVertexSelection.Num())
+		{
+			const int32 NumVertices = InCollection.NumElements(FGeometryCollection::VerticesGroup);
+
+			//			TManagedArray<FLinearColor>& VertexColors = InCollection.ModifyAttribute<FLinearColor>("Color", FGeometryCollection::VerticesGroup);
+			if (TManagedArray<FLinearColor>* VertexColors = InCollection.FindAttribute<FLinearColor>("Color", FGeometryCollection::VerticesGroup))
+			{
+				for (int32 Idx = 0; Idx < NumVertices; ++Idx)
+				{
+					if (InVertexSelection.IsSelected(Idx))
+					{
+						(*VertexColors)[Idx] = SelectedColor;
+					}
+					else
+					{
+						(*VertexColors)[Idx] = NonSelectedColor;
+					}
+				}
+			}
+		}
+
+		SetValue(Context, MoveTemp(InCollection), &Collection);
+	}
+}
diff --git a/Engine/Plugins/Experimental/GeometryCollectionPlugin/Source/GeometryCollectionDepNodes/Public/Dataflow/GeometryCollectionDepNodesPlugin.h b/Engine/Plugins/Experimental/GeometryCollectionPlugin/Source/GeometryCollectionDepNodes/Public/Dataflow/GeometryCollectionDepNodesPlugin.h
new file mode 100644
--- /dev/null
+++ b/Engine/Plugins/Experimental/GeometryCollectionPlugin/Source/GeometryCollectionDepNodes/Public/Dataflow/GeometryCollectionDepNodesPlugin.h
@@ -0,0 +1,18 @@
+// Copyright Epic Games, Inc. All Rights Reserved.
+
+#pragma once
+
+#include "CoreMinimal.h"
+#include "Modules/ModuleInterface.h"
+#include "Modules/ModuleManager.h"
+
+/**
+ * The public interface to this module
+ */
+class IGeometryCollectionDepNodesPlugin : public IModuleInterface
+{
+public:
+	virtual void StartupModule();
+	virtual void ShutdownModule();
+};
+
diff --git a/Engine/Plugins/Experimental/GeometryCollectionPlugin/Source/GeometryCollectionDepNodes/Public/Dataflow/SetVertexColorFromFloatArrayDepNode.h b/Engine/Plugins/Experimental/GeometryCollectionPlugin/Source/GeometryCollectionDepNodes/Public/Dataflow/SetVertexColorFromFloatArrayDepNode.h
new file mode 100644
--- /dev/null
+++ b/Engine/Plugins/Experimental/GeometryCollectionPlugin/Source/GeometryCollectionDepNodes/Public/Dataflow/SetVertexColorFromFloatArrayDepNode.h
@@ -0,0 +1,42 @@
+// Copyright Epic Games, Inc. All Rights Reserved.
+
+#pragma once 
+
+#include "CoreMinimal.h"
+#include "Dataflow/DataflowEngine.h"
+#include "GeometryCollection/ManagedArrayCollection.h"
+
+
+#include "SetVertexColorFromFloatArrayDepNode.generated.h"
+
+/** Set collection color based on input float array */
+USTRUCT(meta = (Deprecated = "5.5"))
+struct FSetVertexColorInCollectionFromFloatArrayDataflowNode : public FDataflowNode
+{
+	GENERATED_USTRUCT_BODY()
+		DATAFLOW_NODE_DEFINE_INTERNAL(FSetVertexColorInCollectionFromFloatArrayDataflowNode, "SetVertexColorInCollectionFromFloatArray", "Collection|Utilities", "")
+
+public:
+	/** Collection */
+	UPROPERTY(meta = (DataflowInput, DataflowOutput, DataflowPassthrough = "Collection", DataflowIntrinsic))
+		FManagedArrayCollection Collection;
+
+	/**  */
+	UPROPERTY(meta = (DataflowInput, DataflowIntrinsic))
+		TArray<float> FloatArray;
+
+	/**  */
+	UPROPERTY(EditAnywhere, Category = "Color")
+		float Scale = 1.f;
+
+	FSetVertexColorInCollectionFromFloatArrayDataflowNode(const Dataflow::FNodeParameters& InParam, FGuid InGuid = FGuid::NewGuid())
+		: FDataflowNode(InParam, InGuid)
+	{
+		RegisterInputConnection(&Collection);
+		RegisterInputConnection(&FloatArray);
+		RegisterOutputConnection(&Collection, &Collection);
+	}
+
+	virtual void Evaluate(Dataflow::FContext& Context, const FDataflowOutput* Out) const override;
+
+};
\ No newline at end of file
diff --git a/Engine/Plugins/Experimental/GeometryCollectionPlugin/Source/GeometryCollectionDepNodes/Public/Dataflow/SetVertexColorFromVertexSelectionDepNode.h b/Engine/Plugins/Experimental/GeometryCollectionPlugin/Source/GeometryCollectionDepNodes/Public/Dataflow/SetVertexColorFromVertexSelectionDepNode.h
new file mode 100644
--- /dev/null
+++ b/Engine/Plugins/Experimental/GeometryCollectionPlugin/Source/GeometryCollectionDepNodes/Public/Dataflow/SetVertexColorFromVertexSelectionDepNode.h
@@ -0,0 +1,46 @@
+// Copyright Epic Games, Inc. All Rights Reserved.
+
+#pragma once 
+
+#include "CoreMinimal.h"
+#include "Dataflow/DataflowEngine.h"
+#include "Dataflow/DataflowSelection.h"
+#include "GeometryCollection/ManagedArrayCollection.h"
+
+#include "SetVertexColorFromVertexSelectionDepNode.generated.h"
+
+/** Set the collections vertex color from the selection set. */
+USTRUCT(meta = (Deprecated = "5.5"))
+struct FSetVertexColorInCollectionFromVertexSelectionDataflowNode : public FDataflowNode
+{
+	GENERATED_USTRUCT_BODY()
+	DATAFLOW_NODE_DEFINE_INTERNAL(FSetVertexColorInCollectionFromVertexSelectionDataflowNode, "SetVertexColorInCollectionFromVertexSelection", "Collection|Utilities", "")
+
+public:
+	/** Collection */
+	UPROPERTY(meta = (DataflowInput, DataflowOutput, DataflowPassthrough = "Collection", DataflowIntrinsic))
+		FManagedArrayCollection Collection;
+
+	/**  */
+	UPROPERTY(meta = (DataflowInput, DisplayName = "VertexSelection", DataflowIntrinsic))
+		FDataflowVertexSelection VertexSelection;
+
+	/**  */
+	UPROPERTY(EditAnywhere, Category = "Color")
+		FLinearColor SelectedColor = FLinearColor(FColor::Yellow);
+
+	/**  */
+	UPROPERTY(EditAnywhere, Category = "Color", meta = (DisplayName = "NonSelected Color"))
+		FLinearColor NonSelectedColor = FLinearColor(FColor::Blue);
+
+	FSetVertexColorInCollectionFromVertexSelectionDataflowNode(const Dataflow::FNodeParameters& InParam, FGuid InGuid = FGuid::NewGuid())
+		: FDataflowNode(InParam, InGuid)
+	{
+		RegisterInputConnection(&Collection);
+		RegisterInputConnection(&VertexSelection);
+		RegisterOutputConnection(&Collection, &Collection);
+	}
+
+	virtual void Evaluate(Dataflow::FContext& Context, const FDataflowOutput* Out) const override;
+
+};
diff --git a/Engine/Plugins/Experimental/GeometryCollectionPlugin/Source/GeometryCollectionDeprecatedNodes/GeometryCollectionDeprecatedNodes.Build.cs b/Engine/Plugins/Experimental/GeometryCollectionPlugin/Source/GeometryCollectionDeprecatedNodes/GeometryCollectionDeprecatedNodes.Build.cs
deleted file mode 100644
diff --git a/Engine/Plugins/Experimental/GeometryCollectionPlugin/Source/GeometryCollectionDeprecatedNodes/Private/Dataflow/GeometryCollectionDeprecatedNodesPlugin.cpp b/Engine/Plugins/Experimental/GeometryCollectionPlugin/Source/GeometryCollectionDeprecatedNodes/Private/Dataflow/GeometryCollectionDeprecatedNodesPlugin.cpp
deleted file mode 100644
diff --git a/Engine/Plugins/Experimental/GeometryCollectionPlugin/Source/GeometryCollectionDeprecatedNodes/Private/Dataflow/SetVertexColorFromFloatArrayDepNode.cpp b/Engine/Plugins/Experimental/GeometryCollectionPlugin/Source/GeometryCollectionDeprecatedNodes/Private/Dataflow/SetVertexColorFromFloatArrayDepNode.cpp
deleted file mode 100644
diff --git a/Engine/Plugins/Experimental/GeometryCollectionPlugin/Source/GeometryCollectionDeprecatedNodes/Private/Dataflow/SetVertexColorFromVertexSelectionDepNode.cpp b/Engine/Plugins/Experimental/GeometryCollectionPlugin/Source/GeometryCollectionDeprecatedNodes/Private/Dataflow/SetVertexColorFromVertexSelectionDepNode.cpp
deleted file mode 100644
diff --git a/Engine/Plugins/Experimental/GeometryCollectionPlugin/Source/GeometryCollectionDeprecatedNodes/Public/Dataflow/GeometryCollectionDeprecatedNodesPlugin.h b/Engine/Plugins/Experimental/GeometryCollectionPlugin/Source/GeometryCollectionDeprecatedNodes/Public/Dataflow/GeometryCollectionDeprecatedNodesPlugin.h
deleted file mode 100644
diff --git a/Engine/Plugins/Experimental/GeometryCollectionPlugin/Source/GeometryCollectionDeprecatedNodes/Public/Dataflow/SetVertexColorFromFloatArrayDepNode.h b/Engine/Plugins/Experimental/GeometryCollectionPlugin/Source/GeometryCollectionDeprecatedNodes/Public/Dataflow/SetVertexColorFromFloatArrayDepNode.h
deleted file mode 100644
diff --git a/Engine/Plugins/Experimental/GeometryCollectionPlugin/Source/GeometryCollectionDeprecatedNodes/Public/Dataflow/SetVertexColorFromVertexSelectionDepNode.h b/Engine/Plugins/Experimental/GeometryCollectionPlugin/Source/GeometryCollectionDeprecatedNodes/Public/Dataflow/SetVertexColorFromVertexSelectionDepNode.h
deleted file mode 100644
