Dataflow : Add generic math expression node

diff --git a/Engine/Plugins/Experimental/GeometryCollectionPlugin/Source/GeometryCollectionNodes/Private/Dataflow/GeometryCollectionMathNodes.cpp b/Engine/Plugins/Experimental/GeometryCollectionPlugin/Source/GeometryCollectionNodes/Private/Dataflow/GeometryCollectionMathNodes.cpp
--- a/Engine/Plugins/Experimental/GeometryCollectionPlugin/Source/GeometryCollectionNodes/Private/Dataflow/GeometryCollectionMathNodes.cpp
+++ b/Engine/Plugins/Experimental/GeometryCollectionPlugin/Source/GeometryCollectionNodes/Private/Dataflow/GeometryCollectionMathNodes.cpp
@@ -70,6 +70,7 @@ namespace Dataflow
 		DATAFLOW_NODE_REGISTER_CREATION_FACTORY(FMathConstantsDataflowNode);
 		DATAFLOW_NODE_REGISTER_CREATION_FACTORY(FOneMinusDataflowNode);
 		DATAFLOW_NODE_REGISTER_CREATION_FACTORY(FFloatMathExpressionDataflowNode);
+		DATAFLOW_NODE_REGISTER_CREATION_FACTORY(FMathExpressionDataflowNode);
 
 		// Math
 		DATAFLOW_NODE_REGISTER_CREATION_FACTORY_NODE_COLORS_BY_CATEGORY("Math", FLinearColor(0.f, 0.4f, 0.8f), CDefaultNodeBodyTintColor);
@@ -920,4 +921,70 @@ void FFloatMathExpressionDataflowNode::Evaluate(Dataflow::FContext& Context, con
 		}
 	}
 	SetValue(Context, FloatResult, &ReturnValue);
-}
\ No newline at end of file
+}
+
+
+
+FMathExpressionDataflowNode::FMathExpressionDataflowNode(const Dataflow::FNodeParameters& InParam, FGuid InGuid)
+	: FDataflowNode(InParam, InGuid)
+{
+	FDataflowInput* InputA = RegisterInputConnection(&A);
+	FDataflowInput* InputB = RegisterInputConnection(&B);
+	FDataflowInput* InputC = RegisterInputConnection(&C);
+	FDataflowInput* InputD = RegisterInputConnection(&D);
+	FDataflowOutput* Output = RegisterOutputConnection(&ReturnValue);
+}
+
+void FMathExpressionDataflowNode::Evaluate(Dataflow::FContext& Context, const FDataflowOutput* Out) const
+{
+	check(Out->IsA(&ReturnValue));
+
+	double Result = 0.0f;
+
+	
+	FString ExpressionToEvaluate{ Expression };
+	ExpressionToEvaluate = ExpressionToEvaluate.TrimStartAndEnd();
+	if (!ExpressionToEvaluate.IsEmpty())
+	{
+		const FString VarA("{A}");
+		if (ExpressionToEvaluate.Contains(VarA))
+		{
+			const double InA = GetValue(Context, &A);
+			const FString StrA = FString::SanitizeFloat(InA);
+			ExpressionToEvaluate.ReplaceInline(*VarA, *StrA, ESearchCase::CaseSensitive);
+		}
+
+		const FString VarB("{B}");
+		if (ExpressionToEvaluate.Contains(VarB))
+		{
+			const double InB = GetValue(Context, &B);
+			const FString StrB = FString::SanitizeFloat(InB);
+			ExpressionToEvaluate.ReplaceInline(*VarB, *StrB, ESearchCase::CaseSensitive);
+		}
+
+		const FString VarC("{C}");
+		if (ExpressionToEvaluate.Contains(VarC))
+		{
+			const double InC = GetValue(Context, &C);
+			const FString StrC = FString::SanitizeFloat(InC);
+			ExpressionToEvaluate.ReplaceInline(*VarC, *StrC, ESearchCase::CaseSensitive);
+		}
+
+		const FString VarD("{D}");
+		if (ExpressionToEvaluate.Contains(VarD))
+		{
+			const double InD = GetValue(Context, &D);
+			const FString StrD = FString::SanitizeFloat(InD);
+			ExpressionToEvaluate.ReplaceInline(*VarD, *StrD, ESearchCase::CaseSensitive);
+		}
+
+		FBasicMathExpressionEvaluator Evaluator;
+		TValueOrError<double, FExpressionError> EvalResult = Evaluator.Evaluate(*ExpressionToEvaluate);
+		if (EvalResult.IsValid())
+		{
+			Result = FMath::Clamp((double)EvalResult.GetValue(), TNumericLimits<double>::Lowest(), TNumericLimits<double>::Max());
+		}
+	}
+	SetValue(Context, Result, &ReturnValue);
+}
+
diff --git a/Engine/Plugins/Experimental/GeometryCollectionPlugin/Source/GeometryCollectionNodes/Public/Dataflow/GeometryCollectionMathNodes.h b/Engine/Plugins/Experimental/GeometryCollectionPlugin/Source/GeometryCollectionNodes/Public/Dataflow/GeometryCollectionMathNodes.h
--- a/Engine/Plugins/Experimental/GeometryCollectionPlugin/Source/GeometryCollectionNodes/Public/Dataflow/GeometryCollectionMathNodes.h
+++ b/Engine/Plugins/Experimental/GeometryCollectionPlugin/Source/GeometryCollectionNodes/Public/Dataflow/GeometryCollectionMathNodes.h
@@ -1780,6 +1780,41 @@ public:
 	virtual void Evaluate(Dataflow::FContext& Context, const FDataflowOutput* Out) const override;
 };
 
+/**
+*
+* Expression node for floats
+*
+*/
+USTRUCT()
+struct FMathExpressionDataflowNode : public FDataflowNode
+{
+	GENERATED_USTRUCT_BODY()
+	DATAFLOW_NODE_DEFINE_INTERNAL(FMathExpressionDataflowNode, "MathExpression", "Math", "")
+
+public:
+	UPROPERTY(EditAnywhere, Category = Inputs, meta = (DataflowInput));
+	FDataflowNumericTypes A;
+
+	UPROPERTY(EditAnywhere, Category = Inputs, meta = (DataflowInput));
+	FDataflowNumericTypes B;
+
+	UPROPERTY(EditAnywhere, Category = Inputs, meta = (DataflowInput));
+	FDataflowNumericTypes C;
+
+	UPROPERTY(EditAnywhere, Category = Inputs, meta = (DataflowInput));
+	FDataflowNumericTypes D;
+
+	UPROPERTY(EditAnywhere, Category = Expression);
+	FString Expression;
+
+	UPROPERTY(meta = (DataflowOutput))
+	FDataflowNumericTypes ReturnValue;
+
+	FMathExpressionDataflowNode(const Dataflow::FNodeParameters& InParam, FGuid InGuid = FGuid::NewGuid());
+
+	virtual void Evaluate(Dataflow::FContext& Context, const FDataflowOutput* Out) const override;
+};
+
 namespace Dataflow
 {
 	void GeometryCollectionMathNodes();
