Move the creation of default size spceific data in the constructor of the geometry collection object

diff --git a/Engine/Plugins/Experimental/ChaosEditor/Source/FractureEditor/Private/FractureToolGenerators.cpp b/Engine/Plugins/Experimental/ChaosEditor/Source/FractureEditor/Private/FractureToolGenerators.cpp
--- a/Engine/Plugins/Experimental/ChaosEditor/Source/FractureEditor/Private/FractureToolGenerators.cpp
+++ b/Engine/Plugins/Experimental/ChaosEditor/Source/FractureEditor/Private/FractureToolGenerators.cpp
@@ -931,7 +931,6 @@ class AGeometryCollectionActor* UFractureToolGenerateAsset::CreateNewGeometryAct
 
 	UPackage* Package = CreatePackage(*UniquePackageName);
 	UGeometryCollection* InGeometryCollection = static_cast<UGeometryCollection*>(NewObject<UGeometryCollection>(Package, UGeometryCollection::StaticClass(), FName(*UniqueAssetName), RF_Transactional | RF_Public | RF_Standalone));
-	if(!InGeometryCollection->SizeSpecificData.Num()) InGeometryCollection->SizeSpecificData.Add(FGeometryCollectionSizeSpecificData());
 
 	// Create the new Geometry Collection actor
 	AGeometryCollectionActor* NewActor = Cast<AGeometryCollectionActor>(AddActor(GetSelectedLevel(), AGeometryCollectionActor::StaticClass()));
diff --git a/Engine/Plugins/Experimental/GeometryCollectionPlugin/Source/GeometryCollectionEditor/Private/GeometryCollection/GeometryCollectionFactory.cpp b/Engine/Plugins/Experimental/GeometryCollectionPlugin/Source/GeometryCollectionEditor/Private/GeometryCollection/GeometryCollectionFactory.cpp
--- a/Engine/Plugins/Experimental/GeometryCollectionPlugin/Source/GeometryCollectionEditor/Private/GeometryCollection/GeometryCollectionFactory.cpp
+++ b/Engine/Plugins/Experimental/GeometryCollectionPlugin/Source/GeometryCollectionEditor/Private/GeometryCollection/GeometryCollectionFactory.cpp
@@ -33,9 +33,7 @@ UGeometryCollectionFactory::UGeometryCollectionFactory(const FObjectInitializer&
 
 UGeometryCollection* UGeometryCollectionFactory::StaticFactoryCreateNew(UClass* Class, UObject* InParent, FName Name, EObjectFlags Flags, UObject* Context, FFeedbackContext* Warn)
 {
-	UGeometryCollection* NewGeometryCollection = static_cast<UGeometryCollection*>(NewObject<UGeometryCollection>(InParent, Class, Name, Flags | RF_Transactional | RF_Public | RF_Standalone));
-	if (!NewGeometryCollection->SizeSpecificData.Num()) NewGeometryCollection->SizeSpecificData.Add(FGeometryCollectionSizeSpecificData());
-	return NewGeometryCollection;
+	return static_cast<UGeometryCollection*>(NewObject<UGeometryCollection>(InParent, Class, Name, Flags | RF_Transactional | RF_Public | RF_Standalone));
 }
 
 void AppendActorComponentsRecursive(
diff --git a/Engine/Source/Runtime/Experimental/GeometryCollectionEngine/Private/GeometryCollection/GeometryCollectionObject.cpp b/Engine/Source/Runtime/Experimental/GeometryCollectionEngine/Private/GeometryCollection/GeometryCollectionObject.cpp
--- a/Engine/Source/Runtime/Experimental/GeometryCollectionEngine/Private/GeometryCollection/GeometryCollectionObject.cpp
+++ b/Engine/Source/Runtime/Experimental/GeometryCollectionEngine/Private/GeometryCollection/GeometryCollectionObject.cpp
@@ -144,6 +144,9 @@ UGeometryCollection::UGeometryCollection(const FObjectInitializer& ObjectInitial
 	bStripOnCook = GeometryCollectionAssetForceStripOnCook;
 #endif
 	PhysicsMaterial = GEngine? GEngine->DefaultPhysMaterial: nullptr;
+
+	// make sure we have at least one size specific entry
+	SizeSpecificData.AddDefaulted();
 }
 
 FGeometryCollectionLevelSetData::FGeometryCollectionLevelSetData()
