Geometry collection / Dataflow : properly support material instances

diff --git a/Engine/Plugins/Experimental/GeometryCollectionPlugin/Source/GeometryCollectionNodes/Private/Dataflow/GeometryCollectionAssetNodes.cpp b/Engine/Plugins/Experimental/GeometryCollectionPlugin/Source/GeometryCollectionNodes/Private/Dataflow/GeometryCollectionAssetNodes.cpp
--- a/Engine/Plugins/Experimental/GeometryCollectionPlugin/Source/GeometryCollectionNodes/Private/Dataflow/GeometryCollectionAssetNodes.cpp
+++ b/Engine/Plugins/Experimental/GeometryCollectionPlugin/Source/GeometryCollectionNodes/Private/Dataflow/GeometryCollectionAssetNodes.cpp
@@ -10,6 +10,7 @@
 #include "GeometryCollection/GeometryCollectionClusteringUtility.h"
 #include "GeometryCollection/Facades/CollectionHierarchyFacade.h"
 #include "GeometryCollection/Facades/CollectionInstancedMeshFacade.h"
+#include "Materials/MaterialInterface.h"
 #include "PreviewScene.h"
 
 namespace Dataflow
@@ -39,7 +40,9 @@ FGeometryCollectionTerminalDataflowNode::FGeometryCollectionTerminalDataflowNode
 	RegisterInputConnection(&Collection);
 	RegisterOutputConnection(&Collection, &Collection);
 	RegisterInputConnection(&Materials);
+	RegisterInputConnection(&MaterialInstances);
 	RegisterOutputConnection(&Materials, &Materials);
+	RegisterOutputConnection(&MaterialInstances, &MaterialInstances);
 	RegisterInputConnection(&InstancedMeshes);
 	RegisterOutputConnection(&InstancedMeshes, &InstancedMeshes);
 }
@@ -50,6 +53,7 @@ void FGeometryCollectionTerminalDataflowNode::SetAssetValue(TObjectPtr<UObject>
 {
 	using FGeometryCollectionPtr = TSharedPtr<FGeometryCollection, ESPMode::ThreadSafe>;
 	using FMaterialArray = TArray<TObjectPtr<UMaterial>>;
+	using FMaterialInstanceArray = TArray<TObjectPtr<UMaterialInterface>>;
 	using FInstancedMeshesArray = TArray<FGeometryCollectionAutoInstanceMesh>;
 
 	if (UGeometryCollection* CollectionAsset = Cast<UGeometryCollection>(Asset.Get()))
@@ -58,6 +62,7 @@ void FGeometryCollectionTerminalDataflowNode::SetAssetValue(TObjectPtr<UObject>
 		{
 			const FManagedArrayCollection& InCollection = GetValue(Context, &Collection);
 			const FMaterialArray& InMaterials = GetValue(Context, &Materials);
+			const FMaterialInstanceArray& InMaterialInstances = GetValue(Context, &MaterialInstances);
 			const FInstancedMeshesArray& InInstancedMeshes = GetValue(Context, &InstancedMeshes);
 
 			if (InCollection.NumElements(FGeometryCollection::TransformGroup) == 0)
@@ -67,7 +72,14 @@ void FGeometryCollectionTerminalDataflowNode::SetAssetValue(TObjectPtr<UObject>
 			}
 
 			const bool bHasInternalMaterial = false; // with data flow there's no assumption of internal materials
-			CollectionAsset->ResetFrom(InCollection, InMaterials, false);
+			if (InMaterialInstances.Num() > 0)
+			{
+				CollectionAsset->ResetFrom(InCollection, InMaterialInstances, false);
+			}
+			else
+			{
+				CollectionAsset->ResetFrom(InCollection, InMaterials, false);
+			}
 			CollectionAsset->SetAutoInstanceMeshes(InInstancedMeshes);
 
 #if WITH_EDITOR
@@ -82,16 +94,11 @@ void FGeometryCollectionTerminalDataflowNode::SetAssetValue(TObjectPtr<UObject>
 
 void FGeometryCollectionTerminalDataflowNode::Evaluate(Dataflow::FContext& Context) const
 {
-	using FMaterialArray = TArray<TObjectPtr<UMaterial>>;
-	using FInstancedMeshesArray = TArray<FGeometryCollectionAutoInstanceMesh>;
-
-	const FManagedArrayCollection& InCollection = GetValue(Context, &Collection);
-	const FMaterialArray& InMaterials = GetValue(Context, &Materials);
-	const FInstancedMeshesArray& InInstancedMeshes = GetValue(Context, &InstancedMeshes);
-
-	SetValue(Context, InCollection, &Collection);
-	SetValue(Context, InMaterials, &Materials);
-	SetValue(Context, InInstancedMeshes, &InstancedMeshes);
+	// simply forward all inputs to corresponding outputs
+	ForwardInput(Context, &Collection, &Collection);
+	ForwardInput(Context, &Materials, &Materials);
+	ForwardInput(Context, &MaterialInstances, &MaterialInstances);
+	ForwardInput(Context, &InstancedMeshes, &InstancedMeshes);
 }
 
 // ===========================================================================================================================
@@ -150,6 +157,7 @@ FCreateGeometryCollectionFromSourcesDataflowNode::FCreateGeometryCollectionFromS
 	RegisterInputConnection(&Sources);
 	RegisterOutputConnection(&Collection);
 	RegisterOutputConnection(&Materials);
+	RegisterOutputConnection(&MaterialInstances);
 	RegisterOutputConnection(&InstancedMeshes);
 }
 
@@ -160,7 +168,7 @@ void FCreateGeometryCollectionFromSourcesDataflowNode::Evaluate(Dataflow::FConte
 	const TArray<FGeometryCollectionSource>& InSources = GetValue(Context, &Sources);
 
 	FGeometryCollection OutCollection;
-	TArray<TObjectPtr<UMaterial>> OutMaterials;
+	TArray<TObjectPtr<UMaterialInterface>> OutMaterialInstances;
 	TArray<FGeometryCollectionAutoInstanceMesh> OutInstancedMeshes;
 
 	// make sure we have an attribute for instanced meshes
@@ -174,7 +182,7 @@ void FCreateGeometryCollectionFromSourcesDataflowNode::Evaluate(Dataflow::FConte
 		const int32 NumTransformsBeforeAppending = OutCollection.NumElements(FGeometryCollection::TransformGroup);
 
 		// todo: change AppendGeometryCollectionSource to take a FManagedArrayCollection so we could move the collection when assigning it to the output
-		FGeometryCollectionEngineConversion::AppendGeometryCollectionSource(Source, OutCollection, MutableView(OutMaterials), bReindexMaterialsInLoop);
+		FGeometryCollectionEngineConversion::AppendGeometryCollectionSource(Source, OutCollection, MutableView(OutMaterialInstances), bReindexMaterialsInLoop);
 
 		// todo(chaos) if the source is a geometry collection this will not work properly 
 		FGeometryCollectionAutoInstanceMesh InstancedMesh;
@@ -221,9 +229,13 @@ void FCreateGeometryCollectionFromSourcesDataflowNode::Evaluate(Dataflow::FConte
 	Chaos::Facades::FCollectionHierarchyFacade HierarchyFacade(OutCollection);
 	HierarchyFacade.GenerateLevelAttribute();
 
+	TArray<TObjectPtr<UMaterial>> OutMaterials;
+	FGeometryCollectionEngineConversion::GetMaterialsFromInstances(OutMaterialInstances, OutMaterials);
+
 	// we have to make a copy since we have generated a FGeometryCollection which is inherited from FManagedArrayCollection
 	SetValue(Context, static_cast<const FManagedArrayCollection&>(OutCollection), &Collection);
 	SetValue(Context, MoveTemp(OutMaterials), &Materials);
+	SetValue(Context, MoveTemp(OutMaterialInstances), &MaterialInstances);
 	SetValue(Context, MoveTemp(OutInstancedMeshes), &InstancedMeshes);
 }
 
@@ -234,6 +246,7 @@ FGeometryCollectionToCollectionDataflowNode::FGeometryCollectionToCollectionData
 {
 	RegisterOutputConnection(&Collection);
 	RegisterOutputConnection(&Materials);
+	RegisterOutputConnection(&MaterialInstances);
 	RegisterOutputConnection(&InstancedMeshes);
 }
 
@@ -242,17 +255,21 @@ void FGeometryCollectionToCollectionDataflowNode::Evaluate(Dataflow::FContext& C
 	ensure(Out->IsA(&Collection) || Out->IsA(&Materials) || Out->IsA(&InstancedMeshes));
 
 	FManagedArrayCollection OutCollection;
-	TArray<TObjectPtr<UMaterial>> OutMaterials;
+	TArray<TObjectPtr<UMaterialInterface>> OutMaterialInstances;
 	TArray<FGeometryCollectionAutoInstanceMesh> OutInstancedMeshes;
 
 	if (GeometryCollection)
 	{
-		FGeometryCollectionEngineConversion::ConvertGeometryCollectionToGeometryCollection(GeometryCollection, OutCollection, OutMaterials, OutInstancedMeshes);
+		FGeometryCollectionEngineConversion::ConvertGeometryCollectionToGeometryCollection(GeometryCollection, OutCollection, OutMaterialInstances, OutInstancedMeshes);
 	}
 
+	TArray<TObjectPtr<UMaterial>> OutMaterials;
+	FGeometryCollectionEngineConversion::GetMaterialsFromInstances(OutMaterialInstances, OutMaterials);
+
 	// Set Outputs
 	SetValue(Context, MoveTemp(OutCollection), &Collection);
 	SetValue(Context, MoveTemp(OutMaterials), &Materials);
+	SetValue(Context, MoveTemp(OutMaterialInstances), &MaterialInstances);
 	SetValue(Context, MoveTemp(OutInstancedMeshes), &InstancedMeshes);
 }
 
@@ -263,6 +280,7 @@ FBlueprintToCollectionDataflowNode::FBlueprintToCollectionDataflowNode(const Dat
 {
 	RegisterOutputConnection(&Collection);
 	RegisterOutputConnection(&Materials);
+	RegisterOutputConnection(&MaterialInstances);
 	RegisterOutputConnection(&InstancedMeshes);
 }
 
@@ -271,7 +289,7 @@ void FBlueprintToCollectionDataflowNode::Evaluate(Dataflow::FContext& Context, c
 	ensure(Out->IsA(&Collection) || Out->IsA(&Materials) || Out->IsA(&InstancedMeshes));
 
 	FManagedArrayCollection OutCollection;
-	TArray<TObjectPtr<UMaterial>> OutMaterials;
+	TArray<TObjectPtr<UMaterialInterface>> OutMaterialInstances;
 	TArray<FGeometryCollectionAutoInstanceMesh> OutInstancedMeshes;
 
 	if (Blueprint)
@@ -288,15 +306,19 @@ void FBlueprintToCollectionDataflowNode::Evaluate(Dataflow::FContext& Context, c
 				if (AActor* PreviewActor = PreviewWorld->SpawnActor(Blueprint->GeneratedClass, nullptr, SpawnInfo))
 				{
 					FGeometryCollectionEngineConversion::FSkeletalMeshToCollectionConversionParameters ConversionParameters;
-					FGeometryCollectionEngineConversion::ConvertActorToGeometryCollection(PreviewActor, OutCollection, OutMaterials, OutInstancedMeshes, ConversionParameters, bSplitComponents);
+					FGeometryCollectionEngineConversion::ConvertActorToGeometryCollection(PreviewActor, OutCollection, OutMaterialInstances, OutInstancedMeshes, ConversionParameters, bSplitComponents);
 				}
 			}
 		}
 	}
 
+	TArray<TObjectPtr<UMaterial>> OutMaterials;
+	FGeometryCollectionEngineConversion::GetMaterialsFromInstances(OutMaterialInstances, OutMaterials);
+
 	// Set Outputs
 	SetValue(Context, MoveTemp(OutCollection), &Collection);
 	SetValue(Context, MoveTemp(OutMaterials), &Materials);
+	SetValue(Context, MoveTemp(OutMaterialInstances), &MaterialInstances);
 	SetValue(Context, MoveTemp(OutInstancedMeshes), &InstancedMeshes);
 }
 
diff --git a/Engine/Plugins/Experimental/GeometryCollectionPlugin/Source/GeometryCollectionNodes/Private/Dataflow/GeometryCollectionStaticMeshToCollectionNode.cpp b/Engine/Plugins/Experimental/GeometryCollectionPlugin/Source/GeometryCollectionNodes/Private/Dataflow/GeometryCollectionStaticMeshToCollectionNode.cpp
--- a/Engine/Plugins/Experimental/GeometryCollectionPlugin/Source/GeometryCollectionNodes/Private/Dataflow/GeometryCollectionStaticMeshToCollectionNode.cpp
+++ b/Engine/Plugins/Experimental/GeometryCollectionPlugin/Source/GeometryCollectionNodes/Private/Dataflow/GeometryCollectionStaticMeshToCollectionNode.cpp
@@ -19,6 +19,7 @@ FStaticMeshToCollectionDataflowNode::FStaticMeshToCollectionDataflowNode(const D
 	RegisterInputConnection(&StaticMesh);
 	RegisterOutputConnection(&Collection);
 	RegisterOutputConnection(&Materials);
+	RegisterOutputConnection(&MaterialInstances);
 	RegisterOutputConnection(&InstancedMeshes);
 }
 
@@ -27,18 +28,22 @@ void FStaticMeshToCollectionDataflowNode::Evaluate(Dataflow::FContext& Context,
 	ensure(Out->IsA(&Collection) || Out->IsA(&Materials) || Out->IsA(&InstancedMeshes));
 
 	FManagedArrayCollection OutCollection;
-	TArray<TObjectPtr<UMaterial>> OutMaterials;
+	TArray<TObjectPtr<UMaterialInterface>> OutMaterialInstances;
 	TArray<FGeometryCollectionAutoInstanceMesh> OutInstancedMeshes;
 
 	TObjectPtr<UStaticMesh> StaticMeshVal = GetValue(Context, &StaticMesh, StaticMesh);
 	if (StaticMeshVal)
 	{
-		FGeometryCollectionEngineConversion::ConvertStaticMeshToGeometryCollection(StaticMeshVal, OutCollection, OutMaterials, OutInstancedMeshes, bSetInternalFromMaterialIndex, bSplitComponents);
+		FGeometryCollectionEngineConversion::ConvertStaticMeshToGeometryCollection(StaticMeshVal, OutCollection, OutMaterialInstances, OutInstancedMeshes, bSetInternalFromMaterialIndex, bSplitComponents);
 	}
 
+	TArray<TObjectPtr<UMaterial>> OutMaterials;
+	FGeometryCollectionEngineConversion::GetMaterialsFromInstances(OutMaterialInstances, OutMaterials);
+
 	// Set Outputs
 	SetValue(Context, MoveTemp(OutCollection), &Collection);
 	SetValue(Context, MoveTemp(OutMaterials), &Materials);
+	SetValue(Context, MoveTemp(OutMaterialInstances), &MaterialInstances);
 	SetValue(Context, MoveTemp(OutInstancedMeshes), &InstancedMeshes);
 }
 
diff --git a/Engine/Plugins/Experimental/GeometryCollectionPlugin/Source/GeometryCollectionNodes/Public/Dataflow/GeometryCollectionAssetNodes.h b/Engine/Plugins/Experimental/GeometryCollectionPlugin/Source/GeometryCollectionNodes/Public/Dataflow/GeometryCollectionAssetNodes.h
--- a/Engine/Plugins/Experimental/GeometryCollectionPlugin/Source/GeometryCollectionNodes/Public/Dataflow/GeometryCollectionAssetNodes.h
+++ b/Engine/Plugins/Experimental/GeometryCollectionPlugin/Source/GeometryCollectionNodes/Public/Dataflow/GeometryCollectionAssetNodes.h
@@ -29,6 +29,9 @@ public:
 	UPROPERTY(meta = (DataflowInput, DataflowOutput, DataflowPassthrough = "Materials", DisplayName = "Materials"))
 	TArray<TObjectPtr<UMaterial>> Materials;
 
+	UPROPERTY(meta = (DataflowInput, DataflowOutput, DataflowPassthrough = "MaterialInstances", DisplayName = "MaterialInstances"))
+	TArray<TObjectPtr<UMaterialInterface>> MaterialInstances;
+
 	/** array of instanced meshes*/
 	UPROPERTY(meta = (DataflowInput, DataflowOutput, DataflowPassthrough = "InstancedMeshes", DisplayName = "InstancedMeshes"))
 	TArray<FGeometryCollectionAutoInstanceMesh> InstancedMeshes;
@@ -109,6 +112,10 @@ public:
 	UPROPERTY(meta = (DataflowOutput, DisplayName = "Materials"))
 	TArray<TObjectPtr<UMaterial>> Materials;
 
+	/** Materials array to use for this asset */
+	UPROPERTY(meta = (DataflowOutput, DisplayName = "MaterialInstances"))
+	TArray<TObjectPtr<UMaterialInterface>> MaterialInstances;
+
 	/** array of instanced meshes*/
 	UPROPERTY(meta = (DataflowOutput, DisplayName = "InstancedMeshes"))
 	TArray<FGeometryCollectionAutoInstanceMesh> InstancedMeshes;
@@ -138,6 +145,10 @@ public:
 	UPROPERTY(meta = (DataflowOutput, DisplayName = "Materials"))
 	TArray<TObjectPtr<UMaterial>> Materials;
 
+	/** Material instances array from the static mesh */
+	UPROPERTY(meta = (DataflowOutput, DisplayName = "MaterialInstances"))
+	TArray<TObjectPtr<UMaterialInterface>> MaterialInstances;
+
 	/** Array of instanced meshes*/
 	UPROPERTY(meta = (DataflowOutput, DisplayName = "InstancedMeshes"))
 	TArray<FGeometryCollectionAutoInstanceMesh> InstancedMeshes;
@@ -174,6 +185,10 @@ public:
  	UPROPERTY(meta = (DataflowOutput, DisplayName = "Materials"))
  	TArray<TObjectPtr<UMaterial>> Materials;
  
+	/** Material instances array from the static mesh */
+	UPROPERTY(meta = (DataflowOutput, DisplayName = "MaterialInstances"))
+	TArray<TObjectPtr<UMaterialInterface>> MaterialInstances;
+
  	/** Array of instanced meshes*/
  	UPROPERTY(meta = (DataflowOutput, DisplayName = "InstancedMeshes"))
  	TArray<FGeometryCollectionAutoInstanceMesh> InstancedMeshes;
diff --git a/Engine/Plugins/Experimental/GeometryCollectionPlugin/Source/GeometryCollectionNodes/Public/Dataflow/GeometryCollectionStaticMeshToCollectionNode.h b/Engine/Plugins/Experimental/GeometryCollectionPlugin/Source/GeometryCollectionNodes/Public/Dataflow/GeometryCollectionStaticMeshToCollectionNode.h
--- a/Engine/Plugins/Experimental/GeometryCollectionPlugin/Source/GeometryCollectionNodes/Public/Dataflow/GeometryCollectionStaticMeshToCollectionNode.h
+++ b/Engine/Plugins/Experimental/GeometryCollectionPlugin/Source/GeometryCollectionNodes/Public/Dataflow/GeometryCollectionStaticMeshToCollectionNode.h
@@ -43,6 +43,10 @@ public:
 	UPROPERTY(meta = (DataflowOutput, DisplayName = "Materials"))
 	TArray<TObjectPtr<UMaterial>> Materials;
 
+	/** Material instances array from the static mesh */
+	UPROPERTY(meta = (DataflowOutput, DisplayName = "MaterialInstances"))
+	TArray<TObjectPtr<UMaterialInterface>> MaterialInstances;
+
 	/** Array of instanced meshes*/
 	UPROPERTY(meta = (DataflowOutput, DisplayName = "InstancedMeshes"))
 	TArray<FGeometryCollectionAutoInstanceMesh> InstancedMeshes;
diff --git a/Engine/Source/Runtime/Experimental/GeometryCollectionEngine/Private/GeometryCollection/GeometryCollectionEngineConversion.cpp b/Engine/Source/Runtime/Experimental/GeometryCollectionEngine/Private/GeometryCollection/GeometryCollectionEngineConversion.cpp
--- a/Engine/Source/Runtime/Experimental/GeometryCollectionEngine/Private/GeometryCollection/GeometryCollectionEngineConversion.cpp
+++ b/Engine/Source/Runtime/Experimental/GeometryCollectionEngine/Private/GeometryCollection/GeometryCollectionEngineConversion.cpp
@@ -1488,12 +1488,24 @@ int32 FGeometryCollectionEngineConversion::AppendSkeletalMeshMaterials(const USk
 
 void FGeometryCollectionEngineConversion::AppendGeometryCollectionSource(const FGeometryCollectionSource& GeometryCollectionSource, FGeometryCollection& GeometryCollectionInOut, TArray<UMaterial*>& MaterialsInOut, bool ReindexMaterials)
 {
-	if (const UObject* SourceObject = GeometryCollectionSource.SourceGeometryObject.TryLoad())
+	const int32 StartMaterialIndex = MaterialsInOut.Num();
+	if (AppendGeometryCollectionSourceNoMaterial(GeometryCollectionSource, GeometryCollectionInOut, StartMaterialIndex, ReindexMaterials))
 	{
-		const int32 StartMaterialIndex = MaterialsInOut.Num();
-
 		MaterialsInOut.Append(GeometryCollectionSource.SourceMaterial);
-
+	}
+}
+void FGeometryCollectionEngineConversion::AppendGeometryCollectionSource(const FGeometryCollectionSource& GeometryCollectionSource, FGeometryCollection& GeometryCollectionInOut, TArray<UMaterialInterface*>& MaterialInstancesInOut, bool ReindexMaterials)
+{
+	const int32 StartMaterialIndex = MaterialInstancesInOut.Num();
+	if (AppendGeometryCollectionSourceNoMaterial(GeometryCollectionSource, GeometryCollectionInOut, StartMaterialIndex, ReindexMaterials))
+	{
+		MaterialInstancesInOut.Append(GeometryCollectionSource.SourceMaterial);
+	}
+}
+bool FGeometryCollectionEngineConversion::AppendGeometryCollectionSourceNoMaterial(const FGeometryCollectionSource& GeometryCollectionSource, FGeometryCollection& GeometryCollectionInOut, int32 StartMaterialIndex, bool ReindexMaterials)
+{
+	if (const UObject* SourceObject = GeometryCollectionSource.SourceGeometryObject.TryLoad())
+	{
 		if (const UStaticMesh* SourceStaticMesh = Cast<UStaticMesh>(SourceObject))
 		{
 			bool bLegacyAddInternal = GeometryCollectionSource.bAddInternalMaterials;
@@ -1507,6 +1519,7 @@ void FGeometryCollectionEngineConversion::AppendGeometryCollectionSource(const F
 				GeometryCollectionSource.bSplitComponents,
 				GeometryCollectionSource.bSetInternalFromMaterialIndex
 				);
+			return true;
 		}
 		else if (const USkeletalMesh* SourceSkeletalMesh = Cast<USkeletalMesh>(SourceObject))
 		{
@@ -1517,6 +1530,7 @@ void FGeometryCollectionEngineConversion::AppendGeometryCollectionSource(const F
 				&GeometryCollectionInOut,
 				ReindexMaterials
 				);
+			return true;
 		}
 		else if (const UGeometryCollection* SourceGeometryCollection = Cast<UGeometryCollection>(SourceObject))
 		{
@@ -1527,11 +1541,20 @@ void FGeometryCollectionEngineConversion::AppendGeometryCollectionSource(const F
 				&GeometryCollectionInOut,
 				ReindexMaterials
 				);
+			return true;
 		}
 	}
+	return false;
 }
 
 void FGeometryCollectionEngineConversion::ConvertStaticMeshToGeometryCollection(const TObjectPtr<UStaticMesh> StaticMesh, FManagedArrayCollection& OutCollection, TArray<TObjectPtr<UMaterial>>& OutMaterials, TArray<FGeometryCollectionAutoInstanceMesh>& OutInstancedMeshes, bool bSetInternalFromMaterialIndex, bool bSplitComponents)
+{
+	TArray<TObjectPtr<UMaterialInterface>> OutMaterialInstances;
+	ConvertStaticMeshToGeometryCollection(StaticMesh, OutCollection, OutMaterialInstances, OutInstancedMeshes, bSetInternalFromMaterialIndex, bSplitComponents);
+	GetMaterialsFromInstances(OutMaterialInstances, OutMaterials);
+}
+
+void FGeometryCollectionEngineConversion::ConvertStaticMeshToGeometryCollection(const TObjectPtr<UStaticMesh> StaticMesh, FManagedArrayCollection& OutCollection, TArray<TObjectPtr<UMaterialInterface>>& OutMaterialInstances, TArray<FGeometryCollectionAutoInstanceMesh>& OutInstancedMeshes, bool bSetInternalFromMaterialIndex, bool bSplitComponents)
 {
 #if WITH_EDITORONLY_DATA
 	if (UGeometryCollection* NewGeometryCollection = NewObject<UGeometryCollection>())
@@ -1566,10 +1589,7 @@ void FGeometryCollectionEngineConversion::ConvertStaticMeshToGeometryCollection(
 		NewGeometryCollection->InitializeMaterials();
 
 		// Materials
-		for (auto& Material : NewGeometryCollection->Materials)
-		{
-			OutMaterials.Emplace(Material->GetMaterial());
-		}
+		OutMaterialInstances.Append(NewGeometryCollection->Materials);
 
 		TSharedPtr<FGeometryCollection> OutCollectionPtr = NewGeometryCollection->GetGeometryCollection();
 		OutCollectionPtr->CopyTo(&OutCollection);
@@ -1578,16 +1598,20 @@ void FGeometryCollectionEngineConversion::ConvertStaticMeshToGeometryCollection(
 }
 
 void FGeometryCollectionEngineConversion::ConvertGeometryCollectionToGeometryCollection(const TObjectPtr<UGeometryCollection> InGeometryCollectionAssetPtr, FManagedArrayCollection& OutCollection, TArray<TObjectPtr<UMaterial>>& OutMaterials, TArray<FGeometryCollectionAutoInstanceMesh>& OutInstancedMeshes)
+{
+	TArray<TObjectPtr<UMaterialInterface>> OutMaterialInstances;
+	ConvertGeometryCollectionToGeometryCollection(InGeometryCollectionAssetPtr, OutCollection, OutMaterialInstances, OutInstancedMeshes);
+	GetMaterialsFromInstances(OutMaterialInstances, OutMaterials);
+}
+
+void FGeometryCollectionEngineConversion::ConvertGeometryCollectionToGeometryCollection(const TObjectPtr<UGeometryCollection> InGeometryCollectionAssetPtr, FManagedArrayCollection& OutCollection, TArray<TObjectPtr<UMaterialInterface>>& OutMaterialInstances, TArray<FGeometryCollectionAutoInstanceMesh>& OutInstancedMeshes)
 {
 	if (InGeometryCollectionAssetPtr)
 	{
 		const TSharedPtr<FGeometryCollection, ESPMode::ThreadSafe> NewGeometryCollectionPtr = InGeometryCollectionAssetPtr->GetGeometryCollection();
 
 		// Materials
-		for (auto& Material : InGeometryCollectionAssetPtr->Materials)
-		{
-			OutMaterials.Emplace(Material->GetMaterial());
-		}
+		OutMaterialInstances = InGeometryCollectionAssetPtr->Materials;
 
 		// InstanceMeshes
 		OutInstancedMeshes = InGeometryCollectionAssetPtr->AutoInstanceMeshes;
@@ -1599,7 +1623,15 @@ void FGeometryCollectionEngineConversion::ConvertGeometryCollectionToGeometryCol
 	}
 }
 
+
 void FGeometryCollectionEngineConversion::ConvertActorToGeometryCollection(const AActor* Actor, FManagedArrayCollection& OutCollection, TArray<TObjectPtr<UMaterial>>& OutMaterials, TArray<FGeometryCollectionAutoInstanceMesh>& OutInstancedMeshes, const FSkeletalMeshToCollectionConversionParameters& ConversionParameters, bool bSplitComponents)
+{
+	TArray<TObjectPtr<UMaterialInterface>> OutMaterialInstances;
+	ConvertActorToGeometryCollection(Actor, OutCollection, OutMaterialInstances, OutInstancedMeshes, ConversionParameters, bSplitComponents);
+	GetMaterialsFromInstances(OutMaterialInstances, OutMaterials);
+}
+
+void FGeometryCollectionEngineConversion::ConvertActorToGeometryCollection(const AActor* Actor, FManagedArrayCollection& OutCollection, TArray<TObjectPtr<UMaterialInterface>>& OutMaterialInstances, TArray<FGeometryCollectionAutoInstanceMesh>& OutInstancedMeshes, const FSkeletalMeshToCollectionConversionParameters& ConversionParameters, bool bSplitComponents)
 {
 #if WITH_EDITORONLY_DATA
 	const FTransform ActorTransform(Actor->GetTransform());
@@ -1671,10 +1703,7 @@ void FGeometryCollectionEngineConversion::ConvertActorToGeometryCollection(const
 		OutInstancedMeshes.Append(NewGeometryCollection->AutoInstanceMeshes);
 
 		// Materials
-		for (auto& Material : NewGeometryCollection->Materials)
-		{
-			OutMaterials.Emplace(Material->GetMaterial());
-		}
+		OutMaterialInstances = NewGeometryCollection->Materials;
 
 		TSharedPtr<FGeometryCollection> OutCollectionPtr = NewGeometryCollection->GetGeometryCollection();
 		OutCollectionPtr->CopyTo(&OutCollection);
@@ -1748,5 +1777,17 @@ void FGeometryCollectionEngineConversion::ConvertCollectionToSkeleton(const FMan
 	}
 }
 
+void FGeometryCollectionEngineConversion::GetMaterialsFromInstances(const TArray<TObjectPtr<UMaterialInterface>>& MaterialInstances, TArray<TObjectPtr<UMaterial>>& OutMaterials)
+{
+	// somehow append does not like appending to a TArray<ObjectPtr<>>
+	// so we need to store in a array of raw pointers and then transfer over 
+	TArray<UMaterial*> MaterialArray;
+	MaterialArray.Append(MaterialInstances);
+	OutMaterials.Reserve(MaterialInstances.Num());
+	for (UMaterial* Material : MaterialArray)
+	{
+		OutMaterials.Add(Material);
+	}
+}
 
 #undef LOCTEXT_NAMESPACE 
diff --git a/Engine/Source/Runtime/Experimental/GeometryCollectionEngine/Private/GeometryCollection/GeometryCollectionObject.cpp b/Engine/Source/Runtime/Experimental/GeometryCollectionEngine/Private/GeometryCollection/GeometryCollectionObject.cpp
--- a/Engine/Source/Runtime/Experimental/GeometryCollectionEngine/Private/GeometryCollection/GeometryCollectionObject.cpp
+++ b/Engine/Source/Runtime/Experimental/GeometryCollectionEngine/Private/GeometryCollection/GeometryCollectionObject.cpp
@@ -567,6 +567,19 @@ void UGeometryCollection::ResetFrom(const FManagedArrayCollection& InCollection,
 	}
 }
 
+void UGeometryCollection::ResetFrom(const FManagedArrayCollection& InCollection, const TArray<UMaterialInterface*>& InMaterialInstances, bool bHasInternalMaterials)
+{
+	if (GeometryCollection.IsValid())
+	{
+		Reset();
+		InCollection.CopyTo(GeometryCollection.Get());
+		// todo(Chaos) : we could certainly run a "dependent attribute update method here instead of having to known about convex specifically 
+		UpdateConvexGeometryIfMissing();
+		Materials.Append(InMaterialInstances);
+		InitializeMaterials(bHasInternalMaterials);
+	}
+}
+
 /** AppendGeometry */
 int32 UGeometryCollection::AppendGeometry(const UGeometryCollection & Element, bool ReindexAllMaterials, const FTransform& TransformRoot)
 {
diff --git a/Engine/Source/Runtime/Experimental/GeometryCollectionEngine/Public/GeometryCollection/GeometryCollectionEngineConversion.h b/Engine/Source/Runtime/Experimental/GeometryCollectionEngine/Public/GeometryCollection/GeometryCollectionEngineConversion.h
--- a/Engine/Source/Runtime/Experimental/GeometryCollectionEngine/Public/GeometryCollection/GeometryCollectionEngineConversion.h
+++ b/Engine/Source/Runtime/Experimental/GeometryCollectionEngine/Public/GeometryCollection/GeometryCollectionEngineConversion.h
@@ -179,6 +179,7 @@ public:
 	*  @param MaterialsInOut  : array of materials to append to
 	*/
 	static GEOMETRYCOLLECTIONENGINE_API void AppendGeometryCollectionSource(const FGeometryCollectionSource& GeometryCollectionSource, FGeometryCollection& GeometryCollectionInOut, TArray<UMaterial*>& MaterialsInOut, bool ReindexMaterials = true);
+	static GEOMETRYCOLLECTIONENGINE_API void AppendGeometryCollectionSource(const FGeometryCollectionSource& GeometryCollectionSource, FGeometryCollection& GeometryCollectionInOut, TArray<UMaterialInterface*>& MaterialInstancesInOut, bool ReindexMaterials = true);
 
 	/**
 	*  Converts a StaticMesh to a GeometryCollection
@@ -190,6 +191,7 @@ public:
 	*  @param bSplitComponents : Split the components
 	*/
 	static GEOMETRYCOLLECTIONENGINE_API void ConvertStaticMeshToGeometryCollection(const TObjectPtr<UStaticMesh> StaticMesh, FManagedArrayCollection& OutCollection, TArray<TObjectPtr<UMaterial>>& OutMaterials, TArray<FGeometryCollectionAutoInstanceMesh>& OutInstancedMeshes, bool bSetInternalFromMaterialIndex = true, bool bSplitComponents = false);
+	static GEOMETRYCOLLECTIONENGINE_API void ConvertStaticMeshToGeometryCollection(const TObjectPtr<UStaticMesh> StaticMesh, FManagedArrayCollection& OutCollection, TArray<TObjectPtr<UMaterialInterface>>& OutMaterialInstances, TArray<FGeometryCollectionAutoInstanceMesh>& OutInstancedMeshes, bool bSetInternalFromMaterialIndex = true, bool bSplitComponents = false);
 
 	/**
 	*  Converts a UGeometryCollection asset to an FManagedArrayCollection, and arrays of materials and instanced meshes
@@ -199,6 +201,7 @@ public:
 	*  @param OutInstancedMeshes : InstancedMeshes
 	*/	
 	static GEOMETRYCOLLECTIONENGINE_API void ConvertGeometryCollectionToGeometryCollection(const TObjectPtr<UGeometryCollection> InGeometryCollectionAssetPtr, FManagedArrayCollection& OutCollection, TArray<TObjectPtr<UMaterial>>& OutMaterials, TArray<FGeometryCollectionAutoInstanceMesh>& OutInstancedMeshes);
+	static GEOMETRYCOLLECTIONENGINE_API void ConvertGeometryCollectionToGeometryCollection(const TObjectPtr<UGeometryCollection> InGeometryCollectionAssetPtr, FManagedArrayCollection& OutCollection, TArray<TObjectPtr<UMaterialInterface>>& OutMaterialInstances, TArray<FGeometryCollectionAutoInstanceMesh>& OutInstancedMeshes);
 
 	/**
 	*  Converts an actor to an FManagedArrayCollection, and arrays of materials and instanced meshes
@@ -209,6 +212,7 @@ public:
 	*  @param bSplitComponents : Split the components
 	*/
 	static GEOMETRYCOLLECTIONENGINE_API void ConvertActorToGeometryCollection(const AActor* Actor, FManagedArrayCollection& OutCollection, TArray<TObjectPtr<UMaterial>>& OutMaterials, TArray<FGeometryCollectionAutoInstanceMesh>& OutInstancedMeshes, const FSkeletalMeshToCollectionConversionParameters& ConversionParameters, bool bSplitComponents = false);
+	static GEOMETRYCOLLECTIONENGINE_API void ConvertActorToGeometryCollection(const AActor* Actor, FManagedArrayCollection& OutCollection, TArray<TObjectPtr<UMaterialInterface>>& OutMaterialInstances, TArray<FGeometryCollectionAutoInstanceMesh>& OutInstancedMeshes, const FSkeletalMeshToCollectionConversionParameters& ConversionParameters, bool bSplitComponents = false);
 
 	/**
 	*  Converts an FManagedArrayCollection to a USkeleton
@@ -218,5 +222,15 @@ public:
 	*/
 	static GEOMETRYCOLLECTIONENGINE_API void ConvertCollectionToSkeleton(const FManagedArrayCollection& InCollection, USkeleton* OutSkeleton, TArray<int32>& OutIndexRemap);
 
-
+	/** utility function to convert an array of material instances to an array of material */
+	static GEOMETRYCOLLECTIONENGINE_API void GetMaterialsFromInstances(const TArray<TObjectPtr<UMaterialInterface>>& MaterialInstances, TArray<TObjectPtr<UMaterial>>& OutMaterials);
+private:
+	/**
+	*  Appends a GeometryCollectionSource to a GeometryCollection ignoring the materials
+	*  this is meant to be used by the AppendGeometryCollectionSource variants
+	*  @param GeometryCollectionSource : geometry collection source object ( from UGeometryCollection collection asset )
+	*  @param GeometryCollectionInOut : GeometryCollection to append to
+	*  @return true if the source was added successfully ( valid object )
+	*/
+	static GEOMETRYCOLLECTIONENGINE_API bool AppendGeometryCollectionSourceNoMaterial(const FGeometryCollectionSource& GeometryCollectionSource, FGeometryCollection& GeometryCollectionInOut, int32 StartMaterialIndex, bool ReindexMaterials = true);
 };
diff --git a/Engine/Source/Runtime/Experimental/GeometryCollectionEngine/Public/GeometryCollection/GeometryCollectionObject.h b/Engine/Source/Runtime/Experimental/GeometryCollectionEngine/Public/GeometryCollection/GeometryCollectionObject.h
--- a/Engine/Source/Runtime/Experimental/GeometryCollectionEngine/Public/GeometryCollection/GeometryCollectionObject.h
+++ b/Engine/Source/Runtime/Experimental/GeometryCollectionEngine/Public/GeometryCollection/GeometryCollectionObject.h
@@ -395,6 +395,7 @@ public:
 
 	/** Reset the collection from another set of attributes and materials. */
 	GEOMETRYCOLLECTIONENGINE_API void ResetFrom(const FManagedArrayCollection& InCollection, const TArray<UMaterial*>& InMaterials, bool bHasInternalMaterials);
+	GEOMETRYCOLLECTIONENGINE_API void ResetFrom(const FManagedArrayCollection& InCollection, const TArray<UMaterialInterface*>& InMaterialInstances, bool bHasInternalMaterials);
 	
 	GEOMETRYCOLLECTIONENGINE_API int32 AppendGeometry(const UGeometryCollection & Element, bool ReindexAllMaterials = false, const FTransform& TransformRoot = FTransform::Identity);
 	GEOMETRYCOLLECTIONENGINE_API int32 NumElements(const FName& Group) const;
