Dataflow : Add <Name,Type> for dataflow node renderers.

diff --git a/Engine/Plugins/Experimental/ChaosFlesh/Source/ChaosFleshDeprecatedNodes/Public/Dataflow/ChaosFleshConstructTetGridNode.h b/Engine/Plugins/Experimental/ChaosFlesh/Source/ChaosFleshDeprecatedNodes/Public/Dataflow/ChaosFleshConstructTetGridNode.h
--- a/Engine/Plugins/Experimental/ChaosFlesh/Source/ChaosFleshDeprecatedNodes/Public/Dataflow/ChaosFleshConstructTetGridNode.h
+++ b/Engine/Plugins/Experimental/ChaosFlesh/Source/ChaosFleshDeprecatedNodes/Public/Dataflow/ChaosFleshConstructTetGridNode.h
@@ -15,7 +15,7 @@ struct FConstructTetGridNode : public FDataflowNode
 {
 	GENERATED_USTRUCT_BODY()
 	DATAFLOW_NODE_DEFINE_INTERNAL(FConstructTetGridNode, "TetGrid", "Flesh", "")
-	DATAFLOW_NODE_RENDER_TYPE(FGeometryCollection::StaticType(), "Collection")
+	DATAFLOW_NODE_RENDER_TYPE("SurfaceRender",FGeometryCollection::StaticType(), "Collection")
 
 public:
 	typedef FManagedArrayCollection DataType;
diff --git a/Engine/Plugins/Experimental/ChaosFlesh/Source/ChaosFleshDeprecatedNodes/Public/Dataflow/ChaosFleshCreateTetrahedralCollectionNode.h b/Engine/Plugins/Experimental/ChaosFlesh/Source/ChaosFleshDeprecatedNodes/Public/Dataflow/ChaosFleshCreateTetrahedralCollectionNode.h
--- a/Engine/Plugins/Experimental/ChaosFlesh/Source/ChaosFleshDeprecatedNodes/Public/Dataflow/ChaosFleshCreateTetrahedralCollectionNode.h
+++ b/Engine/Plugins/Experimental/ChaosFlesh/Source/ChaosFleshDeprecatedNodes/Public/Dataflow/ChaosFleshCreateTetrahedralCollectionNode.h
@@ -26,7 +26,7 @@ struct FGenerateTetrahedralCollectionDataflowNodes : public FDataflowNode
 {
 	GENERATED_USTRUCT_BODY()
 	DATAFLOW_NODE_DEFINE_INTERNAL(FGenerateTetrahedralCollectionDataflowNodes, "GenerateTetrahedralCollection", "Flesh", "")
-	DATAFLOW_NODE_RENDER_TYPE(FGeometryCollection::StaticType(), "Collection")
+	DATAFLOW_NODE_RENDER_TYPE("SurfaceRender", FGeometryCollection::StaticType(), "Collection")
 
 public:
 	typedef FManagedArrayCollection DataType;
diff --git a/Engine/Plugins/Experimental/ChaosFlesh/Source/ChaosFleshDeprecatedNodes/Public/Dataflow/ChaosFleshKinematicOriginInsertionInitializationNode.h b/Engine/Plugins/Experimental/ChaosFlesh/Source/ChaosFleshDeprecatedNodes/Public/Dataflow/ChaosFleshKinematicOriginInsertionInitializationNode.h
--- a/Engine/Plugins/Experimental/ChaosFlesh/Source/ChaosFleshDeprecatedNodes/Public/Dataflow/ChaosFleshKinematicOriginInsertionInitializationNode.h
+++ b/Engine/Plugins/Experimental/ChaosFlesh/Source/ChaosFleshDeprecatedNodes/Public/Dataflow/ChaosFleshKinematicOriginInsertionInitializationNode.h
@@ -15,7 +15,7 @@ struct FKinematicOriginInsertionInitializationDataflowNode : public FDataflowNod
 {
 	GENERATED_USTRUCT_BODY()
 	DATAFLOW_NODE_DEFINE_INTERNAL(FKinematicOriginInsertionInitializationDataflowNode, "KinematicOriginInsertionInitialization", "Flesh", "")
-	DATAFLOW_NODE_RENDER_TYPE(FGeometryCollection::StaticType(), "Collection")
+	DATAFLOW_NODE_RENDER_TYPE("SurfaceRender", FGeometryCollection::StaticType(), "Collection")
 
 public:
 	typedef FManagedArrayCollection DataType;
diff --git a/Engine/Plugins/Experimental/ChaosFlesh/Source/ChaosFleshDeprecatedNodes/Public/Dataflow/ChaosFleshKinematicTetrahedralConstraintNode.h b/Engine/Plugins/Experimental/ChaosFlesh/Source/ChaosFleshDeprecatedNodes/Public/Dataflow/ChaosFleshKinematicTetrahedralConstraintNode.h
--- a/Engine/Plugins/Experimental/ChaosFlesh/Source/ChaosFleshDeprecatedNodes/Public/Dataflow/ChaosFleshKinematicTetrahedralConstraintNode.h
+++ b/Engine/Plugins/Experimental/ChaosFlesh/Source/ChaosFleshDeprecatedNodes/Public/Dataflow/ChaosFleshKinematicTetrahedralConstraintNode.h
@@ -16,7 +16,7 @@ struct FKinematicTetrahedralBindingsDataflowNode : public FDataflowNode
 {
 	GENERATED_USTRUCT_BODY()
 	DATAFLOW_NODE_DEFINE_INTERNAL(FKinematicTetrahedralBindingsDataflowNode, "KinematicTetrahedralBindings", "Flesh", "")
-	DATAFLOW_NODE_RENDER_TYPE(FGeometryCollection::StaticType(), "Collection")
+	DATAFLOW_NODE_RENDER_TYPE("SurfaceRender", FGeometryCollection::StaticType(), "Collection")
 
 public:
 	typedef FManagedArrayCollection DataType;
diff --git a/Engine/Plugins/Experimental/ChaosFlesh/Source/ChaosFleshNodes/Public/Dataflow/ChaosFleshAddKinematicParticlesNode.h b/Engine/Plugins/Experimental/ChaosFlesh/Source/ChaosFleshNodes/Public/Dataflow/ChaosFleshAddKinematicParticlesNode.h
--- a/Engine/Plugins/Experimental/ChaosFlesh/Source/ChaosFleshNodes/Public/Dataflow/ChaosFleshAddKinematicParticlesNode.h
+++ b/Engine/Plugins/Experimental/ChaosFlesh/Source/ChaosFleshNodes/Public/Dataflow/ChaosFleshAddKinematicParticlesNode.h
@@ -16,7 +16,7 @@ struct FAddKinematicParticlesDataflowNode : public FDataflowNode
 {
 	GENERATED_USTRUCT_BODY()
 	DATAFLOW_NODE_DEFINE_INTERNAL(FAddKinematicParticlesDataflowNode, "AddKinematicParticles", "Flesh", "")
-	DATAFLOW_NODE_RENDER_TYPE(FGeometryCollection::StaticType(), "Collection")
+	DATAFLOW_NODE_RENDER_TYPE("SurfaceRender", FGeometryCollection::StaticType(), "Collection")
 
 public:
 	typedef FManagedArrayCollection DataType;
diff --git a/Engine/Plugins/Experimental/ChaosFlesh/Source/ChaosFleshNodes/Public/Dataflow/ChaosFleshAppendTetrahedralCollectionNode.h b/Engine/Plugins/Experimental/ChaosFlesh/Source/ChaosFleshNodes/Public/Dataflow/ChaosFleshAppendTetrahedralCollectionNode.h
--- a/Engine/Plugins/Experimental/ChaosFlesh/Source/ChaosFleshNodes/Public/Dataflow/ChaosFleshAppendTetrahedralCollectionNode.h
+++ b/Engine/Plugins/Experimental/ChaosFlesh/Source/ChaosFleshNodes/Public/Dataflow/ChaosFleshAppendTetrahedralCollectionNode.h
@@ -17,7 +17,7 @@ struct FAppendTetrahedralCollectionDataflowNode : public FDataflowNode
 {
 	GENERATED_USTRUCT_BODY()
 	DATAFLOW_NODE_DEFINE_INTERNAL(FAppendTetrahedralCollectionDataflowNode, "AppendTetrahedralCollection", "Flesh", "")
-	DATAFLOW_NODE_RENDER_TYPE(FGeometryCollection::StaticType(), "Collection1")
+	DATAFLOW_NODE_RENDER_TYPE("SurfaceRender", FGeometryCollection::StaticType(), "Collection1")
 public:
 	UPROPERTY(meta = (DataflowInput, DataflowOutput, DisplayName = "Collection1", DataflowPassthrough = "Collection1"));
 	FManagedArrayCollection Collection1;
diff --git a/Engine/Plugins/Experimental/ChaosFlesh/Source/ChaosFleshNodes/Public/Dataflow/ChaosFleshCollisionBodyConstraintNode.h b/Engine/Plugins/Experimental/ChaosFlesh/Source/ChaosFleshNodes/Public/Dataflow/ChaosFleshCollisionBodyConstraintNode.h
--- a/Engine/Plugins/Experimental/ChaosFlesh/Source/ChaosFleshNodes/Public/Dataflow/ChaosFleshCollisionBodyConstraintNode.h
+++ b/Engine/Plugins/Experimental/ChaosFlesh/Source/ChaosFleshNodes/Public/Dataflow/ChaosFleshCollisionBodyConstraintNode.h
@@ -16,7 +16,7 @@ struct FKinematicBodySetupInitializationDataflowNode : public FDataflowNode
 {
 	GENERATED_USTRUCT_BODY()
 	DATAFLOW_NODE_DEFINE_INTERNAL(FKinematicBodySetupInitializationDataflowNode, "KinematicBodySetupInitialization", "Flesh", "")
-	DATAFLOW_NODE_RENDER_TYPE(FGeometryCollection::StaticType(), "Collection")
+	DATAFLOW_NODE_RENDER_TYPE("SurfaceRender", FGeometryCollection::StaticType(), "Collection")
 
 public:
 	typedef FManagedArrayCollection DataType;
diff --git a/Engine/Plugins/Experimental/ChaosFlesh/Source/ChaosFleshNodes/Public/Dataflow/ChaosFleshComputeFiberFieldNode.h b/Engine/Plugins/Experimental/ChaosFlesh/Source/ChaosFleshNodes/Public/Dataflow/ChaosFleshComputeFiberFieldNode.h
--- a/Engine/Plugins/Experimental/ChaosFlesh/Source/ChaosFleshNodes/Public/Dataflow/ChaosFleshComputeFiberFieldNode.h
+++ b/Engine/Plugins/Experimental/ChaosFlesh/Source/ChaosFleshNodes/Public/Dataflow/ChaosFleshComputeFiberFieldNode.h
@@ -20,7 +20,7 @@ struct FComputeFiberFieldNode : public FDataflowNode
 {
 	GENERATED_USTRUCT_BODY()
 	DATAFLOW_NODE_DEFINE_INTERNAL(FComputeFiberFieldNode, "ComputeFiberField", "Flesh", "")
-	DATAFLOW_NODE_RENDER_TYPE(FGeometryCollection::StaticType(), "Collection")
+	DATAFLOW_NODE_RENDER_TYPE("SurfaceRender", FGeometryCollection::StaticType(), "Collection")
 
 public:
 	//typedef FManagedArrayCollection DataType;
diff --git a/Engine/Plugins/Experimental/ChaosFlesh/Source/ChaosFleshNodes/Public/Dataflow/ChaosFleshComputeIslandsNode.h b/Engine/Plugins/Experimental/ChaosFlesh/Source/ChaosFleshNodes/Public/Dataflow/ChaosFleshComputeIslandsNode.h
--- a/Engine/Plugins/Experimental/ChaosFlesh/Source/ChaosFleshNodes/Public/Dataflow/ChaosFleshComputeIslandsNode.h
+++ b/Engine/Plugins/Experimental/ChaosFlesh/Source/ChaosFleshNodes/Public/Dataflow/ChaosFleshComputeIslandsNode.h
@@ -14,7 +14,7 @@ struct FComputeIslandsNode : public FDataflowNode
 {
 	GENERATED_USTRUCT_BODY()
 	DATAFLOW_NODE_DEFINE_INTERNAL(FComputeIslandsNode, "ComputeIslands", "Flesh", "")
-	DATAFLOW_NODE_RENDER_TYPE(FGeometryCollection::StaticType(), "Collection")
+	DATAFLOW_NODE_RENDER_TYPE("SurfaceRender", FGeometryCollection::StaticType(), "Collection")
 
 public:
 
diff --git a/Engine/Plugins/Experimental/ChaosFlesh/Source/ChaosFleshNodes/Public/Dataflow/ChaosFleshCreateTetrahedronNode.h b/Engine/Plugins/Experimental/ChaosFlesh/Source/ChaosFleshNodes/Public/Dataflow/ChaosFleshCreateTetrahedronNode.h
--- a/Engine/Plugins/Experimental/ChaosFlesh/Source/ChaosFleshNodes/Public/Dataflow/ChaosFleshCreateTetrahedronNode.h
+++ b/Engine/Plugins/Experimental/ChaosFlesh/Source/ChaosFleshNodes/Public/Dataflow/ChaosFleshCreateTetrahedronNode.h
@@ -22,7 +22,7 @@ struct FCreateTetrahedronDataflowNode : public FDataflowNode
 {
 	GENERATED_USTRUCT_BODY()
 	DATAFLOW_NODE_DEFINE_INTERNAL(FCreateTetrahedronDataflowNode, "CreateTetrahedron", "Flesh", "")
-	DATAFLOW_NODE_RENDER_TYPE(FGeometryCollection::StaticType(), "Collection")
+	DATAFLOW_NODE_RENDER_TYPE("SurfaceRender", FGeometryCollection::StaticType(), "Collection")
 
 public:
 	typedef FManagedArrayCollection DataType;
diff --git a/Engine/Plugins/Experimental/ChaosFlesh/Source/ChaosFleshNodes/Public/Dataflow/ChaosFleshFleshAssetTerminalNode.h b/Engine/Plugins/Experimental/ChaosFlesh/Source/ChaosFleshNodes/Public/Dataflow/ChaosFleshFleshAssetTerminalNode.h
--- a/Engine/Plugins/Experimental/ChaosFlesh/Source/ChaosFleshNodes/Public/Dataflow/ChaosFleshFleshAssetTerminalNode.h
+++ b/Engine/Plugins/Experimental/ChaosFlesh/Source/ChaosFleshNodes/Public/Dataflow/ChaosFleshFleshAssetTerminalNode.h
@@ -14,7 +14,7 @@ struct FFleshAssetTerminalDataflowNode : public FDataflowTerminalNode
 {
 	GENERATED_USTRUCT_BODY()
 	DATAFLOW_NODE_DEFINE_INTERNAL(FFleshAssetTerminalDataflowNode, "FleshAssetTerminal", "Terminal", "")
-	DATAFLOW_NODE_RENDER_TYPE(FGeometryCollection::StaticType(), "Collection")
+	DATAFLOW_NODE_RENDER_TYPE("SurfaceRender", FGeometryCollection::StaticType(), "Collection")
 
 public:
 
diff --git a/Engine/Plugins/Experimental/ChaosFlesh/Source/ChaosFleshNodes/Public/Dataflow/ChaosFleshGenerateOriginInsertionNode.h b/Engine/Plugins/Experimental/ChaosFlesh/Source/ChaosFleshNodes/Public/Dataflow/ChaosFleshGenerateOriginInsertionNode.h
--- a/Engine/Plugins/Experimental/ChaosFlesh/Source/ChaosFleshNodes/Public/Dataflow/ChaosFleshGenerateOriginInsertionNode.h
+++ b/Engine/Plugins/Experimental/ChaosFlesh/Source/ChaosFleshNodes/Public/Dataflow/ChaosFleshGenerateOriginInsertionNode.h
@@ -14,7 +14,7 @@ struct FGenerateOriginInsertionNode : public FDataflowNode
 {
 	GENERATED_USTRUCT_BODY()
 	DATAFLOW_NODE_DEFINE_INTERNAL(FGenerateOriginInsertionNode, "GenerateOriginInsertion", "Flesh", "")
-	DATAFLOW_NODE_RENDER_TYPE(FGeometryCollection::StaticType(), "Collection")
+	DATAFLOW_NODE_RENDER_TYPE("SurfaceRender", FGeometryCollection::StaticType(), "Collection")
 
 public:
 	//typedef FManagedArrayCollection DataType;
diff --git a/Engine/Plugins/Experimental/ChaosFlesh/Source/ChaosFleshNodes/Public/Dataflow/ChaosFleshGenerateSkeletalBindingsNode.h b/Engine/Plugins/Experimental/ChaosFlesh/Source/ChaosFleshNodes/Public/Dataflow/ChaosFleshGenerateSkeletalBindingsNode.h
--- a/Engine/Plugins/Experimental/ChaosFlesh/Source/ChaosFleshNodes/Public/Dataflow/ChaosFleshGenerateSkeletalBindingsNode.h
+++ b/Engine/Plugins/Experimental/ChaosFlesh/Source/ChaosFleshNodes/Public/Dataflow/ChaosFleshGenerateSkeletalBindingsNode.h
@@ -19,7 +19,7 @@ struct FGenerateSkeletalBindings : public FDataflowNode
 {
 	GENERATED_USTRUCT_BODY()
 	DATAFLOW_NODE_DEFINE_INTERNAL(FGenerateSkeletalBindings, "GenerateSkeletalBindings", "Flesh", "")
-	DATAFLOW_NODE_RENDER_TYPE(FGeometryCollection::StaticType(), "Collection")
+	DATAFLOW_NODE_RENDER_TYPE("SurfaceRender", FGeometryCollection::StaticType(), "Collection")
 
 public:
 	typedef FManagedArrayCollection DataType;
diff --git a/Engine/Plugins/Experimental/ChaosFlesh/Source/ChaosFleshNodes/Public/Dataflow/ChaosFleshGenerateSurfaceBindingsNode.h b/Engine/Plugins/Experimental/ChaosFlesh/Source/ChaosFleshNodes/Public/Dataflow/ChaosFleshGenerateSurfaceBindingsNode.h
--- a/Engine/Plugins/Experimental/ChaosFlesh/Source/ChaosFleshNodes/Public/Dataflow/ChaosFleshGenerateSurfaceBindingsNode.h
+++ b/Engine/Plugins/Experimental/ChaosFlesh/Source/ChaosFleshNodes/Public/Dataflow/ChaosFleshGenerateSurfaceBindingsNode.h
@@ -28,7 +28,7 @@ struct FGenerateSurfaceBindings : public FDataflowNode
 {
 	GENERATED_USTRUCT_BODY()
 	DATAFLOW_NODE_DEFINE_INTERNAL(FGenerateSurfaceBindings, "GenerateSurfaceBindings", "Flesh", "")
-	DATAFLOW_NODE_RENDER_TYPE(FGeometryCollection::StaticType(), "Collection")
+	DATAFLOW_NODE_RENDER_TYPE("SurfaceRender", FGeometryCollection::StaticType(), "Collection")
 
 public:
 	typedef FManagedArrayCollection DataType;
diff --git a/Engine/Plugins/Experimental/ChaosFlesh/Source/ChaosFleshNodes/Public/Dataflow/ChaosFleshGetFleshAssetNode.h b/Engine/Plugins/Experimental/ChaosFlesh/Source/ChaosFleshNodes/Public/Dataflow/ChaosFleshGetFleshAssetNode.h
--- a/Engine/Plugins/Experimental/ChaosFlesh/Source/ChaosFleshNodes/Public/Dataflow/ChaosFleshGetFleshAssetNode.h
+++ b/Engine/Plugins/Experimental/ChaosFlesh/Source/ChaosFleshNodes/Public/Dataflow/ChaosFleshGetFleshAssetNode.h
@@ -14,7 +14,7 @@ struct FGetFleshAssetDataflowNode : public FDataflowNode
 {
 	GENERATED_USTRUCT_BODY()
 	DATAFLOW_NODE_DEFINE_INTERNAL(FGetFleshAssetDataflowNode, "GetFleshAsset", "Flesh", "")
-	DATAFLOW_NODE_RENDER_TYPE(FGeometryCollection::StaticType(), "Collection")
+	DATAFLOW_NODE_RENDER_TYPE("SurfaceRender", FGeometryCollection::StaticType(), "Collection")
 
 public:
 
diff --git a/Engine/Plugins/Experimental/ChaosFlesh/Source/ChaosFleshNodes/Public/Dataflow/ChaosFleshGetSurfaceIndicesNode.h b/Engine/Plugins/Experimental/ChaosFlesh/Source/ChaosFleshNodes/Public/Dataflow/ChaosFleshGetSurfaceIndicesNode.h
--- a/Engine/Plugins/Experimental/ChaosFlesh/Source/ChaosFleshNodes/Public/Dataflow/ChaosFleshGetSurfaceIndicesNode.h
+++ b/Engine/Plugins/Experimental/ChaosFlesh/Source/ChaosFleshNodes/Public/Dataflow/ChaosFleshGetSurfaceIndicesNode.h
@@ -13,7 +13,7 @@ struct FGetSurfaceIndicesNode : public FDataflowNode
 {
 	GENERATED_USTRUCT_BODY()
 	DATAFLOW_NODE_DEFINE_INTERNAL(FGetSurfaceIndicesNode, "GetSurfaceIndices", "Flesh", "")
-	DATAFLOW_NODE_RENDER_TYPE(FGeometryCollection::StaticType(), "Collection")
+	DATAFLOW_NODE_RENDER_TYPE("SurfaceRender", FGeometryCollection::StaticType(), "Collection")
 
 public:
 	UPROPERTY(meta = (DataflowInput, DisplayName = "Collection"))
diff --git a/Engine/Plugins/Experimental/ChaosFlesh/Source/ChaosFleshNodes/Public/Dataflow/ChaosFleshIsolateComponentNode.h b/Engine/Plugins/Experimental/ChaosFlesh/Source/ChaosFleshNodes/Public/Dataflow/ChaosFleshIsolateComponentNode.h
--- a/Engine/Plugins/Experimental/ChaosFlesh/Source/ChaosFleshNodes/Public/Dataflow/ChaosFleshIsolateComponentNode.h
+++ b/Engine/Plugins/Experimental/ChaosFlesh/Source/ChaosFleshNodes/Public/Dataflow/ChaosFleshIsolateComponentNode.h
@@ -14,7 +14,7 @@ struct FIsolateComponentNode : public FDataflowNode
 {
 	GENERATED_USTRUCT_BODY()
 	DATAFLOW_NODE_DEFINE_INTERNAL(FIsolateComponentNode, "IsolateComponent", "Flesh", "")
-	DATAFLOW_NODE_RENDER_TYPE(FGeometryCollection::StaticType(), "Collection")
+	DATAFLOW_NODE_RENDER_TYPE("SurfaceRender", FGeometryCollection::StaticType(), "Collection")
 
 public:
 	//typedef FManagedArrayCollection DataType;
diff --git a/Engine/Plugins/Experimental/ChaosFlesh/Source/ChaosFleshNodes/Public/Dataflow/ChaosFleshKinematicConstraintNode.h b/Engine/Plugins/Experimental/ChaosFlesh/Source/ChaosFleshNodes/Public/Dataflow/ChaosFleshKinematicConstraintNode.h
--- a/Engine/Plugins/Experimental/ChaosFlesh/Source/ChaosFleshNodes/Public/Dataflow/ChaosFleshKinematicConstraintNode.h
+++ b/Engine/Plugins/Experimental/ChaosFlesh/Source/ChaosFleshNodes/Public/Dataflow/ChaosFleshKinematicConstraintNode.h
@@ -16,7 +16,7 @@ struct FKinematicInitializationDataflowNode : public FDataflowNode
 {
 	GENERATED_USTRUCT_BODY()
 	DATAFLOW_NODE_DEFINE_INTERNAL(FKinematicInitializationDataflowNode, "KinematicInitialization", "Flesh", "")
-	DATAFLOW_NODE_RENDER_TYPE(FGeometryCollection::StaticType(), "Collection")
+	DATAFLOW_NODE_RENDER_TYPE("SurfaceRender", FGeometryCollection::StaticType(), "Collection")
 
 public:
 	typedef FManagedArrayCollection DataType;
diff --git a/Engine/Plugins/Experimental/ChaosFlesh/Source/ChaosFleshNodes/Public/Dataflow/ChaosFleshKinematicMuscleAttachmentsNode.h b/Engine/Plugins/Experimental/ChaosFlesh/Source/ChaosFleshNodes/Public/Dataflow/ChaosFleshKinematicMuscleAttachmentsNode.h
--- a/Engine/Plugins/Experimental/ChaosFlesh/Source/ChaosFleshNodes/Public/Dataflow/ChaosFleshKinematicMuscleAttachmentsNode.h
+++ b/Engine/Plugins/Experimental/ChaosFlesh/Source/ChaosFleshNodes/Public/Dataflow/ChaosFleshKinematicMuscleAttachmentsNode.h
@@ -14,7 +14,7 @@ struct FKinematicMuscleAttachmentsDataflowNode : public FDataflowNode
 {
 	GENERATED_USTRUCT_BODY()
 	DATAFLOW_NODE_DEFINE_INTERNAL(FKinematicMuscleAttachmentsDataflowNode, "KinematicMuscleAttachments", "Flesh", "")
-	DATAFLOW_NODE_RENDER_TYPE(FGeometryCollection::StaticType(), "Collection")
+	DATAFLOW_NODE_RENDER_TYPE("SurfaceRender", FGeometryCollection::StaticType(), "Collection")
 
 public:
 	typedef FManagedArrayCollection DataType;
diff --git a/Engine/Plugins/Experimental/ChaosFlesh/Source/ChaosFleshNodes/Public/Dataflow/ChaosFleshKinematicSkeletonConstraintNode.h b/Engine/Plugins/Experimental/ChaosFlesh/Source/ChaosFleshNodes/Public/Dataflow/ChaosFleshKinematicSkeletonConstraintNode.h
--- a/Engine/Plugins/Experimental/ChaosFlesh/Source/ChaosFleshNodes/Public/Dataflow/ChaosFleshKinematicSkeletonConstraintNode.h
+++ b/Engine/Plugins/Experimental/ChaosFlesh/Source/ChaosFleshNodes/Public/Dataflow/ChaosFleshKinematicSkeletonConstraintNode.h
@@ -15,7 +15,7 @@ struct FKinematicSkeletonConstraintDataflowNode : public FDataflowNode
 {
 	GENERATED_USTRUCT_BODY()
 	DATAFLOW_NODE_DEFINE_INTERNAL(FKinematicSkeletonConstraintDataflowNode, "KinematicSkeletonConstraint", "Flesh", "")
-	DATAFLOW_NODE_RENDER_TYPE(FGeometryCollection::StaticType(), "Collection")
+	DATAFLOW_NODE_RENDER_TYPE("SurfaceRender", FGeometryCollection::StaticType(), "Collection")
 
 public:
 	typedef FManagedArrayCollection DataType;
diff --git a/Engine/Plugins/Experimental/ChaosFlesh/Source/ChaosFleshNodes/Public/Dataflow/ChaosFleshRadialTetrahedronNode.h b/Engine/Plugins/Experimental/ChaosFlesh/Source/ChaosFleshNodes/Public/Dataflow/ChaosFleshRadialTetrahedronNode.h
--- a/Engine/Plugins/Experimental/ChaosFlesh/Source/ChaosFleshNodes/Public/Dataflow/ChaosFleshRadialTetrahedronNode.h
+++ b/Engine/Plugins/Experimental/ChaosFlesh/Source/ChaosFleshNodes/Public/Dataflow/ChaosFleshRadialTetrahedronNode.h
@@ -17,7 +17,7 @@ struct FRadialTetrahedronDataflowNodes : public FDataflowNode
 {
 	GENERATED_USTRUCT_BODY()
 	DATAFLOW_NODE_DEFINE_INTERNAL(FRadialTetrahedronDataflowNodes, "RadialTetrahedron", "Flesh", "")
-	DATAFLOW_NODE_RENDER_TYPE(FGeometryCollection::StaticType(), "Collection")
+	DATAFLOW_NODE_RENDER_TYPE("SurfaceRender", FGeometryCollection::StaticType(), "Collection")
 
 public:
 	typedef FManagedArrayCollection DataType;
diff --git a/Engine/Plugins/Experimental/ChaosFlesh/Source/ChaosFleshNodes/Public/Dataflow/ChaosFleshSetFleshBonePositionTargetBindingNode.h b/Engine/Plugins/Experimental/ChaosFlesh/Source/ChaosFleshNodes/Public/Dataflow/ChaosFleshSetFleshBonePositionTargetBindingNode.h
--- a/Engine/Plugins/Experimental/ChaosFlesh/Source/ChaosFleshNodes/Public/Dataflow/ChaosFleshSetFleshBonePositionTargetBindingNode.h
+++ b/Engine/Plugins/Experimental/ChaosFlesh/Source/ChaosFleshNodes/Public/Dataflow/ChaosFleshSetFleshBonePositionTargetBindingNode.h
@@ -24,7 +24,7 @@ struct FSetFleshBonePositionTargetBindingDataflowNode : public FDataflowNode
 {
 	GENERATED_USTRUCT_BODY()
 		DATAFLOW_NODE_DEFINE_INTERNAL(FSetFleshBonePositionTargetBindingDataflowNode, "SetFleshBonePositionTargetBinding", "Flesh", "")
-		DATAFLOW_NODE_RENDER_TYPE(FGeometryCollection::StaticType(), "Collection")
+		DATAFLOW_NODE_RENDER_TYPE("SurfaceRender",FGeometryCollection::StaticType(),  "Collection")
 
 public:
 	typedef FManagedArrayCollection DataType;
diff --git a/Engine/Plugins/Experimental/ChaosFlesh/Source/ChaosFleshNodes/Public/Dataflow/ChaosFleshSetFleshDefaultPropertiesNode.h b/Engine/Plugins/Experimental/ChaosFlesh/Source/ChaosFleshNodes/Public/Dataflow/ChaosFleshSetFleshDefaultPropertiesNode.h
--- a/Engine/Plugins/Experimental/ChaosFlesh/Source/ChaosFleshNodes/Public/Dataflow/ChaosFleshSetFleshDefaultPropertiesNode.h
+++ b/Engine/Plugins/Experimental/ChaosFlesh/Source/ChaosFleshNodes/Public/Dataflow/ChaosFleshSetFleshDefaultPropertiesNode.h
@@ -15,7 +15,7 @@ struct FSetFleshDefaultPropertiesNode : public FDataflowNode
 {
 	GENERATED_USTRUCT_BODY()
 	DATAFLOW_NODE_DEFINE_INTERNAL(FSetFleshDefaultPropertiesNode, "SetFleshDefaultProperties", "Flesh", "")
-	DATAFLOW_NODE_RENDER_TYPE(FGeometryCollection::StaticType(), "Collection")
+	DATAFLOW_NODE_RENDER_TYPE("SurfaceRender",FGeometryCollection::StaticType(),  "Collection")
 
 public:
 
diff --git a/Engine/Plugins/Experimental/ChaosFlesh/Source/ChaosFleshNodes/Public/Dataflow/ChaosFleshSetVertexTetrahedraPositionTargetBindingNode.h b/Engine/Plugins/Experimental/ChaosFlesh/Source/ChaosFleshNodes/Public/Dataflow/ChaosFleshSetVertexTetrahedraPositionTargetBindingNode.h
--- a/Engine/Plugins/Experimental/ChaosFlesh/Source/ChaosFleshNodes/Public/Dataflow/ChaosFleshSetVertexTetrahedraPositionTargetBindingNode.h
+++ b/Engine/Plugins/Experimental/ChaosFlesh/Source/ChaosFleshNodes/Public/Dataflow/ChaosFleshSetVertexTetrahedraPositionTargetBindingNode.h
@@ -13,7 +13,7 @@ struct FSetVertexTetrahedraPositionTargetBindingDataflowNode : public FDataflowN
 {
 	GENERATED_USTRUCT_BODY()
 		DATAFLOW_NODE_DEFINE_INTERNAL(FSetVertexTetrahedraPositionTargetBindingDataflowNode, "SetVertexTetrahedraPositionTargetBinding", "Flesh", "")
-		DATAFLOW_NODE_RENDER_TYPE(FGeometryCollection::StaticType(), "Collection")
+		DATAFLOW_NODE_RENDER_TYPE("SurfaceRender",FGeometryCollection::StaticType(),  "Collection")
 
 public:
 	typedef FManagedArrayCollection DataType;
diff --git a/Engine/Plugins/Experimental/ChaosFlesh/Source/ChaosFleshNodes/Public/Dataflow/ChaosFleshSetVertexTrianglePositionTargetBindingNode.h b/Engine/Plugins/Experimental/ChaosFlesh/Source/ChaosFleshNodes/Public/Dataflow/ChaosFleshSetVertexTrianglePositionTargetBindingNode.h
--- a/Engine/Plugins/Experimental/ChaosFlesh/Source/ChaosFleshNodes/Public/Dataflow/ChaosFleshSetVertexTrianglePositionTargetBindingNode.h
+++ b/Engine/Plugins/Experimental/ChaosFlesh/Source/ChaosFleshNodes/Public/Dataflow/ChaosFleshSetVertexTrianglePositionTargetBindingNode.h
@@ -13,7 +13,7 @@ struct FSetVertexTrianglePositionTargetBindingDataflowNode : public FDataflowNod
 {
 	GENERATED_USTRUCT_BODY()
 		DATAFLOW_NODE_DEFINE_INTERNAL(FSetVertexTrianglePositionTargetBindingDataflowNode, "SetVertexTrianglePositionTargetBinding", "Flesh", "")
-		DATAFLOW_NODE_RENDER_TYPE(FGeometryCollection::StaticType(), "Collection")
+		DATAFLOW_NODE_RENDER_TYPE("SurfaceRender",FGeometryCollection::StaticType(),  "Collection")
 
 public:
 	typedef FManagedArrayCollection DataType;
diff --git a/Engine/Plugins/Experimental/ChaosFlesh/Source/ChaosFleshNodes/Public/Dataflow/ChaosFleshSetVertexVertexPositionTargetBindingNode.h b/Engine/Plugins/Experimental/ChaosFlesh/Source/ChaosFleshNodes/Public/Dataflow/ChaosFleshSetVertexVertexPositionTargetBindingNode.h
--- a/Engine/Plugins/Experimental/ChaosFlesh/Source/ChaosFleshNodes/Public/Dataflow/ChaosFleshSetVertexVertexPositionTargetBindingNode.h
+++ b/Engine/Plugins/Experimental/ChaosFlesh/Source/ChaosFleshNodes/Public/Dataflow/ChaosFleshSetVertexVertexPositionTargetBindingNode.h
@@ -14,7 +14,7 @@ struct FSetVertexVertexPositionTargetBindingDataflowNode : public FDataflowNode
 {
 	GENERATED_USTRUCT_BODY()
 	DATAFLOW_NODE_DEFINE_INTERNAL(FSetVertexVertexPositionTargetBindingDataflowNode, "SetVertexVertexPositionTargetBinding", "Flesh", "")
-	DATAFLOW_NODE_RENDER_TYPE(FGeometryCollection::StaticType(), "Collection")
+	DATAFLOW_NODE_RENDER_TYPE("SurfaceRender",FGeometryCollection::StaticType(),  "Collection")
 
 public:
 	typedef FManagedArrayCollection DataType;
diff --git a/Engine/Plugins/Experimental/ChaosFlesh/Source/ChaosFleshNodes/Public/Dataflow/ChaosFleshSkeletalMeshConstraintNode.h b/Engine/Plugins/Experimental/ChaosFlesh/Source/ChaosFleshNodes/Public/Dataflow/ChaosFleshSkeletalMeshConstraintNode.h
--- a/Engine/Plugins/Experimental/ChaosFlesh/Source/ChaosFleshNodes/Public/Dataflow/ChaosFleshSkeletalMeshConstraintNode.h
+++ b/Engine/Plugins/Experimental/ChaosFlesh/Source/ChaosFleshNodes/Public/Dataflow/ChaosFleshSkeletalMeshConstraintNode.h
@@ -15,7 +15,7 @@ struct FKinematicSkeletalMeshInitializationDataflowNode : public FDataflowNode
 {
 	GENERATED_USTRUCT_BODY()
 	DATAFLOW_NODE_DEFINE_INTERNAL(FKinematicSkeletalMeshInitializationDataflowNode, "KinematicSkeletalMeshInitialization", "Flesh", "")
-	DATAFLOW_NODE_RENDER_TYPE(FGeometryCollection::StaticType(), "Collection")
+	DATAFLOW_NODE_RENDER_TYPE("SurfaceRender",FGeometryCollection::StaticType(),  "Collection")
 
 public:
 	typedef FManagedArrayCollection DataType;
diff --git a/Engine/Plugins/Experimental/ChaosFlesh/Source/ChaosFleshNodes/Public/Dataflow/ChaosFleshSkinSimulationPropertiesNode.h b/Engine/Plugins/Experimental/ChaosFlesh/Source/ChaosFleshNodes/Public/Dataflow/ChaosFleshSkinSimulationPropertiesNode.h
--- a/Engine/Plugins/Experimental/ChaosFlesh/Source/ChaosFleshNodes/Public/Dataflow/ChaosFleshSkinSimulationPropertiesNode.h
+++ b/Engine/Plugins/Experimental/ChaosFlesh/Source/ChaosFleshNodes/Public/Dataflow/ChaosFleshSkinSimulationPropertiesNode.h
@@ -18,7 +18,7 @@ struct FSkinSimulationPropertiesDataflowNodes : public FDataflowNode
 {
 	GENERATED_USTRUCT_BODY()
 	DATAFLOW_NODE_DEFINE_INTERNAL(FSkinSimulationPropertiesDataflowNodes, "SkinSimulationProperties", "Flesh", "")
-	DATAFLOW_NODE_RENDER_TYPE(FGeometryCollection::StaticType(), "Collection")
+	DATAFLOW_NODE_RENDER_TYPE("SurfaceRender",FGeometryCollection::StaticType(),  "Collection")
 
 public:
 	typedef FManagedArrayCollection DataType;
diff --git a/Engine/Plugins/Experimental/ChaosFlesh/Source/ChaosFleshNodes/Public/Dataflow/ChaosFleshTriangleMeshSimulationPropertiesNode.h b/Engine/Plugins/Experimental/ChaosFlesh/Source/ChaosFleshNodes/Public/Dataflow/ChaosFleshTriangleMeshSimulationPropertiesNode.h
--- a/Engine/Plugins/Experimental/ChaosFlesh/Source/ChaosFleshNodes/Public/Dataflow/ChaosFleshTriangleMeshSimulationPropertiesNode.h
+++ b/Engine/Plugins/Experimental/ChaosFlesh/Source/ChaosFleshNodes/Public/Dataflow/ChaosFleshTriangleMeshSimulationPropertiesNode.h
@@ -19,7 +19,7 @@ struct FTriangleMeshSimulationPropertiesDataflowNodes : public FDataflowNode
 {
 	GENERATED_USTRUCT_BODY()
 	DATAFLOW_NODE_DEFINE_INTERNAL(FTriangleMeshSimulationPropertiesDataflowNodes, "TriangleMeshSimulationProperties", "Flesh", "")
-	DATAFLOW_NODE_RENDER_TYPE(FGeometryCollection::StaticType(), "Collection")
+	DATAFLOW_NODE_RENDER_TYPE("SurfaceRender",FGeometryCollection::StaticType(),  "Collection")
 
 public:
 	typedef FManagedArrayCollection DataType;
diff --git a/Engine/Plugins/Experimental/ChaosFlesh/Source/ChaosFleshNodes/Public/Dataflow/ChaosFleshVisualizeFiberFieldNode.h b/Engine/Plugins/Experimental/ChaosFlesh/Source/ChaosFleshNodes/Public/Dataflow/ChaosFleshVisualizeFiberFieldNode.h
--- a/Engine/Plugins/Experimental/ChaosFlesh/Source/ChaosFleshNodes/Public/Dataflow/ChaosFleshVisualizeFiberFieldNode.h
+++ b/Engine/Plugins/Experimental/ChaosFlesh/Source/ChaosFleshNodes/Public/Dataflow/ChaosFleshVisualizeFiberFieldNode.h
@@ -17,7 +17,7 @@ struct FVisualizeFiberFieldNode : public FDataflowNode
 {
 	GENERATED_USTRUCT_BODY()
 	DATAFLOW_NODE_DEFINE_INTERNAL(FVisualizeFiberFieldNode, "VisualizeFiberField", "Flesh", "")
-	DATAFLOW_NODE_RENDER_TYPE(FFieldCollection::StaticType(), "VectorField")
+	DATAFLOW_NODE_RENDER_TYPE("VolumeRender",FFieldCollection::StaticType(), "VectorField")
 
 public:
 	UPROPERTY(meta = (DataflowInput, DisplayName = "Collection"))
@@ -47,7 +47,7 @@ struct FVisualizePositionTargetsNode : public FDataflowNode
 {
 	GENERATED_USTRUCT_BODY()
 	DATAFLOW_NODE_DEFINE_INTERNAL(FVisualizePositionTargetsNode, "VisualizePositionTargets", "Flesh", "")
-		DATAFLOW_NODE_RENDER_TYPE(FFieldCollection::StaticType(), "VectorField")
+	DATAFLOW_NODE_RENDER_TYPE("VolumeRender",FFieldCollection::StaticType(), "VectorField")
 
 public:
 	UPROPERTY(meta = (DataflowInput, DisplayName = "Collection"))
@@ -74,7 +74,7 @@ struct FVisualizeKinematicFacesNode : public FDataflowNode
 {
 	GENERATED_USTRUCT_BODY()
 	DATAFLOW_NODE_DEFINE_INTERNAL(FVisualizeKinematicFacesNode, "VisualizeKinematicFaces", "Flesh", "")
-	DATAFLOW_NODE_RENDER_TYPE(FGeometryCollection::StaticType(), "Collection")
+	DATAFLOW_NODE_RENDER_TYPE("SurfaceRender",FGeometryCollection::StaticType(),  "Collection")
 
 public:
 	UPROPERTY(meta = (DataflowInput, DataflowOutput, DisplayName = "Collection", DataflowPassthrough = "Collection"))
diff --git a/Engine/Plugins/Experimental/Dataflow/Source/DataflowEnginePlugin/Private/Dataflow/DataflowEngineRendering.cpp b/Engine/Plugins/Experimental/Dataflow/Source/DataflowEnginePlugin/Private/Dataflow/DataflowEngineRendering.cpp
--- a/Engine/Plugins/Experimental/Dataflow/Source/DataflowEnginePlugin/Private/Dataflow/DataflowEngineRendering.cpp
+++ b/Engine/Plugins/Experimental/Dataflow/Source/DataflowEnginePlugin/Private/Dataflow/DataflowEngineRendering.cpp
@@ -237,7 +237,7 @@ namespace Dataflow
 
 		*		@param Outputs : {FManagedArrayCollection : "Collection"}
 		*/
-		FRenderingFactory::GetInstance()->RegisterOutput(FGeometryCollection::StaticType(),
+		FRenderingFactory::GetInstance()->RegisterOutput({ "SurfaceRender", FGeometryCollection::StaticType()},
 			[](GeometryCollection::Facades::FRenderingFacade& RenderCollection, const Dataflow::FGraphRenderingState& State)
 			{
 				if (State.GetRenderOutputs().Num())
@@ -282,7 +282,7 @@ namespace Dataflow
 
 		*		@param Outputs : {FDynamicMesh3 : "Mesh"}
 		*/
-		FRenderingFactory::GetInstance()->RegisterOutput(FName("FDynamicMesh3"),
+		FRenderingFactory::GetInstance()->RegisterOutput({ "SurfaceRender", FName("FDynamicMesh3") },
 			[](GeometryCollection::Facades::FRenderingFacade& RenderCollection, const Dataflow::FGraphRenderingState& State)
 			{
 				if (State.GetRenderOutputs().Num())
@@ -344,7 +344,7 @@ namespace Dataflow
 
 		*		@param Outputs : {FBox : "Box"}
 		*/
-		FRenderingFactory::GetInstance()->RegisterOutput(FName("FBox"),
+		FRenderingFactory::GetInstance()->RegisterOutput({ "SurfaceRender", FName("FBox") },
 			[](GeometryCollection::Facades::FRenderingFacade& RenderCollection, const Dataflow::FGraphRenderingState& State)
 			{
 				if (State.GetRenderOutputs().Num())
@@ -405,7 +405,7 @@ namespace Dataflow
 
 		*		@param Outputs : {FFieldCollection : "VectorField"}
 		*/
-		FRenderingFactory::GetInstance()->RegisterOutput(FFieldCollection::StaticType(),
+		FRenderingFactory::GetInstance()->RegisterOutput({ "VolumeRender", FFieldCollection::StaticType() },
 			[](GeometryCollection::Facades::FRenderingFacade& RenderCollection, const Dataflow::FGraphRenderingState& State)
 			{
 				if (State.GetRenderOutputs().Num())
diff --git a/Engine/Plugins/Experimental/Dataflow/Source/DataflowNodes/Public/Dataflow/DataflowCollectionAddScalarVertexPropertyNode.h b/Engine/Plugins/Experimental/Dataflow/Source/DataflowNodes/Public/Dataflow/DataflowCollectionAddScalarVertexPropertyNode.h
--- a/Engine/Plugins/Experimental/Dataflow/Source/DataflowNodes/Public/Dataflow/DataflowCollectionAddScalarVertexPropertyNode.h
+++ b/Engine/Plugins/Experimental/Dataflow/Source/DataflowNodes/Public/Dataflow/DataflowCollectionAddScalarVertexPropertyNode.h
@@ -15,7 +15,7 @@ struct DATAFLOWNODES_API FDataflowCollectionAddScalarVertexPropertyNode : public
 {
 	GENERATED_USTRUCT_BODY()
 	DATAFLOW_NODE_DEFINE_INTERNAL(FDataflowCollectionAddScalarVertexPropertyNode, "AddScalarVertexProperty", "Collection", "Add a saved scalar property to a collection")
-	DATAFLOW_NODE_RENDER_TYPE(FGeometryCollection::StaticType(), "Collection")
+	DATAFLOW_NODE_RENDER_TYPE("SurfaceRender", FGeometryCollection::StaticType(), "Collection")
 
 
 public:
diff --git a/Engine/Plugins/Experimental/GeometryCollectionPlugin/Source/GeometryCollectionNodes/Public/Dataflow/GeometryCollectionAssetNodes.h b/Engine/Plugins/Experimental/GeometryCollectionPlugin/Source/GeometryCollectionNodes/Public/Dataflow/GeometryCollectionAssetNodes.h
--- a/Engine/Plugins/Experimental/GeometryCollectionPlugin/Source/GeometryCollectionNodes/Public/Dataflow/GeometryCollectionAssetNodes.h
+++ b/Engine/Plugins/Experimental/GeometryCollectionPlugin/Source/GeometryCollectionNodes/Public/Dataflow/GeometryCollectionAssetNodes.h
@@ -18,7 +18,7 @@ struct FGeometryCollectionTerminalDataflowNode : public FDataflowTerminalNode
 {
 	GENERATED_USTRUCT_BODY()
 	DATAFLOW_NODE_DEFINE_INTERNAL(FGeometryCollectionTerminalDataflowNode, "GeometryCollectionTerminal", "Terminal", "")
-	DATAFLOW_NODE_RENDER_TYPE(FGeometryCollection::StaticType(), "Collection")
+	DATAFLOW_NODE_RENDER_TYPE("SurfaceRender",FGeometryCollection::StaticType(),  "Collection")
 
 public:
 
@@ -130,7 +130,7 @@ struct FGeometryCollectionToCollectionDataflowNode : public FDataflowNode
 {
 	GENERATED_USTRUCT_BODY()
 	DATAFLOW_NODE_DEFINE_INTERNAL(FGeometryCollectionToCollectionDataflowNode, "GeometryCollectionToCollection", "GeometryCollection|Asset", "")
-	DATAFLOW_NODE_RENDER_TYPE(FGeometryCollection::StaticType(), "Collection")
+	DATAFLOW_NODE_RENDER_TYPE("SurfaceRender",FGeometryCollection::StaticType(),  "Collection")
 
 public:
 	/** Asset input */
@@ -166,7 +166,7 @@ public:
  {
  	GENERATED_USTRUCT_BODY()
  	DATAFLOW_NODE_DEFINE_INTERNAL(FBlueprintToCollectionDataflowNode, "BlueprintToCollection", "GeometryCollection|Asset", "")
-	DATAFLOW_NODE_RENDER_TYPE(FGeometryCollection::StaticType(), "Collection")
+	DATAFLOW_NODE_RENDER_TYPE("SurfaceRender",FGeometryCollection::StaticType(),  "Collection")
 
  public:
  	/** Asset input */
diff --git a/Engine/Plugins/Experimental/GeometryCollectionPlugin/Source/GeometryCollectionNodes/Public/Dataflow/GeometryCollectionFracturingNodes.h b/Engine/Plugins/Experimental/GeometryCollectionPlugin/Source/GeometryCollectionNodes/Public/Dataflow/GeometryCollectionFracturingNodes.h
--- a/Engine/Plugins/Experimental/GeometryCollectionPlugin/Source/GeometryCollectionNodes/Public/Dataflow/GeometryCollectionFracturingNodes.h
+++ b/Engine/Plugins/Experimental/GeometryCollectionPlugin/Source/GeometryCollectionNodes/Public/Dataflow/GeometryCollectionFracturingNodes.h
@@ -279,7 +279,7 @@ struct FExplodedViewDataflowNode : public FDataflowNode
 {
 	GENERATED_USTRUCT_BODY()
 	DATAFLOW_NODE_DEFINE_INTERNAL(FExplodedViewDataflowNode, "ExplodedView", "GeometryCollection|Fracture|Utilities", "")
-	DATAFLOW_NODE_RENDER_TYPE(FGeometryCollection::StaticType(), "Collection")
+	DATAFLOW_NODE_RENDER_TYPE("SurfaceRender",FGeometryCollection::StaticType(),  "Collection")
 
 public:
 	UPROPERTY(meta = (DataflowInput, DataflowOutput, DataflowPassthrough = "Collection", DataflowIntrinsic))
diff --git a/Engine/Plugins/Experimental/GeometryCollectionPlugin/Source/GeometryCollectionNodes/Public/Dataflow/GeometryCollectionMakeNodes.h b/Engine/Plugins/Experimental/GeometryCollectionPlugin/Source/GeometryCollectionNodes/Public/Dataflow/GeometryCollectionMakeNodes.h
--- a/Engine/Plugins/Experimental/GeometryCollectionPlugin/Source/GeometryCollectionNodes/Public/Dataflow/GeometryCollectionMakeNodes.h
+++ b/Engine/Plugins/Experimental/GeometryCollectionPlugin/Source/GeometryCollectionNodes/Public/Dataflow/GeometryCollectionMakeNodes.h
@@ -90,7 +90,7 @@ struct FMakeBoxDataflowNode : public FDataflowNode
 {
 	GENERATED_USTRUCT_BODY()
 	DATAFLOW_NODE_DEFINE_INTERNAL(FMakeBoxDataflowNode, "MakeBox", "Generators|Box", "")
-	DATAFLOW_NODE_RENDER_TYPE(FName("FBox"), "Box")
+	DATAFLOW_NODE_RENDER_TYPE("SurfaceRender",FName("FBox"), "Box")
 
 public:
 	UPROPERTY(EditAnywhere, Category = "Box", meta = (DisplayName = "Input Data Type"));
diff --git a/Engine/Plugins/Experimental/GeometryCollectionPlugin/Source/GeometryCollectionNodes/Public/Dataflow/GeometryCollectionMeshNodes.h b/Engine/Plugins/Experimental/GeometryCollectionPlugin/Source/GeometryCollectionNodes/Public/Dataflow/GeometryCollectionMeshNodes.h
--- a/Engine/Plugins/Experimental/GeometryCollectionPlugin/Source/GeometryCollectionNodes/Public/Dataflow/GeometryCollectionMeshNodes.h
+++ b/Engine/Plugins/Experimental/GeometryCollectionPlugin/Source/GeometryCollectionNodes/Public/Dataflow/GeometryCollectionMeshNodes.h
@@ -61,7 +61,7 @@ struct FBoxToMeshDataflowNode : public FDataflowNode
 {
 	GENERATED_USTRUCT_BODY()
 	DATAFLOW_NODE_DEFINE_INTERNAL(FBoxToMeshDataflowNode, "BoxToMesh", "Mesh|Utilities", "")
-	DATAFLOW_NODE_RENDER_TYPE(FName("FDynamicMesh3"), "Mesh")
+	DATAFLOW_NODE_RENDER_TYPE("SurfaceRender",FName("FDynamicMesh3"), "Mesh")
 
 public:
 	/** BoundingBox input */
@@ -131,7 +131,7 @@ struct FMeshToCollectionDataflowNode : public FDataflowNode
 {
 	GENERATED_USTRUCT_BODY()
 	DATAFLOW_NODE_DEFINE_INTERNAL(FMeshToCollectionDataflowNode, "MeshToCollection", "Mesh|Utilities", "")
-	DATAFLOW_NODE_RENDER_TYPE(FGeometryCollection::StaticType(), "Collection")
+	DATAFLOW_NODE_RENDER_TYPE("SurfaceRender",FGeometryCollection::StaticType(),  "Collection")
 
 public:
 	/** DynamicMesh to convert */
@@ -164,7 +164,7 @@ struct FCollectionToMeshDataflowNode : public FDataflowNode
 {
 	GENERATED_USTRUCT_BODY()
 	DATAFLOW_NODE_DEFINE_INTERNAL(FCollectionToMeshDataflowNode, "CollectionToMesh", "Mesh|Utilities", "")
-	DATAFLOW_NODE_RENDER_TYPE(FName("FDynamicMesh3"), "Mesh")
+	DATAFLOW_NODE_RENDER_TYPE("SurfaceRender",FName("FDynamicMesh3"), "Mesh")
 
 public:
 	/** Collection to convert*/
@@ -200,7 +200,7 @@ struct FStaticMeshToMeshDataflowNode : public FDataflowNode
 {
 	GENERATED_USTRUCT_BODY()
 	DATAFLOW_NODE_DEFINE_INTERNAL(FStaticMeshToMeshDataflowNode, "StaticMeshToMesh", "Mesh|Utilities", "")
-	DATAFLOW_NODE_RENDER_TYPE(FName("FDynamicMesh3"), "Mesh")
+	DATAFLOW_NODE_RENDER_TYPE("SurfaceRender",FName("FDynamicMesh3"), "Mesh")
 
 public:
 	/** StaticMesh to convert */
@@ -240,7 +240,7 @@ struct FMeshAppendDataflowNode : public FDataflowNode
 {
 	GENERATED_USTRUCT_BODY()
 	DATAFLOW_NODE_DEFINE_INTERNAL(FMeshAppendDataflowNode, "MeshAppend", "Mesh|Utilities", "")
-	DATAFLOW_NODE_RENDER_TYPE(FName("FDynamicMesh3"), "Mesh")
+	DATAFLOW_NODE_RENDER_TYPE("SurfaceRender",FName("FDynamicMesh3"), "Mesh")
 
 public:
 	/** Mesh input */
@@ -289,7 +289,7 @@ struct FMeshBooleanDataflowNode : public FDataflowNode
 {
 	GENERATED_USTRUCT_BODY()
 	DATAFLOW_NODE_DEFINE_INTERNAL(FMeshBooleanDataflowNode, "MeshBoolean", "Mesh|Utilities", "")
-	DATAFLOW_NODE_RENDER_TYPE(FName("FDynamicMesh3"), "Mesh")
+	DATAFLOW_NODE_RENDER_TYPE("SurfaceRender",FName("FDynamicMesh3"), "Mesh")
 
 public:
 	/** Boolean operation */
@@ -331,7 +331,7 @@ struct FMeshCopyToPointsDataflowNode : public FDataflowNode
 {
 	GENERATED_USTRUCT_BODY()
 	DATAFLOW_NODE_DEFINE_INTERNAL(FMeshCopyToPointsDataflowNode, "MeshCopyToPoints", "Mesh|Utilities", "")
-	DATAFLOW_NODE_RENDER_TYPE(FName("FDynamicMesh3"), "Mesh")
+	DATAFLOW_NODE_RENDER_TYPE("SurfaceRender",FName("FDynamicMesh3"), "Mesh")
 
 public:
 	/** Points to copy meshes onto */
diff --git a/Engine/Plugins/Experimental/GeometryCollectionPlugin/Source/GeometryCollectionNodes/Public/Dataflow/GeometryCollectionNodes.h b/Engine/Plugins/Experimental/GeometryCollectionPlugin/Source/GeometryCollectionNodes/Public/Dataflow/GeometryCollectionNodes.h
--- a/Engine/Plugins/Experimental/GeometryCollectionPlugin/Source/GeometryCollectionNodes/Public/Dataflow/GeometryCollectionNodes.h
+++ b/Engine/Plugins/Experimental/GeometryCollectionPlugin/Source/GeometryCollectionNodes/Public/Dataflow/GeometryCollectionNodes.h
@@ -172,7 +172,7 @@ struct FBoundingBoxDataflowNode : public FDataflowNode
 {
 	GENERATED_USTRUCT_BODY()
 	DATAFLOW_NODE_DEFINE_INTERNAL(FBoundingBoxDataflowNode, "BoundingBox", "Utilities|Box", "")
-	DATAFLOW_NODE_RENDER_TYPE(FName("FBox"), "BoundingBox")
+	DATAFLOW_NODE_RENDER_TYPE("SurfaceRender",FName("FBox"), "BoundingBox")
 
 public:
 	UPROPERTY(meta = (DataflowInput))
@@ -572,7 +572,7 @@ struct FTransformCollectionDataflowNode : public FDataflowNode
 {
 	GENERATED_USTRUCT_BODY()
 	DATAFLOW_NODE_DEFINE_INTERNAL(FTransformCollectionDataflowNode, "TransformCollection", "Math|Transform", "")
-	DATAFLOW_NODE_RENDER_TYPE(FGeometryCollection::StaticType(), "Collection")
+	DATAFLOW_NODE_RENDER_TYPE("SurfaceRender",FGeometryCollection::StaticType(),  "Collection")
 
 public:
 	/** Output mesh */
@@ -637,7 +637,7 @@ struct FBakeTransformsInCollectionDataflowNode : public FDataflowNode
 {
 	GENERATED_USTRUCT_BODY()
 	DATAFLOW_NODE_DEFINE_INTERNAL(FBakeTransformsInCollectionDataflowNode, "BakeTransformsInCollection", "Math|Transform", "")
-	DATAFLOW_NODE_RENDER_TYPE(FGeometryCollection::StaticType(), "Collection")
+	DATAFLOW_NODE_RENDER_TYPE("SurfaceRender",FGeometryCollection::StaticType(),  "Collection")
 
 public:
 	/** Collection to bake transforms in */
@@ -666,7 +666,7 @@ struct FTransformMeshDataflowNode : public FDataflowNode
 {
 	GENERATED_USTRUCT_BODY()
 	DATAFLOW_NODE_DEFINE_INTERNAL(FTransformMeshDataflowNode, "TransformMesh", "Math|Transform", "")
-	DATAFLOW_NODE_RENDER_TYPE(FName("FDynamicMesh3"), "Mesh")
+	DATAFLOW_NODE_RENDER_TYPE("SurfaceRender",FName("FDynamicMesh3"), "Mesh")
 
 public:
 	/** Output mesh */
diff --git a/Engine/Plugins/Experimental/GeometryCollectionPlugin/Source/GeometryCollectionNodes/Public/Dataflow/GeometryCollectionSkeletalMeshToCollectionNode.h b/Engine/Plugins/Experimental/GeometryCollectionPlugin/Source/GeometryCollectionNodes/Public/Dataflow/GeometryCollectionSkeletalMeshToCollectionNode.h
--- a/Engine/Plugins/Experimental/GeometryCollectionPlugin/Source/GeometryCollectionNodes/Public/Dataflow/GeometryCollectionSkeletalMeshToCollectionNode.h
+++ b/Engine/Plugins/Experimental/GeometryCollectionPlugin/Source/GeometryCollectionNodes/Public/Dataflow/GeometryCollectionSkeletalMeshToCollectionNode.h
@@ -16,7 +16,7 @@ struct FSkeletalMeshToCollectionDataflowNode : public FDataflowNode
 {
 	GENERATED_USTRUCT_BODY()
 	DATAFLOW_NODE_DEFINE_INTERNAL(FSkeletalMeshToCollectionDataflowNode, "SkeletalMeshToCollection", "GeometryCollection", "")
-	DATAFLOW_NODE_RENDER_TYPE(FGeometryCollection::StaticType(), "Collection")
+	DATAFLOW_NODE_RENDER_TYPE("SurfaceRender",FGeometryCollection::StaticType(),  "Collection")
 
 public:
 	typedef FManagedArrayCollection DataType;
diff --git a/Engine/Plugins/Experimental/GeometryCollectionPlugin/Source/GeometryCollectionNodes/Public/Dataflow/GeometryCollectionStaticMeshToCollectionNode.h b/Engine/Plugins/Experimental/GeometryCollectionPlugin/Source/GeometryCollectionNodes/Public/Dataflow/GeometryCollectionStaticMeshToCollectionNode.h
--- a/Engine/Plugins/Experimental/GeometryCollectionPlugin/Source/GeometryCollectionNodes/Public/Dataflow/GeometryCollectionStaticMeshToCollectionNode.h
+++ b/Engine/Plugins/Experimental/GeometryCollectionPlugin/Source/GeometryCollectionNodes/Public/Dataflow/GeometryCollectionStaticMeshToCollectionNode.h
@@ -20,7 +20,7 @@ struct FStaticMeshToCollectionDataflowNode : public FDataflowNode
 {
 	GENERATED_USTRUCT_BODY()
 	DATAFLOW_NODE_DEFINE_INTERNAL(FStaticMeshToCollectionDataflowNode, "StaticMeshToCollection", "GeometryCollection|Asset", "")
-	DATAFLOW_NODE_RENDER_TYPE(FGeometryCollection::StaticType(), "Collection")
+	DATAFLOW_NODE_RENDER_TYPE("SurfaceRender",FGeometryCollection::StaticType(),  "Collection")
 
 public:
 	/** Asset input */
diff --git a/Engine/Plugins/Experimental/GeometryCollectionPlugin/Source/GeometryCollectionNodes/Public/Dataflow/GeometryCollectionTriangleBoundaryIndicesNode.h b/Engine/Plugins/Experimental/GeometryCollectionPlugin/Source/GeometryCollectionNodes/Public/Dataflow/GeometryCollectionTriangleBoundaryIndicesNode.h
--- a/Engine/Plugins/Experimental/GeometryCollectionPlugin/Source/GeometryCollectionNodes/Public/Dataflow/GeometryCollectionTriangleBoundaryIndicesNode.h
+++ b/Engine/Plugins/Experimental/GeometryCollectionPlugin/Source/GeometryCollectionNodes/Public/Dataflow/GeometryCollectionTriangleBoundaryIndicesNode.h
@@ -14,7 +14,7 @@ struct FTriangleBoundaryIndicesNode : public FDataflowNode
 {
 	GENERATED_USTRUCT_BODY()
 	DATAFLOW_NODE_DEFINE_INTERNAL(FTriangleBoundaryIndicesNode, "TriangleBoundaryIndices", "Geometry Collection|Utiltites", "")
-		DATAFLOW_NODE_RENDER_TYPE(FGeometryCollection::StaticType(), "Collection")
+		DATAFLOW_NODE_RENDER_TYPE("SurfaceRender",FGeometryCollection::StaticType(),  "Collection")
 		
 
 public:
diff --git a/Engine/Plugins/Experimental/GeometryCollectionPlugin/Source/GeometryCollectionNodes/Public/Dataflow/GeometryCollectionVerticesNodes.h b/Engine/Plugins/Experimental/GeometryCollectionPlugin/Source/GeometryCollectionNodes/Public/Dataflow/GeometryCollectionVerticesNodes.h
--- a/Engine/Plugins/Experimental/GeometryCollectionPlugin/Source/GeometryCollectionNodes/Public/Dataflow/GeometryCollectionVerticesNodes.h
+++ b/Engine/Plugins/Experimental/GeometryCollectionPlugin/Source/GeometryCollectionNodes/Public/Dataflow/GeometryCollectionVerticesNodes.h
@@ -15,7 +15,7 @@ struct FTransformCollectionAttributeDataflowNode : public FDataflowNode
 {
 	GENERATED_USTRUCT_BODY()
 	DATAFLOW_NODE_DEFINE_INTERNAL(FTransformCollectionAttributeDataflowNode, "TransformCollectionAttribute", "GeometryCollection", "")
-	DATAFLOW_NODE_RENDER_TYPE(FGeometryCollection::StaticType(), "Collection")
+	DATAFLOW_NODE_RENDER_TYPE("SurfaceRender",FGeometryCollection::StaticType(),  "Collection")
 
 public:
 	typedef FManagedArrayCollection DataType;
diff --git a/Engine/Source/Runtime/Experimental/Dataflow/Core/Public/Dataflow/DataflowNode.h b/Engine/Source/Runtime/Experimental/Dataflow/Core/Public/Dataflow/DataflowNode.h
--- a/Engine/Source/Runtime/Experimental/Dataflow/Core/Public/Dataflow/DataflowNode.h
+++ b/Engine/Source/Runtime/Experimental/Dataflow/Core/Public/Dataflow/DataflowNode.h
@@ -584,10 +584,10 @@ private:
 #define DATAFLOW_NODE_REGISTER_CREATION_FACTORY(A)									\
 	::Dataflow::FNodeFactory::RegisterNodeFromType<A>();
 
-#define DATAFLOW_NODE_RENDER_TYPE(A, B)												\
+#define DATAFLOW_NODE_RENDER_TYPE(A, B, C)												\
 	virtual TArray<::Dataflow::FRenderingParameter> GetRenderParametersImpl() const {		\
 		TArray<::Dataflow::FRenderingParameter> Array;								\
-		Array.Add({ A, {B,} });														\
+		Array.Add({ A, B, {C,} });														\
 		return Array;}
 
 #define DATAFLOW_NODE_DEFINE_INTERNAL(TYPE, DISPLAY_NAME, CATEGORY, TAGS)			\
diff --git a/Engine/Source/Runtime/Experimental/Dataflow/Core/Public/Dataflow/DataflowNodeParameters.h b/Engine/Source/Runtime/Experimental/Dataflow/Core/Public/Dataflow/DataflowNodeParameters.h
--- a/Engine/Source/Runtime/Experimental/Dataflow/Core/Public/Dataflow/DataflowNodeParameters.h
+++ b/Engine/Source/Runtime/Experimental/Dataflow/Core/Public/Dataflow/DataflowNodeParameters.h
@@ -38,11 +38,12 @@ namespace Dataflow
 
 	struct FRenderingParameter {
 		FRenderingParameter() {}
-		FRenderingParameter(FName InTypeName, const TArray<FName>& InOutputs)
-			: Type(InTypeName), Outputs(InOutputs) {}
-		FRenderingParameter(FName InTypeName, TArray<FName>&& InOutputs)
-			: Type(InTypeName), Outputs(InOutputs) {}
+		FRenderingParameter(FString InRenderName, FName InTypeName, const TArray<FName>& InOutputs)
+			: Name(InRenderName), Type(InTypeName), Outputs(InOutputs) {}
+		FRenderingParameter(FString InRenderName, FName InTypeName, TArray<FName>&& InOutputs)
+			: Name(InRenderName), Type(InTypeName), Outputs(InOutputs) {}
 
+		FString Name = FString("");
 		FName Type = FName("");
 		TArray<FName> Outputs;
 	};
diff --git a/Engine/Source/Runtime/Experimental/Dataflow/Engine/Private/Dataflow/DataflowRenderingFactory.cpp b/Engine/Source/Runtime/Experimental/Dataflow/Engine/Private/Dataflow/DataflowRenderingFactory.cpp
--- a/Engine/Source/Runtime/Experimental/Dataflow/Engine/Private/Dataflow/DataflowRenderingFactory.cpp
+++ b/Engine/Source/Runtime/Experimental/Dataflow/Engine/Private/Dataflow/DataflowRenderingFactory.cpp
@@ -13,16 +13,17 @@ namespace Dataflow
 
 	void FRenderingFactory::RenderNodeOutput(GeometryCollection::Facades::FRenderingFacade& RenderData, const FGraphRenderingState& State)
 	{
-		FName OutputType = State.GetRenderType();
-		if (RenderMap.Contains(State.GetRenderType()))
+		if (RenderMap.Contains(State.GetRenderKey()))
 		{
-			RenderMap[State.GetRenderType()](RenderData, State);
+			RenderMap[State.GetRenderKey()](RenderData, State);
 		}
 		else
 		{
 			UE_LOG(LogChaos, Warning,
-				TEXT("Warning : Dataflow missing output rendering type(%s)"), *State.GetRenderType().ToString());
-
+				TEXT("Warning : Dataflow missing output renderer <%s,%s> for node %s"), 
+				*State.GetRenderKey().Get<0>(),
+				*State.GetRenderKey().Get<1>().ToString(),
+				*State.GetNodeName().ToString());
 		}
 	}
 }
diff --git a/Engine/Source/Runtime/Experimental/Dataflow/Engine/Public/Dataflow/DataflowRenderingFactory.h b/Engine/Source/Runtime/Experimental/Dataflow/Engine/Public/Dataflow/DataflowRenderingFactory.h
--- a/Engine/Source/Runtime/Experimental/Dataflow/Engine/Public/Dataflow/DataflowRenderingFactory.h
+++ b/Engine/Source/Runtime/Experimental/Dataflow/Engine/Public/Dataflow/DataflowRenderingFactory.h
@@ -12,18 +12,21 @@
 namespace Dataflow
 {
 	class FContext;
+	typedef TPair<FString, FName> FRenderKey;
 
 	struct FGraphRenderingState {
 		FGraphRenderingState(const FGuid InGuid, const FDataflowNode* InNode, const FRenderingParameter& InParameters, Dataflow::FContext& InContext)
 			: NodeGuid(InGuid)
 			, Node(InNode)
+			, RenderName(InParameters.Name)
 			, RenderType(InParameters.Type)
 			, RenderOutputs(InParameters.Outputs)
 			, Context(InContext)
 		{}
 
 		const FGuid& GetGuid() const { return NodeGuid; }
-		const FName& GetRenderType() const { return RenderType; }
+		FName GetNodeName() const { return Node?Node->GetName():FName(); }
+		FRenderKey GetRenderKey() const { return { RenderName,RenderType }; }
 		const TArray<FName>& GetRenderOutputs() const { return RenderOutputs; }
 
 		template<class T>
@@ -43,6 +46,7 @@ namespace Dataflow
 		const FGuid NodeGuid;
 		const FDataflowNode* Node = nullptr;
 
+		FString RenderName;
 		FName RenderType;
 		TArray<FName> RenderOutputs;
 
@@ -57,7 +61,7 @@ namespace Dataflow
 		typedef TFunction<void(GeometryCollection::Facades::FRenderingFacade& RenderData, const FGraphRenderingState& State)> FOutputRenderingFunction;
 
 		// All Maps indexed by TypeName
-		TMap<FName, FOutputRenderingFunction > RenderMap;		// [TypeName] -> NewNodeFunction
+		TMap<FRenderKey, FOutputRenderingFunction > RenderMap;		// [TypeName] -> NewNodeFunction
 		DATAFLOWENGINE_API static FRenderingFactory* Instance;
 		FRenderingFactory() {}
 
@@ -73,23 +77,25 @@ namespace Dataflow
 			return Instance;
 		}
 
-		void RegisterOutput(const FName& Type, FOutputRenderingFunction InFunction)
+		void RegisterOutput(const FRenderKey& Key, FOutputRenderingFunction InFunction)
 		{
-			if (RenderMap.Contains(Type))
+			if (RenderMap.Contains(Key))
 			{
 				UE_LOG(LogChaos, Warning,
 					TEXT("Warning : Dataflow output rendering registration conflicts with "
-						"existing type(%s)"), *Type.ToString());
+						"existing renderer(<%s,%s>)"), 
+					*Key.Get<0>(),
+					*Key.Get<1>().ToString());
 			}
 			else
 			{
-				RenderMap.Add(Type, InFunction);
+				RenderMap.Add(Key, InFunction);
 			}
 		}
 
 		DATAFLOWENGINE_API void RenderNodeOutput(GeometryCollection::Facades::FRenderingFacade& RenderData, const FGraphRenderingState& State);
 
-		bool Contains(FName InType) const { return RenderMap.Contains(InType); }
+		bool Contains(const FRenderKey& InKey) const { return RenderMap.Contains(InKey); }
 
 	};
 
