GeometryCollection : Restrict editability in editor with cvar 'p.Chaos.GC.IsEditable false'
Dataflow :  Restrict editability in editor with cvar 'p.Dataflow.IsEditable false'

diff --git a/Engine/Plugins/Experimental/Dataflow/Source/DataflowEditor/Private/Dataflow/AssetDefinition_DataflowAsset.cpp b/Engine/Plugins/Experimental/Dataflow/Source/DataflowEditor/Private/Dataflow/AssetDefinition_DataflowAsset.cpp
--- a/Engine/Plugins/Experimental/Dataflow/Source/DataflowEditor/Private/Dataflow/AssetDefinition_DataflowAsset.cpp
+++ b/Engine/Plugins/Experimental/Dataflow/Source/DataflowEditor/Private/Dataflow/AssetDefinition_DataflowAsset.cpp
@@ -16,6 +16,10 @@
 
 #define LOCTEXT_NAMESPACE "AssetActions_DataflowAsset"
 
+
+bool bCanEditDataflow = true;
+FAutoConsoleVariableRef CVarDataflowIsEditable(TEXT("p.Dataflow.IsEditable"), bCanEditDataflow, TEXT("Whether to allow edits of the dataflow [def:true]"));
+
 namespace DataflowAssetDefinitionHelpers
 {
 	// Return true if we should proceed, false if we should re-open the dialog
@@ -169,6 +173,16 @@ UThumbnailInfo* UAssetDefinition_DataflowAsset::LoadThumbnailInfo(const FAssetDa
 	return UE::Editor::FindOrCreateThumbnailInfo(InAsset.GetAsset(), USceneThumbnailInfo::StaticClass());
 }
 
+FAssetOpenSupport UAssetDefinition_DataflowAsset::GetAssetOpenSupport(const FAssetOpenSupportArgs& OpenSupportArgs) const
+{
+	if (bCanEditDataflow)
+	{
+		return Super::GetAssetOpenSupport(OpenSupportArgs);
+	}
+	return FAssetOpenSupport(EAssetOpenMethod::View, false);
+}
+
+
 EAssetCommandResult UAssetDefinition_DataflowAsset::OpenAssets(const FAssetOpenArgs& OpenArgs) const
 {
 	TArray<UDataflow*> DataflowObjects = OpenArgs.LoadObjects<UDataflow>();
diff --git a/Engine/Plugins/Experimental/Dataflow/Source/DataflowEditor/Public/Dataflow/AssetDefinition_DataflowAsset.h b/Engine/Plugins/Experimental/Dataflow/Source/DataflowEditor/Public/Dataflow/AssetDefinition_DataflowAsset.h
--- a/Engine/Plugins/Experimental/Dataflow/Source/DataflowEditor/Public/Dataflow/AssetDefinition_DataflowAsset.h
+++ b/Engine/Plugins/Experimental/Dataflow/Source/DataflowEditor/Public/Dataflow/AssetDefinition_DataflowAsset.h
@@ -108,6 +108,7 @@ private:
 	virtual TConstArrayView<FAssetCategoryPath> GetAssetCategories() const override;
 	virtual UThumbnailInfo* LoadThumbnailInfo(const FAssetData& InAssetData) const override;
 
+	virtual	FAssetOpenSupport GetAssetOpenSupport(const FAssetOpenSupportArgs& OpenSupportArgs) const;
 	virtual EAssetCommandResult OpenAssets(const FAssetOpenArgs& OpenArgs) const override;
 };
 
diff --git a/Engine/Plugins/Experimental/GeometryCollectionPlugin/Source/GeometryCollectionEditor/Private/GeometryCollection/AssetDefinition_GeometryCollection.cpp b/Engine/Plugins/Experimental/GeometryCollectionPlugin/Source/GeometryCollectionEditor/Private/GeometryCollection/AssetDefinition_GeometryCollection.cpp
--- a/Engine/Plugins/Experimental/GeometryCollectionPlugin/Source/GeometryCollectionEditor/Private/GeometryCollection/AssetDefinition_GeometryCollection.cpp
+++ b/Engine/Plugins/Experimental/GeometryCollectionPlugin/Source/GeometryCollectionEditor/Private/GeometryCollection/AssetDefinition_GeometryCollection.cpp
@@ -12,6 +12,8 @@
 
 #define LOCTEXT_NAMESPACE "AssetActions_GeometryCollection"
 
+bool bCanEditGeometryCollection = true;
+FAutoConsoleVariableRef CVarGeometryCollectionIsEditable(TEXT("p.Chaos.GC.IsEditable"), bCanEditGeometryCollection, TEXT("Whether to allow edits of the geometry collection"));
 
 namespace UE::GeometryCollection
 {
@@ -49,6 +51,16 @@ UThumbnailInfo* UAssetDefinition_GeometryCollection::LoadThumbnailInfo(const FAs
 	return UE::Editor::FindOrCreateThumbnailInfo(InAsset.GetAsset(), USceneThumbnailInfo::StaticClass());
 }
 
+FAssetOpenSupport UAssetDefinition_GeometryCollection::GetAssetOpenSupport(const FAssetOpenSupportArgs& OpenSupportArgs) const
+{
+	if (bCanEditGeometryCollection)
+	{
+		return Super::GetAssetOpenSupport(OpenSupportArgs);
+	}
+	return FAssetOpenSupport(EAssetOpenMethod::View, false);
+}
+
+
 EAssetCommandResult UAssetDefinition_GeometryCollection::OpenAssets(const FAssetOpenArgs& OpenArgs) const
 {
 	TArray<UGeometryCollection*> GeometryCollectionObjects = OpenArgs.LoadObjects<UGeometryCollection>();
diff --git a/Engine/Plugins/Experimental/GeometryCollectionPlugin/Source/GeometryCollectionEditor/Public/GeometryCollection/AssetDefinition_GeometryCollection.h b/Engine/Plugins/Experimental/GeometryCollectionPlugin/Source/GeometryCollectionEditor/Public/GeometryCollection/AssetDefinition_GeometryCollection.h
--- a/Engine/Plugins/Experimental/GeometryCollectionPlugin/Source/GeometryCollectionEditor/Public/GeometryCollection/AssetDefinition_GeometryCollection.h
+++ b/Engine/Plugins/Experimental/GeometryCollectionPlugin/Source/GeometryCollectionEditor/Public/GeometryCollection/AssetDefinition_GeometryCollection.h
@@ -19,6 +19,7 @@ private:
 	virtual TConstArrayView<FAssetCategoryPath> GetAssetCategories() const override;
 	virtual UThumbnailInfo* LoadThumbnailInfo(const FAssetData& InAssetData) const override;
 
+	virtual	FAssetOpenSupport GetAssetOpenSupport(const FAssetOpenSupportArgs& OpenSupportArgs) const;
 	virtual EAssetCommandResult OpenAssets(const FAssetOpenArgs& OpenArgs) const override;
 };
 
