Dataflow Geometry Collection Nodes : make SetAnchorState robust to non connection collection input

diff --git a/Engine/Plugins/Experimental/GeometryCollectionPlugin/Source/GeometryCollectionNodes/Private/Dataflow/GeometryCollectionNodes.cpp b/Engine/Plugins/Experimental/GeometryCollectionPlugin/Source/GeometryCollectionNodes/Private/Dataflow/GeometryCollectionNodes.cpp
--- a/Engine/Plugins/Experimental/GeometryCollectionPlugin/Source/GeometryCollectionNodes/Private/Dataflow/GeometryCollectionNodes.cpp
+++ b/Engine/Plugins/Experimental/GeometryCollectionPlugin/Source/GeometryCollectionNodes/Private/Dataflow/GeometryCollectionNodes.cpp
@@ -646,29 +646,36 @@ void FSetAnchorStateDataflowNode::Evaluate(Dataflow::FContext& Context, const FD
 	if (Out->IsA<FManagedArrayCollection>(&Collection))
 	{
 		const FManagedArrayCollection& InCollection = GetValue<FManagedArrayCollection>(Context, &Collection);
-		FDataflowTransformSelection InTransformSelection = GetValue<FDataflowTransformSelection>(Context, &TransformSelection);
-
-		if (TUniquePtr<FGeometryCollection> GeomCollection = TUniquePtr<FGeometryCollection>(InCollection.NewCopy<FGeometryCollection>()))
+		if (IsConnected(&Collection))
 		{
-			Chaos::Facades::FCollectionAnchoringFacade AnchoringFacade(*GeomCollection);
-			if (!AnchoringFacade.HasAnchoredAttribute())
-			{
-				AnchoringFacade.AddAnchoredAttribute();
-			}
-
-			bool bAnchored = (AnchorState == EAnchorStateEnum::Dataflow_AnchorState_Anchored) ? true : false;
-			TArray<int32> BoneIndices;
-			InTransformSelection.AsArray(BoneIndices);
-			AnchoringFacade.SetAnchored(BoneIndices, bAnchored);
+			FDataflowTransformSelection InTransformSelection = GetValue<FDataflowTransformSelection>(Context, &TransformSelection);
 
-			if (bSetNotSelectedBonesToOppositeState)
+			if (TUniquePtr<FGeometryCollection> GeomCollection = TUniquePtr<FGeometryCollection>(InCollection.NewCopy<FGeometryCollection>()))
 			{
-				InTransformSelection.Invert();
+				Chaos::Facades::FCollectionAnchoringFacade AnchoringFacade(*GeomCollection);
+				if (!AnchoringFacade.HasAnchoredAttribute())
+				{
+					AnchoringFacade.AddAnchoredAttribute();
+				}
+
+				bool bAnchored = (AnchorState == EAnchorStateEnum::Dataflow_AnchorState_Anchored) ? true : false;
+				TArray<int32> BoneIndices;
 				InTransformSelection.AsArray(BoneIndices);
-				AnchoringFacade.SetAnchored(BoneIndices, !bAnchored);
-			}
+				AnchoringFacade.SetAnchored(BoneIndices, bAnchored);
 
-			SetValue<const FManagedArrayCollection&>(Context, *GeomCollection, &Collection);
+				if (bSetNotSelectedBonesToOppositeState)
+				{
+					InTransformSelection.Invert();
+					InTransformSelection.AsArray(BoneIndices);
+					AnchoringFacade.SetAnchored(BoneIndices, !bAnchored);
+				}
+
+				SetValue(Context, static_cast<const FManagedArrayCollection&>(*GeomCollection), &Collection);
+			}
+		}
+		else
+		{
+			SetValue(Context, InCollection, &Collection);
 		}
 	}
 }
