Dataflow:
- Modified the transform selection nodes to have Collection output on top of the outputs
- Changed Collection input passthrough to use ForwardInput

diff --git a/Engine/Plugins/Experimental/GeometryCollectionPlugin/Source/GeometryCollectionNodes/Private/Dataflow/GeometryCollectionSelectionNodes.cpp b/Engine/Plugins/Experimental/GeometryCollectionPlugin/Source/GeometryCollectionNodes/Private/Dataflow/GeometryCollectionSelectionNodes.cpp
--- a/Engine/Plugins/Experimental/GeometryCollectionPlugin/Source/GeometryCollectionNodes/Private/Dataflow/GeometryCollectionSelectionNodes.cpp
+++ b/Engine/Plugins/Experimental/GeometryCollectionPlugin/Source/GeometryCollectionNodes/Private/Dataflow/GeometryCollectionSelectionNodes.cpp
@@ -93,8 +93,7 @@ void FCollectionTransformSelectionAllDataflowNode::Evaluate(Dataflow::FContext&
 	}
 	else if (Out->IsA<FManagedArrayCollection>(&Collection))
 	{
-		const FManagedArrayCollection& InCollection = GetValue<FManagedArrayCollection>(Context, &Collection);
-		SetValue(Context, InCollection, &Collection);
+		SafeForwardInput(Context, &Collection, &Collection);
 	}
 }
 
@@ -294,8 +293,7 @@ void FCollectionTransformSelectionNoneDataflowNode::Evaluate(Dataflow::FContext&
 	}
 	else if (Out->IsA<FManagedArrayCollection>(&Collection))
 	{
-		const FManagedArrayCollection& InCollection = GetValue<FManagedArrayCollection>(Context, &Collection);
-		SetValue(Context, InCollection, &Collection);
+		SafeForwardInput(Context, &Collection, &Collection);
 	}
 }
 
@@ -333,8 +331,7 @@ void FCollectionTransformSelectionRandomDataflowNode::Evaluate(Dataflow::FContex
 	}
 	else if (Out->IsA<FManagedArrayCollection>(&Collection))
 	{
-		const FManagedArrayCollection& InCollection = GetValue<FManagedArrayCollection>(Context, &Collection);
-		SetValue(Context, InCollection, &Collection);
+		SafeForwardInput(Context, &Collection, &Collection);
 	}
 }
 
@@ -356,8 +353,7 @@ void FCollectionTransformSelectionRootDataflowNode::Evaluate(Dataflow::FContext&
 	}
 	else if (Out->IsA<FManagedArrayCollection>(&Collection))
 	{
-		const FManagedArrayCollection& InCollection = GetValue<FManagedArrayCollection>(Context, &Collection);
-		SetValue(Context, InCollection, &Collection);
+		SafeForwardInput(Context, &Collection, &Collection);
 	}
 }
 
@@ -407,7 +403,7 @@ void FCollectionTransformSelectionCustomDataflowNode::Evaluate(Dataflow::FContex
 	}
 	else if (Out->IsA<FManagedArrayCollection>(&Collection))
 	{
-		const FManagedArrayCollection& InCollection = GetValue<FManagedArrayCollection>(Context, &Collection);
+		SafeForwardInput(Context, &Collection, &Collection);
 	}
 }
 
@@ -447,8 +443,7 @@ void FCollectionTransformSelectionFromIndexArrayDataflowNode::Evaluate(Dataflow:
 	}
 	else if (Out->IsA(&Collection))
 	{
-		const FManagedArrayCollection& InCollection = GetValue(Context, &Collection);
-		SetValue(Context, InCollection, &Collection);
+		SafeForwardInput(Context, &Collection, &Collection);
 	}
 }
 
@@ -470,8 +465,7 @@ void FCollectionTransformSelectionParentDataflowNode::Evaluate(Dataflow::FContex
 	}
 	else if (Out->IsA<FManagedArrayCollection>(&Collection))
 	{
-		const FManagedArrayCollection& InCollection = GetValue<FManagedArrayCollection>(Context, &Collection);
-		SetValue(Context, InCollection, &Collection);
+		SafeForwardInput(Context, &Collection, &Collection);
 	}
 }
 
@@ -512,8 +506,7 @@ void FCollectionTransformSelectionChildrenDataflowNode::Evaluate(Dataflow::FCont
 	}
 	else if (Out->IsA<FManagedArrayCollection>(&Collection))
 	{
-		const FManagedArrayCollection& InCollection = GetValue<FManagedArrayCollection>(Context, &Collection);
-		SetValue(Context, InCollection, &Collection);
+		SafeForwardInput(Context, &Collection, &Collection);
 	}
 }
 
@@ -535,8 +528,7 @@ void FCollectionTransformSelectionSiblingsDataflowNode::Evaluate(Dataflow::FCont
 	}
 	else if (Out->IsA<FManagedArrayCollection>(&Collection))
 	{
-		const FManagedArrayCollection& InCollection = GetValue<FManagedArrayCollection>(Context, &Collection);
-		SetValue(Context, InCollection, &Collection);
+		SafeForwardInput(Context, &Collection, &Collection);
 	}
 }
 
@@ -558,8 +550,7 @@ void FCollectionTransformSelectionLevelDataflowNode::Evaluate(Dataflow::FContext
 	}
 	else if (Out->IsA<FManagedArrayCollection>(&Collection))
 	{
-		const FManagedArrayCollection& InCollection = GetValue<FManagedArrayCollection>(Context, &Collection);
-		SetValue(Context, InCollection, &Collection);
+		SafeForwardInput(Context, &Collection, &Collection);
 	}
 }
 
@@ -583,8 +574,7 @@ void FCollectionTransformSelectionTargetLevelDataflowNode::Evaluate(Dataflow::FC
 	}
 	else if (Out->IsA(&Collection))
 	{
-		const FManagedArrayCollection& InCollection = GetValue(Context, &Collection);
-		SetValue(Context, InCollection, &Collection);
+		SafeForwardInput(Context, &Collection, &Collection);
 	}
 }
 
@@ -606,8 +596,7 @@ void FCollectionTransformSelectionContactDataflowNode::Evaluate(Dataflow::FConte
 	}
 	else if (Out->IsA<FManagedArrayCollection>(&Collection))
 	{
-		const FManagedArrayCollection& InCollection = GetValue<FManagedArrayCollection>(Context, &Collection);
-		SetValue(Context, InCollection, &Collection);
+		SafeForwardInput(Context, &Collection, &Collection);
 	}
 }
 
@@ -629,8 +618,7 @@ void FCollectionTransformSelectionLeafDataflowNode::Evaluate(Dataflow::FContext&
 	}
 	else if (Out->IsA<FManagedArrayCollection>(&Collection))
 	{
-		const FManagedArrayCollection& InCollection = GetValue<FManagedArrayCollection>(Context, &Collection);
-		SetValue(Context, InCollection, &Collection);
+		SafeForwardInput(Context, &Collection, &Collection);
 	}
 }
 
@@ -652,8 +640,7 @@ void FCollectionTransformSelectionClusterDataflowNode::Evaluate(Dataflow::FConte
 	}
 	else if (Out->IsA<FManagedArrayCollection>(&Collection))
 	{
-		const FManagedArrayCollection& InCollection = GetValue<FManagedArrayCollection>(Context, &Collection);
-		SetValue(Context, InCollection, &Collection);
+		SafeForwardInput(Context, &Collection, &Collection);
 	}
 }
 
@@ -678,8 +665,7 @@ void FCollectionTransformSelectionBySizeDataflowNode::Evaluate(Dataflow::FContex
 	}
 	else if (Out->IsA<FManagedArrayCollection>(&Collection))
 	{
-		const FManagedArrayCollection& InCollection = GetValue<FManagedArrayCollection>(Context, &Collection);
-		SetValue(Context, InCollection, &Collection);
+		SafeForwardInput(Context, &Collection, &Collection);
 	}
 }
 
@@ -704,8 +690,7 @@ void FCollectionTransformSelectionByVolumeDataflowNode::Evaluate(Dataflow::FCont
 	}
 	else if (Out->IsA<FManagedArrayCollection>(&Collection))
 	{
-		const FManagedArrayCollection& InCollection = GetValue<FManagedArrayCollection>(Context, &Collection);
-		SetValue(Context, InCollection, &Collection);
+		SafeForwardInput(Context, &Collection, &Collection);
 	}
 }
 
@@ -742,8 +727,7 @@ void FCollectionTransformSelectionInBoxDataflowNode::Evaluate(Dataflow::FContext
 	}
 	else if (Out->IsA<FManagedArrayCollection>(&Collection))
 	{
-		const FManagedArrayCollection& InCollection = GetValue<FManagedArrayCollection>(Context, &Collection);
-		SetValue(Context, InCollection, &Collection);
+		SafeForwardInput(Context, &Collection, &Collection);
 	}
 }
 
@@ -780,8 +764,7 @@ void FCollectionTransformSelectionInSphereDataflowNode::Evaluate(Dataflow::FCont
 	}
 	else if (Out->IsA<FManagedArrayCollection>(&Collection))
 	{
-		const FManagedArrayCollection& InCollection = GetValue<FManagedArrayCollection>(Context, &Collection);
-		SetValue(Context, InCollection, &Collection);
+		SafeForwardInput(Context, &Collection, &Collection);
 	}
 }
 
@@ -806,8 +789,7 @@ void FCollectionTransformSelectionByFloatAttrDataflowNode::Evaluate(Dataflow::FC
 	}
 	else if (Out->IsA<FManagedArrayCollection>(&Collection))
 	{
-		const FManagedArrayCollection& InCollection = GetValue<FManagedArrayCollection>(Context, &Collection);
-		SetValue(Context, InCollection, &Collection);
+		SafeForwardInput(Context, &Collection, &Collection);
 	}
 }
 
@@ -863,8 +845,7 @@ void FCollectionTransformSelectionByIntAttrDataflowNode::Evaluate(Dataflow::FCon
 	}
 	else if (Out->IsA<FManagedArrayCollection>(&Collection))
 	{
-		const FManagedArrayCollection& InCollection = GetValue<FManagedArrayCollection>(Context, &Collection);
-		SetValue(Context, InCollection, &Collection);
+		SafeForwardInput(Context, &Collection, &Collection);
 	}
 }
 
@@ -914,8 +895,7 @@ void FCollectionVertexSelectionCustomDataflowNode::Evaluate(Dataflow::FContext&
 	}
 	else if (Out->IsA<FManagedArrayCollection>(&Collection))
 	{
-		const FManagedArrayCollection& InCollection = GetValue<FManagedArrayCollection>(Context, &Collection);
-		SetValue(Context, InCollection, &Collection);
+		SafeForwardInput(Context, &Collection, &Collection);
 	}
 }
 
@@ -965,8 +945,7 @@ void FCollectionFaceSelectionCustomDataflowNode::Evaluate(Dataflow::FContext& Co
 	}
 	else if (Out->IsA<FManagedArrayCollection>(&Collection))
 	{
-		const FManagedArrayCollection& InCollection = GetValue<FManagedArrayCollection>(Context, &Collection);
-		SetValue(Context, InCollection, &Collection);
+		SafeForwardInput(Context, &Collection, &Collection);
 	}
 }
 
@@ -1086,8 +1065,7 @@ void FCollectionSelectionConvertDataflowNode::Evaluate(Dataflow::FContext& Conte
 	}
 	else if (Out->IsA<FManagedArrayCollection>(&Collection))
 	{
-		const FManagedArrayCollection& InCollection = GetValue<FManagedArrayCollection>(Context, &Collection);
-		SetValue(Context, InCollection, &Collection);
+		SafeForwardInput(Context, &Collection, &Collection);
 	}
 }
 
