Dataflow:
- Updated existing Dataflow nodes to match parameters and behavior of the Fracture tools
- Changed NoiseSeed parameter to int32
- Versioned up FUniformScatterPointsDataflowNode, FRadialScatterPointsDataflowNode, FVoronoiFractureDataflowNode, FPlaneCutterDataflowNode
- Added FGridScatterPointsDataflowNode, FSliceCutterDataflowNode, FBrickCutterDataflowNode
- Added missing DataflowInput/DataflowOutput declaration check to ValidateConnections()
- Bug fixes

diff --git a/Engine/Plugins/Experimental/Fracture/Source/FractureEngine/Private/FractureEngineFracturing.cpp b/Engine/Plugins/Experimental/Fracture/Source/FractureEngine/Private/FractureEngineFracturing.cpp
--- a/Engine/Plugins/Experimental/Fracture/Source/FractureEngine/Private/FractureEngineFracturing.cpp
+++ b/Engine/Plugins/Experimental/Fracture/Source/FractureEngine/Private/FractureEngineFracturing.cpp
@@ -174,12 +174,33 @@ static float GetMaxVertexMovement(float Grout, float Amplitude, int OctaveNumber
 }
 
 
-void FFractureEngineFracturing::VoronoiFracture(FManagedArrayCollection& InOutCollection,
-	const FDataflowTransformSelection& InTransformSelection,
-	const TArray<FVector>& InSites,
-	float InRandomSeed,
+static void RandomReduceSelection(FDataflowTransformSelection& InOutTransformSelection, int32 InRandomSeed, float InProbToKeep)
+{
+	FRandomStream RandStream(InRandomSeed);
+
+	TArray<int32> SelectedBones = InOutTransformSelection.AsArray();
+
+	for (int32 i = 0; i < SelectedBones.Num(); i++)
+	{
+		if (RandStream.GetFraction() >= InProbToKeep) // range does not include 1, so if ProbToKeep is 1 this will never remove
+		{
+			SelectedBones.RemoveAtSwap(i);
+			i--;
+		}
+	}
+
+	InOutTransformSelection.Clear();
+	InOutTransformSelection.SetFromArray(SelectedBones);
+}
+
+
+int32 FFractureEngineFracturing::VoronoiFracture(FManagedArrayCollection& InOutCollection,
+	FDataflowTransformSelection InTransformSelection,
+	TArray<FVector> InSites,
+	const FTransform& InTransform,
+	int32 InRandomSeed,
 	float InChanceToFracture,
-	bool InGroupFracture,
+	bool InSplitIslands,
 	float InGrout,
 	float InAmplitude,
 	float InFrequency,
@@ -221,6 +242,12 @@ void FFractureEngineFracturing::VoronoiFracture(FManagedArrayCollection& InOutCo
 					}
 				}
 
+				FVector Origin = InTransform.GetTranslation();
+				for (FVector& Site : InSites)
+				{
+					Site -= Origin;
+				}
+
 				//
 				// Compute Voronoi Bounds
 				//
@@ -245,26 +272,35 @@ void FFractureEngineFracturing::VoronoiFracture(FManagedArrayCollection& InOutCo
 				FPlanarCells VoronoiPlanarCells = FPlanarCells(InSites, Voronoi);
 				VoronoiPlanarCells.InternalSurfaceMaterials.NoiseSettings = NoiseSettings;
 
+				RandomReduceSelection(InTransformSelection, InRandomSeed, InChanceToFracture);
 				TArray<int32> TransformSelectionArr = InTransformSelection.AsArray();
+
 				if (!FFractureEngineSelection::IsBoneSelectionValid(InOutCollection, TransformSelectionArr))
 				{
-					return;
+					return INDEX_NONE;
 				}
 				
-				int ResultGeometryIndex = CutMultipleWithPlanarCells(VoronoiPlanarCells, *GeomCollection, TransformSelectionArr, InGrout, InCollisionSampleSpacing, InRandomSeed, FTransform().Identity);
+				int ResultGeometryIndex = CutMultipleWithPlanarCells(VoronoiPlanarCells, *GeomCollection, TransformSelectionArr, InGrout, InCollisionSampleSpacing, InRandomSeed, InTransform, true, true, nullptr, Origin, InSplitIslands);
 
 				InOutCollection = (const FManagedArrayCollection&)(*GeomCollection);
+
+				return ResultGeometryIndex;
 			}
 		}
 	}
+
+	return INDEX_NONE;
 }
 
 
-void FFractureEngineFracturing::PlaneCutter(FManagedArrayCollection& InOutCollection,
-	const FDataflowTransformSelection& InTransformSelection,
+int32 FFractureEngineFracturing::PlaneCutter(FManagedArrayCollection& InOutCollection,
+	FDataflowTransformSelection InTransformSelection,
 	const FBox& InBoundingBox,
+	const FTransform& InTransform,
 	int32 InNumPlanes,
-	float InRandomSeed,
+	int32 InRandomSeed,
+	float InChanceToFracture,
+	bool InSplitIslands,
 	float InGrout,
 	float InAmplitude,
 	float InFrequency,
@@ -316,17 +352,549 @@ void FFractureEngineFracturing::PlaneCutter(FManagedArrayCollection& InOutCollec
 		float CollisionSampleSpacingVal = InCollisionSampleSpacing;
 		float GroutVal = InGrout;
 
+		RandomReduceSelection(InTransformSelection, InRandomSeed, InChanceToFracture);
 		TArray<int32> TransformSelectionArr = InTransformSelection.AsArray();
+
 		if (!FFractureEngineSelection::IsBoneSelectionValid(InOutCollection, TransformSelectionArr))
 		{
-			return;
+			return INDEX_NONE;
 		}
 
-		int ResultGeometryIndex = CutMultipleWithMultiplePlanes(CuttingPlanes, InternalSurfaceMaterials, *GeomCollection, TransformSelectionArr, GroutVal, CollisionSampleSpacingVal, InRandomSeed, FTransform().Identity);
+		int ResultGeometryIndex = CutMultipleWithMultiplePlanes(CuttingPlanes, InternalSurfaceMaterials, *GeomCollection, TransformSelectionArr, GroutVal, CollisionSampleSpacingVal, InRandomSeed, InTransform, true, nullptr, InSplitIslands);
 
 		InOutCollection = (const FManagedArrayCollection&)(*GeomCollection);
+
+		return ResultGeometryIndex;
 	}
+
+	return INDEX_NONE;
 }
 
+void FFractureEngineFracturing::GenerateSliceTransforms(TArray<FTransform>& InOutCuttingPlaneTransforms, 
+	const FBox& InBoundingBox,
+	int32 InSlicesX,
+	int32 InSlicesY,
+	int32 InSlicesZ,
+	int32 InRandomSeed,
+	float InSliceAngleVariation,
+	float InSliceOffsetVariation)
+{
+	const FBox& Bounds = InBoundingBox;
+	const FVector& Min = Bounds.Min;
+	const FVector& Max = Bounds.Max;
+	const FVector  Center = Bounds.GetCenter();
+	const FVector Extents(Max - Min);
+	const FVector HalfExtents(Extents * 0.5f);
+
+	const FVector Step(Extents.X / (InSlicesX + 1), Extents.Y / (InSlicesY + 1), Extents.Z / (InSlicesZ + 1));
+
+	InOutCuttingPlaneTransforms.Reserve(InSlicesX * InSlicesY * InSlicesZ);
+
+	FRandomStream RandomStream(InRandomSeed);
+
+	const float SliceAngleVariationInRadians = FMath::DegreesToRadians(InSliceAngleVariation);
+
+	const FVector XMin(-Bounds.GetExtent() * FVector(1.0f, 1.0f, 0.0f));
+	const FVector XMax(Bounds.GetExtent() * FVector(1.0f, 1.0f, 0.0f));
+
+	for (int32 xx = 0; xx < InSlicesX; ++xx)
+	{
+		const FVector SlicePosition(FVector(Min.X, Center.Y, Center.Z) + FVector((Step.X * xx) + Step.X, 0.0f, 0.0f) + RandomStream.VRand() * RandomStream.GetFraction() * InSliceOffsetVariation);
+		FTransform Transform(FQuat(FVector::RightVector, FMath::DegreesToRadians(90)), SlicePosition);
+		const FQuat RotA(FVector::RightVector, RandomStream.FRandRange(0.0f, SliceAngleVariationInRadians));
+		const FQuat RotB(FVector::ForwardVector, RandomStream.FRandRange(0.0f, SliceAngleVariationInRadians));
+		Transform.ConcatenateRotation(FQuat(RotA * RotB));
+		InOutCuttingPlaneTransforms.Emplace(Transform);
+	}
+
+	const FVector YMin(-Bounds.GetExtent() * FVector(1.0f, 1.0f, 0.0f));
+	const FVector YMax(Bounds.GetExtent() * FVector(1.0f, 1.0f, 0.0f));
+
+	for (int32 yy = 0; yy < InSlicesY; ++yy)
+	{
+		const FVector SlicePosition(FVector(Center.X, Min.Y, Center.Z) + FVector(0.0f, (Step.Y * yy) + Step.Y, 0.0f) + RandomStream.VRand() * RandomStream.GetFraction() * InSliceOffsetVariation);
+		FTransform Transform(FQuat(FVector::ForwardVector, FMath::DegreesToRadians(90)), SlicePosition);
+		const FQuat RotA(FVector::RightVector, RandomStream.FRandRange(0.0f, SliceAngleVariationInRadians));
+		const FQuat RotB(FVector::ForwardVector, RandomStream.FRandRange(0.0f, SliceAngleVariationInRadians));
+		Transform.ConcatenateRotation(RotA * RotB);
+		InOutCuttingPlaneTransforms.Emplace(Transform);
+	}
+
+	const FVector ZMin(-Bounds.GetExtent() * FVector(1.0f, 1.0f, 0.0f));
+	const FVector ZMax(Bounds.GetExtent() * FVector(1.0f, 1.0f, 0.0f));
+
+	for (int32 zz = 0; zz < InSlicesZ; ++zz)
+	{
+		const FVector SlicePosition(FVector(Center.X, Center.Y, Min.Z) + FVector(0.0f, 0.0f, (Step.Z * zz) + Step.Z) + RandomStream.VRand() * RandomStream.GetFraction() * InSliceOffsetVariation);
+		FTransform Transform(SlicePosition);
+		const FQuat RotA(FVector::RightVector, RandomStream.FRandRange(0.0f, SliceAngleVariationInRadians));
+		const FQuat RotB(FVector::ForwardVector, RandomStream.FRandRange(0.0f, SliceAngleVariationInRadians));
+		Transform.ConcatenateRotation(RotA * RotB);
+		InOutCuttingPlaneTransforms.Emplace(Transform);
+	}
+}
+
+static void ClearProximity(FGeometryCollection* GeometryCollection)
+{
+	if (GeometryCollection->HasAttribute("Proximity", FGeometryCollection::GeometryGroup))
+	{
+		GeometryCollection->RemoveAttribute("Proximity", FGeometryCollection::GeometryGroup);
+	}
+}
+
+int32 FFractureEngineFracturing::SliceCutter(FManagedArrayCollection& InOutCollection,
+	FDataflowTransformSelection InTransformSelection,
+	const FBox& InBoundingBox,
+	int32 InSlicesX,
+	int32 InSlicesY,
+	int32 InSlicesZ,
+	float InSliceAngleVariation,
+	float InSliceOffsetVariation,
+	int32 InRandomSeed,
+	float InChanceToFracture,
+	bool InSplitIslands,
+	float InGrout,
+	float InAmplitude,
+	float InFrequency,
+	float InPersistence,
+	float InLacunarity,
+	int32 InOctaveNumber,
+	float InPointSpacing,
+	bool InAddSamplesForCollision,
+	float InCollisionSampleSpacing)
+{
+	if (TUniquePtr<FGeometryCollection> GeomCollection = TUniquePtr<FGeometryCollection>(InOutCollection.NewCopy<FGeometryCollection>()))
+	{
+		TArray<FTransform> LocalCuttingPlanesTransforms;
+		GenerateSliceTransforms(LocalCuttingPlanesTransforms,
+			InBoundingBox,
+			InSlicesX,
+			InSlicesY,
+			InSlicesZ,
+			InRandomSeed,
+			InSliceAngleVariation,
+			InSliceOffsetVariation);
+
+		TArray<FPlane> CuttingPlanes;
+		CuttingPlanes.Reserve(LocalCuttingPlanesTransforms.Num());
+
+		for (const FTransform& Transform : LocalCuttingPlanesTransforms)
+		{
+			CuttingPlanes.Add(FPlane(Transform.GetLocation(), Transform.GetUnitAxis(EAxis::Z)));
+		}
+
+		FInternalSurfaceMaterials InternalSurfaceMaterials;
+		FNoiseSettings NoiseSettings;
+
+		if (InAmplitude > 0.f)
+		{
+			NoiseSettings.Amplitude = InAmplitude;
+			NoiseSettings.Frequency = InFrequency;
+			NoiseSettings.Lacunarity = InLacunarity;
+			NoiseSettings.Persistence = InPersistence;
+			NoiseSettings.Octaves = InOctaveNumber;
+			NoiseSettings.PointSpacing = InPointSpacing;
+
+			InternalSurfaceMaterials.NoiseSettings = NoiseSettings;
+		}
+
+		float CollisionSampleSpacingVal = InCollisionSampleSpacing;
+		float GroutVal = InGrout;
+
+		RandomReduceSelection(InTransformSelection, InRandomSeed, InChanceToFracture);
+		TArray<int32> TransformSelectionArr = InTransformSelection.AsArray();
+
+		if (!FFractureEngineSelection::IsBoneSelectionValid(InOutCollection, TransformSelectionArr))
+		{
+			return INDEX_NONE;
+		}
+
+		// Proximity is invalidated.
+		ClearProximity(GeomCollection.Get());
+
+		int ResultGeometryIndex = CutMultipleWithMultiplePlanes(CuttingPlanes, InternalSurfaceMaterials, *GeomCollection, TransformSelectionArr, GroutVal, CollisionSampleSpacingVal, InRandomSeed, FTransform().Identity, true, nullptr, InSplitIslands);
+
+		InOutCollection = (const FManagedArrayCollection&)(*GeomCollection);
+
+		return ResultGeometryIndex;
+	}
+
+	return INDEX_NONE;
+}
+
+namespace FractureToolBrickLocals
+{
+	// Calculate total number of bricks based on given dimensions and the extent of the object to be fractured.
+	// If the input is not valid or the result is too large, this functions returns -1.
+	// It is possible that we are dealing with incredibly large meshes and small brick dimensions. Doing the
+	// calculations in double and checking for NaNs will catch cases where for integer we would have to jump through
+	// multiple hoops to make sure we are not dealing with overflow. And since the limit for the number of bricks is
+	// comparably low, we do not need to worry about loss in precision for very large integers.
+	// For example, running this calculation with brick dimensions of 1 for the Sky Sphere will result in integer
+	// overflow.
+	static int64 CalculateNumBricks(const FVector& Dimensions, const FVector& Extents)
+	{
+		if (Dimensions.GetMin() <= 0 || Extents.GetMin() <= 0)
+		{
+			return -1;
+		}
+
+		const FVector NumBricksPerDim(ceil(Extents.X / Dimensions.X), ceil(Extents.Y / Dimensions.Y), ceil(Extents.Z / Dimensions.Z));
+		if (NumBricksPerDim.ContainsNaN())
+		{
+			return -1;
+		}
+
+		const double NumBricks = NumBricksPerDim.X * NumBricksPerDim.Y * NumBricksPerDim.Z;
+		if (FMath::IsNaN(NumBricks))
+		{
+			return -1;
+		}
+
+		return static_cast<int64>(NumBricks);
+	}
+
+	static FVector GetBrickDimensions(const float InBrickLength,
+		const float InBrickHeight,
+		const float InBrickDepth, 
+		const FVector& InExtents)
+	{
+		// Limit for the total number of bricks.
+		const int64 NumBricksLimit = 8192;
+
+		FVector Dimensions(InBrickLength, InBrickDepth, InBrickHeight);
+
+		// Early out if we have inputs we cannot deal with. If this call to CalculateNumBricks is fine then any other
+		// call to it will be fine, too, and we do not need to check for invalid results again.
+		int64 NumBricks = CalculateNumBricks(Dimensions, InExtents);
+		if (NumBricks < 0)
+		{
+			return FVector::ZeroVector;
+		}
+
+		if (NumBricks > NumBricksLimit)
+		{
+			const int64 InputNumBricks = NumBricks;
+
+			// Determine dimensions safely within the brick limit by iteratively doubling the brick size.
+			FVector SafeDimensions = Dimensions;
+			int64 SafeNumBricks;
+			do
+			{
+				SafeDimensions *= 2;
+				SafeNumBricks = CalculateNumBricks(SafeDimensions, InExtents);
+			} while (SafeNumBricks > NumBricksLimit);
+
+			// Maximize brick dimensions to fit within the brick limit via iterative interval halving.
+			const int32 IterationsMax = 10;
+			int32 Iterations = 0;
+			do
+			{
+				const FVector MidDimensions = (Dimensions + SafeDimensions) / 2;
+				const int64 MidNumBricks = CalculateNumBricks(MidDimensions, InExtents);
+
+				if (MidNumBricks > NumBricksLimit)
+				{
+					Dimensions = MidDimensions;
+					NumBricks = MidNumBricks;
+				}
+				else
+				{
+					SafeDimensions = MidDimensions;
+					SafeNumBricks = MidNumBricks;
+				}
+			} while (++Iterations < IterationsMax);
+
+			Dimensions = SafeDimensions;
+			NumBricks = SafeNumBricks;
+
+//			UE_LOG(LogFractureTool, Warning, TEXT("Brick Fracture: Current brick dimensions of %f x %f x %f would result in %d bricks. "
+//				"Reduced brick dimensions to %f x %f x %f resulting in %d bricks to stay within maximum number of %d bricks."),
+//				InBrickLength, InBrickDepth, InBrickHeight, InputNumBricks,
+//				Dimensions.X, Dimensions.Y, Dimensions.Z, NumBricks, NumBricksLimit);
+		}
+
+		return Dimensions;
+	}
+}
+
+void FFractureEngineFracturing::AddBoxEdges(TArray<TTuple<FVector, FVector>>& InOutEdges, const FVector& InMin, const FVector& InMax)
+{
+	InOutEdges.Emplace(MakeTuple(InMin, FVector(InMin.X, InMax.Y, InMin.Z)));
+	InOutEdges.Emplace(MakeTuple(InMin, FVector(InMin.X, InMin.Y, InMax.Z)));
+	InOutEdges.Emplace(MakeTuple(FVector(InMin.X, InMax.Y, InMax.Z), FVector(InMin.X, InMax.Y, InMin.Z)));
+	InOutEdges.Emplace(MakeTuple(FVector(InMin.X, InMax.Y, InMax.Z), FVector(InMin.X, InMin.Y, InMax.Z)));
+
+	InOutEdges.Emplace(MakeTuple(FVector(InMax.X, InMin.Y, InMin.Z), FVector(InMax.X, InMax.Y, InMin.Z)));
+	InOutEdges.Emplace(MakeTuple(FVector(InMax.X, InMin.Y, InMin.Z), FVector(InMax.X, InMin.Y, InMax.Z)));
+	InOutEdges.Emplace(MakeTuple(InMax, FVector(InMax.X, InMax.Y, InMin.Z)));
+	InOutEdges.Emplace(MakeTuple(InMax, FVector(InMax.X, InMin.Y, InMax.Z)));
+
+	InOutEdges.Emplace(MakeTuple(InMin, FVector(InMax.X, InMin.Y, InMin.Z)));
+	InOutEdges.Emplace(MakeTuple(FVector(InMin.X, InMin.Y, InMax.Z), FVector(InMax.X, InMin.Y, InMax.Z)));
+	InOutEdges.Emplace(MakeTuple(FVector(InMin.X, InMax.Y, InMin.Z), FVector(InMax.X, InMax.Y, InMin.Z)));
+	InOutEdges.Emplace(MakeTuple(FVector(InMin.X, InMax.Y, InMax.Z), InMax));
+}
+
+void FFractureEngineFracturing::GenerateBrickTransforms(const FBox& InBounds, 
+	TArray<FTransform>& InOutBrickTransforms, 
+	const EFractureBrickBondEnum InBond,
+	const float InBrickLength,
+	const float InBrickHeight,
+	const float InBrickDepth,
+	TArray<TTuple<FVector, FVector>>& InOutEdges)
+{
+	const FVector& Min = InBounds.Min;
+	const FVector& Max = InBounds.Max;
+	const FVector Extents(InBounds.Max - InBounds.Min);
+
+	// Determine brick dimensions (length, depth, height) and make sure we do not exceed the limit for the number of bricks.
+	// If we would simply use the input dimensions, we are prone to running out of memory and/or exceeding the storage capabilities of TArray, and crash.
+	const FVector BrickDimensions = FractureToolBrickLocals::GetBrickDimensions(InBrickLength, InBrickHeight, InBrickDepth, Extents);
+
+	// Early out if we have inputs we cannot deal with.
+	if (BrickDimensions == FVector::ZeroVector)
+	{
+		return;
+	}
+
+	// Reserve correct amount of memory to avoid re-allocations.
+	InOutBrickTransforms.Reserve(FractureToolBrickLocals::CalculateNumBricks(BrickDimensions, Extents));
+
+	const FVector BrickHalfDimensions(BrickDimensions * 0.5);
+	const FQuat HeaderRotation(FVector::UpVector, 1.5708);
+
+	if (InBond == EFractureBrickBondEnum::Dataflow_FractureBrickBond_Stretcher)
+	{
+		bool OddY = false;
+		for (float yy = 0.f; yy <= Extents.Y; yy += BrickDimensions.Y)
+		{
+			bool Oddline = false;
+			for (float zz = BrickHalfDimensions.Z; zz <= Extents.Z; zz += BrickDimensions.Z)
+			{
+				for (float xx = 0.f; xx <= Extents.X; xx += BrickDimensions.X)
+				{
+					FVector BrickPosition(Min + FVector(Oddline ^ OddY ? xx : xx + BrickHalfDimensions.X, yy, zz));
+					InOutBrickTransforms.Emplace(FTransform(BrickPosition));
+				}
+				Oddline = !Oddline;
+			}
+			OddY = !OddY;
+		}
+	}
+	else if (InBond == EFractureBrickBondEnum::Dataflow_FractureBrickBond_Stack)
+	{
+		bool OddY = false;
+		for (float yy = 0.f; yy <= Extents.Y; yy += BrickDimensions.Y)
+		{
+			for (float zz = BrickHalfDimensions.Z; zz <= Extents.Z; zz += BrickDimensions.Z)
+			{
+				for (float xx = 0.f; xx <= Extents.X; xx += BrickDimensions.X)
+				{
+					FVector BrickPosition(Min + FVector(OddY ? xx : xx + BrickHalfDimensions.X, yy, zz));
+					InOutBrickTransforms.Emplace(FTransform(BrickPosition));
+				}
+			}
+			OddY = !OddY;
+		}
+	}
+	else if (InBond == EFractureBrickBondEnum::Dataflow_FractureBrickBond_English)
+	{
+		float HalfLengthDepthDifference = BrickHalfDimensions.X - BrickHalfDimensions.Y - BrickHalfDimensions.Y;
+		bool OddY = false;
+		for (float yy = 0.f; yy <= Extents.Y; yy += BrickDimensions.Y)
+		{
+			bool Oddline = false;
+			for (float zz = BrickHalfDimensions.Z; zz <= Extents.Z; zz += BrickDimensions.Z)
+			{
+				if (Oddline && !OddY) // header row
+				{
+					for (float xx = 0.f; xx <= Extents.X; xx += BrickDimensions.Y)
+					{
+						FVector BrickPosition(Min + FVector(Oddline ^ OddY ? xx : xx + BrickHalfDimensions.Y, yy + BrickHalfDimensions.Y, zz));
+						InOutBrickTransforms.Emplace(FTransform(HeaderRotation, BrickPosition));
+					}
+				}
+				else if (!Oddline) // stretchers
+				{
+					for (float xx = 0.f; xx <= Extents.X; xx += BrickDimensions.X)
+					{
+						FVector BrickPosition(Min + FVector(Oddline ^ OddY ? xx : xx + BrickHalfDimensions.X, OddY ? yy + HalfLengthDepthDifference : yy - HalfLengthDepthDifference, zz));
+						InOutBrickTransforms.Emplace(FTransform(BrickPosition));
+					}
+				}
+				Oddline = !Oddline;
+			}
+			OddY = !OddY;
+		}
+	}
+	else if (InBond == EFractureBrickBondEnum::Dataflow_FractureBrickBond_Header)
+	{
+		bool OddY = false;
+		for (float yy = 0.f; yy <= Extents.Y; yy += BrickDimensions.X)
+		{
+			bool Oddline = false;
+			for (float zz = BrickHalfDimensions.Z; zz <= Extents.Z; zz += BrickDimensions.Z)
+			{
+				for (float xx = 0.f; xx <= Extents.X; xx += BrickDimensions.Y)
+				{
+					FVector BrickPosition(Min + FVector(Oddline ^ OddY ? xx : xx + BrickHalfDimensions.Y, yy, zz));
+					InOutBrickTransforms.Emplace(FTransform(HeaderRotation, BrickPosition));
+				}
+				Oddline = !Oddline;
+			}
+			OddY = !OddY;
+		}
+	}
+	else if (InBond == EFractureBrickBondEnum::Dataflow_FractureBrickBond_Flemish)
+	{
+		float HalfLengthDepthDifference = BrickHalfDimensions.X - BrickDimensions.Y;
+		bool OddY = false;
+		int32 RowY = 0;
+		for (float yy = 0.f; yy <= Extents.Y; yy += BrickDimensions.Y)
+		{
+			bool OddZ = false;
+			for (float zz = BrickHalfDimensions.Z; zz <= Extents.Z; zz += BrickDimensions.Z)
+			{
+				bool OddX = OddZ;
+				for (float xx = 0.f; xx <= Extents.X; xx += BrickHalfDimensions.X + BrickHalfDimensions.Y)
+				{
+					FVector BrickPosition(Min + FVector(xx, yy, zz));
+					if (OddX)
+					{
+						if (OddY) // runner
+						{
+							InOutBrickTransforms.Emplace(FTransform(BrickPosition + FVector(0, HalfLengthDepthDifference, 0))); // runner
+						}
+						else
+						{
+							InOutBrickTransforms.Emplace(FTransform(BrickPosition - FVector(0, HalfLengthDepthDifference, 0))); // runner
+
+						}
+					}
+					else if (!OddY) // header
+					{
+						InOutBrickTransforms.Emplace(FTransform(HeaderRotation, BrickPosition + FVector(0, BrickHalfDimensions.Y, 0))); // header
+					}
+					OddX = !OddX;
+				}
+				OddZ = !OddZ;
+			}
+			OddY = !OddY;
+			++RowY;
+		}
+	}
+
+	const FVector BrickMax(BrickHalfDimensions);
+	const FVector BrickMin(-BrickHalfDimensions);
+
+	for (const auto& Transform : InOutBrickTransforms)
+	{
+		AddBoxEdges(InOutEdges, Transform.TransformPosition(BrickMin), Transform.TransformPosition(BrickMax));
+	}
+}
+
+int32 FFractureEngineFracturing::BrickCutter(FManagedArrayCollection& InOutCollection,
+	FDataflowTransformSelection InTransformSelection,
+	const FBox& InBoundingBox,
+	const FTransform& InTransform,
+	EFractureBrickBondEnum InBond,
+	float InBrickLength,
+	float InBrickHeight,
+	float InBrickDepth,
+	int32 InRandomSeed,
+	float InChanceToFracture,
+	bool InSplitIslands,
+	float InGrout,
+	float InAmplitude,
+	float InFrequency,
+	float InPersistence,
+	float InLacunarity,
+	int32 InOctaveNumber,
+	float InPointSpacing,
+	bool InAddSamplesForCollision,
+	float InCollisionSampleSpacing)
+{
+	if (TUniquePtr<FGeometryCollection> GeomCollection = TUniquePtr<FGeometryCollection>(InOutCollection.NewCopy<FGeometryCollection>()))
+	{
+		TArray<FTransform> BrickTransforms;
+		TArray<TTuple<FVector, FVector>> Edges;
+		TArray<TTuple<FVector, FVector>> Boxes;
+
+		BrickTransforms.Empty();
+
+		const FBox Bounds = InBoundingBox;
+		GenerateBrickTransforms(Bounds,
+			BrickTransforms,
+			InBond,
+			InBrickLength,
+			InBrickHeight,
+			InBrickDepth,			
+			Edges);
+
+		// Get the same brick dimensions that were used in GenerateBrickTransform.
+		// If we cannot deal with the input data then the brick dimensions will be zero, but we do not need to
+		// explicitly handle that since it will only affect some local variables. The BrickTransforms will be empty
+		// and there are no further side effects.
+		const FVector BrickDimensions = FractureToolBrickLocals::GetBrickDimensions(InBrickLength, InBrickHeight, InBrickDepth, Bounds.Max - Bounds.Min);
+		const FVector BrickHalfDimensions(BrickDimensions * 0.5);
+
+		const FQuat HeaderRotation(FVector::UpVector, 1.5708);
+
+		TArray<FBox> BricksToCut;
+
+		// space the bricks by the grout setting, constrained to not erase the bricks
+		const float MinDim = FMath::Min3(BrickHalfDimensions.X, BrickHalfDimensions.Y, BrickHalfDimensions.Z);
+		const float HalfGrout = FMath::Clamp(0.5f * InGrout, 0, MinDim * 0.98f);
+		const FVector HalfBrick(BrickHalfDimensions - HalfGrout);
+		const FBox BrickBox(-HalfBrick, HalfBrick);
+
+		FTransform ContextTransform = InTransform;
+		FVector Origin = ContextTransform.GetTranslation();
+
+		for (const FTransform& Trans : BrickTransforms)
+		{
+			FTransform ToApply = Trans * FTransform(-Origin);
+			BricksToCut.Add(BrickBox.TransformBy(ToApply));
+		}
+
+		FInternalSurfaceMaterials InternalSurfaceMaterials;
+		FNoiseSettings NoiseSettings;
+
+		if (InAmplitude > 0.f)
+		{
+			NoiseSettings.Amplitude = InAmplitude;
+			NoiseSettings.Frequency = InFrequency;
+			NoiseSettings.Lacunarity = InLacunarity;
+			NoiseSettings.Persistence = InPersistence;
+			NoiseSettings.Octaves = InOctaveNumber;
+			NoiseSettings.PointSpacing = InPointSpacing;
+
+			InternalSurfaceMaterials.NoiseSettings = NoiseSettings;
+		}
+
+		float CollisionSampleSpacingVal = InCollisionSampleSpacing;
+
+		RandomReduceSelection(InTransformSelection, InRandomSeed, InChanceToFracture);
+		TArray<int32> TransformSelectionArr = InTransformSelection.AsArray();
+
+		if (!FFractureEngineSelection::IsBoneSelectionValid(InOutCollection, TransformSelectionArr))
+		{
+			return INDEX_NONE;
+		}
+
+		float GroutVal = 0.f; // CutterSettings->Grout; // Note: Grout is currently baked directly into the brick cells above
+		const bool bBricksAreTouching = InGrout <= UE_KINDA_SMALL_NUMBER;
+		FPlanarCells Cells = FPlanarCells(BricksToCut, bBricksAreTouching);
+
+		int32 ResultGeometryIndex = CutMultipleWithPlanarCells(Cells, *GeomCollection, TransformSelectionArr, GroutVal, InPointSpacing, InRandomSeed, InTransform, true, true, nullptr, Origin, InSplitIslands);
+
+		InOutCollection = (const FManagedArrayCollection&)(*GeomCollection);
+
+		return ResultGeometryIndex;
+	}
+
+	return INDEX_NONE;
+}
 
 
diff --git a/Engine/Plugins/Experimental/Fracture/Source/FractureEngine/Public/FractureEngineFracturing.h b/Engine/Plugins/Experimental/Fracture/Source/FractureEngine/Public/FractureEngineFracturing.h
--- a/Engine/Plugins/Experimental/Fracture/Source/FractureEngine/Public/FractureEngineFracturing.h
+++ b/Engine/Plugins/Experimental/Fracture/Source/FractureEngine/Public/FractureEngineFracturing.h
@@ -4,9 +4,24 @@
 
 #include "Math/Box.h"
 
+#include "FractureEngineFracturing.generated.h"
+
 struct FDataflowTransformSelection;
 struct FManagedArrayCollection;
 
+UENUM(BlueprintType)
+enum class EFractureBrickBondEnum : uint8
+{
+	Dataflow_FractureBrickBond_Stretcher UMETA(DisplayName = "Stretcher"),
+	Dataflow_FractureBrickBond_Stack UMETA(DisplayName = "Stack"),
+	Dataflow_FractureBrickBond_English UMETA(DisplayName = "English"),
+	Dataflow_FractureBrickBond_Header UMETA(DisplayName = "Header"),
+	Dataflow_FractureBrickBond_Flemish UMETA(DisplayName = "Flemish"),
+
+	//~~~
+	//256th entry
+	Dataflow_Max                UMETA(Hidden)
+};
 
 class FRACTUREENGINE_API FFractureEngineFracturing
 {
@@ -14,12 +29,13 @@ public:
 
 	static void GenerateExplodedViewAttribute(FManagedArrayCollection& InOutCollection, const FVector& InScale, float InUniformScale);
 
-	static void VoronoiFracture(FManagedArrayCollection& InOutCollection,
-		const FDataflowTransformSelection& InTransformSelection,
-		const TArray<FVector>& InSites,
-		float InRandomSeed,
+	static int32 VoronoiFracture(FManagedArrayCollection& InOutCollection,
+		FDataflowTransformSelection InTransformSelection,
+		TArray<FVector> InSites,
+		const FTransform& InTransform,
+		int32 InRandomSeed,
 		float InChanceToFracture,
-		bool InGroupFracture,
+		bool InSplitIslands,
 		float InGrout,
 		float InAmplitude,
 		float InFrequency,
@@ -30,11 +46,77 @@ public:
 		bool InAddSamplesForCollision,
 		float InCollisionSampleSpacing);
 
-	static void PlaneCutter(FManagedArrayCollection& InOutCollection,
-		const FDataflowTransformSelection& InTransformSelection,
+	static int32 PlaneCutter(FManagedArrayCollection& InOutCollection,
+		FDataflowTransformSelection InTransformSelection,
 		const FBox& InBoundingBox,
+		const FTransform& InTransform,
 		int32 InNumPlanes,
-		float InRandomSeed,
+		int32 InRandomSeed,
+		float InChanceToFracture,
+		bool InSplitIslands,
+		float InGrout,
+		float InAmplitude,
+		float InFrequency,
+		float InPersistence,
+		float InLacunarity,
+		int32 InOctaveNumber,
+		float InPointSpacing,
+		bool InAddSamplesForCollision,
+		float InCollisionSampleSpacing);
+
+	static void GenerateSliceTransforms(TArray<FTransform>& InOutCuttingPlaneTransforms,
+		const FBox& InBoundingBox,
+		int32 InSlicesX,
+		int32 InSlicesY,
+		int32 InSlicesZ,
+		int32 InRandomSeed,
+		float InSliceAngleVariation,
+		float InSliceOffsetVariation);
+
+	static int32 SliceCutter(FManagedArrayCollection& InOutCollection,
+		FDataflowTransformSelection InTransformSelection,
+		const FBox& InBoundingBox,
+		int32 InSlicesX,
+		int32 InSlicesY,
+		int32 InSlicesZ,
+		float InSliceAngleVariation,
+		float InSliceOffsetVariation,
+		int32 InRandomSeed,
+		float InChanceToFracture,
+		bool InSplitIslands,
+		float InGrout,
+		float InAmplitude,
+		float InFrequency,
+		float InPersistence,
+		float InLacunarity,
+		int32 InOctaveNumber,
+		float InPointSpacing,
+		bool InAddSamplesForCollision,
+		float InCollisionSampleSpacing);
+
+	static void AddBoxEdges(TArray<TTuple<FVector, FVector>>& InOutEdges, 
+		const FVector& InMin, 
+		const FVector& InMax);
+
+	static void GenerateBrickTransforms(const FBox& InBounds,
+		TArray<FTransform>& InOutBrickTransforms,
+		const EFractureBrickBondEnum InBond,
+		const float InBrickLength,
+		const float InBrickHeight,
+		const float InBrickDepth,
+		TArray<TTuple<FVector, FVector>>& InOutEdges);
+
+	static int32 BrickCutter(FManagedArrayCollection& InOutCollection,
+		FDataflowTransformSelection InTransformSelection,
+		const FBox& InBoundingBox, 
+		const FTransform& InTransform,
+		EFractureBrickBondEnum InBond,
+		float InBrickLength,
+		float InBrickHeight,
+		float InBrickDepth,
+		int32 InRandomSeed,
+		float InChanceToFracture,
+		bool InSplitIslands,
 		float InGrout,
 		float InAmplitude,
 		float InFrequency,
diff --git a/Engine/Plugins/Experimental/GeometryCollectionPlugin/Source/GeometryCollectionNodes/Private/Dataflow/GeometryCollectionFracturingNodes.cpp b/Engine/Plugins/Experimental/GeometryCollectionPlugin/Source/GeometryCollectionNodes/Private/Dataflow/GeometryCollectionFracturingNodes.cpp
--- a/Engine/Plugins/Experimental/GeometryCollectionPlugin/Source/GeometryCollectionNodes/Private/Dataflow/GeometryCollectionFracturingNodes.cpp
+++ b/Engine/Plugins/Experimental/GeometryCollectionPlugin/Source/GeometryCollectionNodes/Private/Dataflow/GeometryCollectionFracturingNodes.cpp
@@ -27,7 +27,9 @@
 #include "GeometryCollection/GeometryCollectionProximityUtility.h"
 #include "FractureEngineClustering.h"
 #include "FractureEngineSelection.h"
-#include "FractureEngineFracturing.h"
+#include "GeometryCollection/Facades/CollectionTransformSelectionFacade.h"
+#include "GeometryCollection/Facades/CollectionBoundsFacade.h"
+#include "Dataflow/DataflowSelection.h"
 
 #include UE_INLINE_GENERATED_CPP_BY_NAME(GeometryCollectionFracturingNodes)
 
@@ -39,10 +41,17 @@ namespace Dataflow
 		static const FLinearColor CDefaultNodeBodyTintColor = FLinearColor(0.f, 0.f, 0.f, 0.5f);
 
 		DATAFLOW_NODE_REGISTER_CREATION_FACTORY(FUniformScatterPointsDataflowNode);
+		DATAFLOW_NODE_REGISTER_CREATION_FACTORY(FUniformScatterPointsDataflowNode_v2);
 		DATAFLOW_NODE_REGISTER_CREATION_FACTORY(FRadialScatterPointsDataflowNode);
+		DATAFLOW_NODE_REGISTER_CREATION_FACTORY(FRadialScatterPointsDataflowNode_v2);
+		DATAFLOW_NODE_REGISTER_CREATION_FACTORY(FGridScatterPointsDataflowNode);
 		DATAFLOW_NODE_REGISTER_CREATION_FACTORY(FVoronoiFractureDataflowNode);
+		DATAFLOW_NODE_REGISTER_CREATION_FACTORY(FVoronoiFractureDataflowNode_v2);
 		DATAFLOW_NODE_REGISTER_CREATION_FACTORY(FPlaneCutterDataflowNode);
+		DATAFLOW_NODE_REGISTER_CREATION_FACTORY(FPlaneCutterDataflowNode_v2);
 		DATAFLOW_NODE_REGISTER_CREATION_FACTORY(FExplodedViewDataflowNode);
+		DATAFLOW_NODE_REGISTER_CREATION_FACTORY(FSliceCutterDataflowNode);
+		DATAFLOW_NODE_REGISTER_CREATION_FACTORY(FBrickCutterDataflowNode);
 
 		// GeometryCollection|Fracture
 		DATAFLOW_NODE_REGISTER_CREATION_FACTORY_NODE_COLORS_BY_CATEGORY("GeometryCollection|Fracture", FLinearColor(1.f, 1.f, 0.8f), CDefaultNodeBodyTintColor);
@@ -78,6 +87,35 @@ void FUniformScatterPointsDataflowNode::Evaluate(Dataflow::FContext& Context, co
 	}
 }
 
+void FUniformScatterPointsDataflowNode_v2::Evaluate(Dataflow::FContext& Context, const FDataflowOutput* Out) const
+{
+	if (Out->IsA<TArray<FVector>>(&Points))
+	{
+		const FBox& BBox = GetValue<FBox>(Context, &BoundingBox);
+		if (BBox.GetVolume() > 0.f)
+		{
+			FRandomStream RandStream(GetValue<int32>(Context, &RandomSeed));
+
+			const FVector Extent(BBox.Max - BBox.Min);
+			const int32 NumPoints = RandStream.RandRange(GetValue<int32>(Context, &MinNumberOfPoints), GetValue<int32>(Context, &MaxNumberOfPoints));
+
+			TArray<FVector> PointsArr;
+			PointsArr.Reserve(NumPoints);
+			for (int32 Idx = 0; Idx < NumPoints; ++Idx)
+			{
+				PointsArr.Emplace(BBox.Min + FVector(RandStream.FRand(), RandStream.FRand(), RandStream.FRand()) * Extent);
+			}
+
+			SetValue(Context, MoveTemp(PointsArr), &Points);
+		}
+		else
+		{
+			// ERROR: Invalid BoundingBox input
+			SetValue(Context, TArray<FVector>(), &Points);
+		}
+	}
+}
+
 void FRadialScatterPointsDataflowNode::Evaluate(Dataflow::FContext& Context, const FDataflowOutput* Out) const
 {
 	if (Out->IsA<TArray<FVector>>(&Points))
@@ -108,6 +146,175 @@ void FRadialScatterPointsDataflowNode::Evaluate(Dataflow::FContext& Context, con
 	}
 }
 
+void FRadialScatterPointsDataflowNode_v2::Evaluate(Dataflow::FContext& Context, const FDataflowOutput* Out) const
+{
+	if (Out->IsA<TArray<FVector>>(&Points))
+	{
+		const FBox InBoundingBox = GetValue<FBox>(Context, &BoundingBox);
+		const FVector InCenter = GetValue<FVector>(Context, &Center);
+		const FVector InNormal = GetValue<FVector>(Context, &Normal);
+		const int32 InRandomSeed = GetValue<int32>(Context, &RandomSeed);
+		const int32 InAngularSteps = GetValue<int32>(Context, &AngularSteps);
+		const float InAngleOffset = GetValue<float>(Context, &AngleOffset);
+		const float InAngularNoise = GetValue<float>(Context, &AngularNoise);
+		const float InRadius = GetValue<float>(Context, &Radius);
+		const int32 InRadialSteps = GetValue<int32>(Context, &RadialSteps);
+		const float InRadialStepExponent = GetValue<float>(Context, &RadialStepExponent);
+		const float InRadialMinStep = GetValue<float>(Context, &RadialMinStep);
+		const float InRadialNoise = GetValue<float>(Context, &RadialNoise);
+		const float InRadialVariability = GetValue<float>(Context, &RadialVariability);
+		const float InAngularVariability = GetValue<float>(Context, &AngularVariability);
+		const float InAxialVariability = GetValue<float>(Context, &AxialVariability);
+
+		TArray<FVector> PointsArr;
+
+		const FVector::FReal AngularStep = 2 * PI / InAngularSteps;
+
+		FVector CenterVal(InBoundingBox.GetCenter() + InCenter);
+
+		FRandomStream RandStream(InRandomSeed);
+		FVector UpVector(InNormal);
+		UpVector.Normalize();
+		FVector BasisX, BasisY;
+		UpVector.FindBestAxisVectors(BasisX, BasisY);
+
+		// Precompute consistent noise for each angular step
+		TArray<FVector::FReal> AngleStepOffsets;
+		AngleStepOffsets.SetNumUninitialized(InAngularSteps);
+		for (int32 AngleIdx = 0; AngleIdx < InAngularSteps; ++AngleIdx)
+		{
+			AngleStepOffsets[AngleIdx] = FMath::DegreesToRadians(RandStream.FRandRange(-1, 1) * InAngularNoise);
+		}
+
+		// Compute radial positions following an (idx+1)^exp curve, and then re-normalize back to the Radius range
+		TArray<FVector::FReal> RadialPositions;
+		RadialPositions.SetNumUninitialized(InRadialSteps);
+		FVector::FReal StepOffset = 0;
+		for (int32 RadIdx = 0; RadIdx < InRadialSteps; ++RadIdx)
+		{
+			FVector::FReal RadialPos = FMath::Pow(RadIdx + 1, InRadialStepExponent) + StepOffset;
+			if (RadIdx == 0)
+			{
+				// Note we bring the first point a half-step toward the center, and shift all subsequent points accordingly
+				// so that for Exponent==1, the step from center to first boundary is the same distance as the step between each boundary
+				// (this is only necessary because there is no Voronoi site at the center)
+				RadialPos *= .5;
+				StepOffset = -RadialPos;
+			}
+
+			RadialPositions[RadIdx] = RadialPos;
+		}
+		// Normalize positions so that the diagram fits in the target radius
+		FVector::FReal RadialPosNorm = InRadius / RadialPositions.Last();
+		for (FVector::FReal& RadialPos : RadialPositions)
+		{
+			RadialPos = RadialPos * RadialPosNorm;
+		}
+		// Add radial noise 
+		for (int32 RadIdx = 0; RadIdx < InRadialSteps; ++RadIdx)
+		{
+			FVector::FReal& RadialPos = RadialPositions[RadIdx];
+			// Offset by RadialNoise, but don't allow noise to take the value below 0
+			RadialPos += RandStream.FRandRange(-FMath::Min(RadialPos, InRadialNoise), InRadialNoise);
+		}
+		// make sure the positions remain in increasing order
+		RadialPositions.Sort();
+		// Adjust positions so they are never closer than the RadialMinStep
+		FVector::FReal LastRadialPos = 0;
+		for (int32 RadIdx = 0; RadIdx < InRadialSteps; ++RadIdx)
+		{
+			FVector::FReal MinStep = InRadialMinStep;
+			if (RadIdx == 0)
+			{
+				MinStep *= .5;
+			}
+			if (RadialPositions[RadIdx] - LastRadialPos < MinStep)
+			{
+				RadialPositions[RadIdx] = LastRadialPos + MinStep;
+			}
+			LastRadialPos = RadialPositions[RadIdx];
+		}
+
+		// Add a bit of noise to work around failure case in Voro++
+		// TODO: fix the failure case in Voro++ and remove this
+		float MinRadialVariability = InRadius > 1.f ? .0001f : 0.f;
+		float UseRadialVariability = FMath::Max(MinRadialVariability, InRadialVariability);
+
+		// Create the radial Voronoi sites
+		for (int32 ii = 0; ii < InRadialSteps; ++ii)
+		{
+			FVector::FReal Len = RadialPositions[ii];
+			FVector::FReal Angle = FMath::DegreesToRadians(InAngleOffset);
+			for (int32 kk = 0; kk < InAngularSteps; ++kk, Angle += AngularStep)
+			{
+				// Add the global noise and the per-point noise into the angle
+				FVector::FReal UseAngle = Angle + AngleStepOffsets[kk] + FMath::DegreesToRadians(RandStream.FRand() * InAngularVariability);
+				// Add per point noise into the radial position
+				FVector::FReal UseRadius = Len + FVector::FReal(RandStream.FRand() * UseRadialVariability);
+				FVector RotatingOffset = UseRadius * (FMath::Cos(UseAngle) * BasisX + FMath::Sin(UseAngle) * BasisY);
+				PointsArr.Emplace(CenterVal + RotatingOffset + UpVector * (RandStream.FRandRange(-1, 1) * InAxialVariability));
+			}
+		}
+
+		SetValue(Context, MoveTemp(PointsArr), &Points);
+	}
+}
+
+void FGridScatterPointsDataflowNode::Evaluate(Dataflow::FContext& Context, const FDataflowOutput* Out) const
+{
+	if (Out->IsA<TArray<FVector>>(&Points))
+	{
+		const FBox& BBox = GetValue<FBox>(Context, &BoundingBox);
+		if (BBox.GetVolume() > 0.f)
+		{
+			const FVector Extent(BBox.Max - BBox.Min);
+			const int32 NumPointsInX = GetValue<int32>(Context, &NumberOfPointsInX);
+			const int32 NumPointsInY = GetValue<int32>(Context, &NumberOfPointsInY);
+			const int32 NumPointsInZ = GetValue<int32>(Context, &NumberOfPointsInZ);
+			
+			if (NumPointsInX >= 2 && NumPointsInY >= 2 && NumPointsInZ >= 2)
+			{
+				const int32 NumPoints = NumPointsInX * NumPointsInY * NumPointsInZ;
+				const float dX = Extent.X / (float)NumPointsInX;
+				const float dY = Extent.Y / (float)NumPointsInY;
+				const float dZ = Extent.Z / (float)NumPointsInZ;
+
+				FRandomStream RandStream(GetValue<int32>(Context, &RandomSeed));
+
+				TArray<FVector> PointsArr;
+				PointsArr.Reserve(NumPoints);
+				for (int32 Idx_X = 0; Idx_X < NumPointsInX; ++Idx_X)
+				{
+					for (int32 Idx_Y = 0; Idx_Y < NumPointsInY; ++Idx_Y)
+					{
+						for (int32 Idx_Z = 0; Idx_Z < NumPointsInZ; ++Idx_Z)
+						{
+							FVector RandomDisplacement = FVector(RandStream.FRandRange(-1.f, 1.f) * GetValue<float>(Context, &MaxRandomDisplacementX),
+								RandStream.FRandRange(-1.f, 1.f) * GetValue<float>(Context, &MaxRandomDisplacementY),
+								RandStream.FRandRange(-1.f, 1.f) * GetValue<float>(Context, &MaxRandomDisplacementZ));
+
+							PointsArr.Emplace(BBox.Min.X + 0.5f * dX + (float)Idx_X * dX + RandomDisplacement.X,
+								BBox.Min.Y + 0.5f * dY + (float)Idx_Y * dY + RandomDisplacement.Y,
+								BBox.Min.Z + 0.5f * dZ + (float)Idx_Z * dZ + RandomDisplacement.Z);
+						}
+					}
+				}
+
+				SetValue(Context, MoveTemp(PointsArr), &Points);
+			}
+			else
+			{
+				// ERROR: Invalid number of points
+				SetValue(Context, TArray<FVector>(), &Points);
+			}
+		}
+		else
+		{
+			// ERROR: Invalid BoundingBox input
+			SetValue(Context, TArray<FVector>(), &Points);
+		}
+	}
+}
 
 void FVoronoiFractureDataflowNode::Evaluate(Dataflow::FContext& Context, const FDataflowOutput* Out) const
 {
@@ -124,9 +331,10 @@ void FVoronoiFractureDataflowNode::Evaluate(Dataflow::FContext& Context, const F
 				FFractureEngineFracturing::VoronoiFracture(InCollection,
 					InTransformSelection,
 					GetValue<TArray<FVector>>(Context, &Points),
-					GetValue<float>(Context, &RandomSeed),
+					FTransform::Identity,
+					(int32)GetValue<float>(Context, &RandomSeed),
 					GetValue<float>(Context, &ChanceToFracture),
-					GroupFracture,
+					true,
 					GetValue<float>(Context, &Grout),
 					GetValue<float>(Context, &Amplitude),
 					GetValue<float>(Context, &Frequency),
@@ -148,6 +356,94 @@ void FVoronoiFractureDataflowNode::Evaluate(Dataflow::FContext& Context, const F
 	}
 }
 
+void FVoronoiFractureDataflowNode_v2::Evaluate(Dataflow::FContext& Context, const FDataflowOutput* Out) const
+{
+	if (Out->IsA<FManagedArrayCollection>(&Collection) ||
+		Out->IsA<FDataflowTransformSelection>(&TransformSelection) ||
+		Out->IsA<FDataflowTransformSelection>(&NewGeometryTransformSelection))
+	{
+		FDataflowTransformSelection InTransformSelection = GetValue<FDataflowTransformSelection>(Context, &TransformSelection);
+
+		//
+		// If not connected select everything by default
+		//
+		if (!IsConnected<FDataflowTransformSelection>(&TransformSelection))
+		{
+			const FManagedArrayCollection& InCollection = GetValue<FManagedArrayCollection>(Context, &Collection);
+
+			GeometryCollection::Facades::FCollectionTransformSelectionFacade TransformSelectionFacade(InCollection);
+			const TArray<int32>& SelectionArr = TransformSelectionFacade.SelectAll();
+
+			FDataflowTransformSelection NewTransformSelection;
+			NewTransformSelection.Initialize(InCollection.NumElements(FGeometryCollection::TransformGroup), false);
+			NewTransformSelection.SetFromArray(SelectionArr);
+
+			InTransformSelection = NewTransformSelection;
+		}
+
+		if (InTransformSelection.AnySelected())
+		{
+			FManagedArrayCollection InCollection = GetValue<FManagedArrayCollection>(Context, &Collection);
+
+			int32 ResultGeometryIndex = FFractureEngineFracturing::VoronoiFracture(InCollection,
+				InTransformSelection,
+				GetValue<TArray<FVector>>(Context, &Points),
+				GetValue<FTransform>(Context, &Transform),
+				GetValue<int32>(Context, &RandomSeed),
+				GetValue<float>(Context, &ChanceToFracture),
+				SplitIslands,
+				GetValue<float>(Context, &Grout),
+				GetValue<float>(Context, &Amplitude),
+				GetValue<float>(Context, &Frequency),
+				GetValue<float>(Context, &Persistence),
+				GetValue<float>(Context, &Lacunarity),
+				GetValue<int32>(Context, &OctaveNumber),
+				GetValue<float>(Context, &PointSpacing),
+				AddSamplesForCollision,
+				GetValue<float>(Context, &CollisionSampleSpacing));
+
+			FDataflowTransformSelection NewSelection;
+			FDataflowTransformSelection OriginalSelection;
+
+			if (ResultGeometryIndex != INDEX_NONE)
+			{
+				if (InCollection.HasAttribute("TransformIndex", FGeometryCollection::GeometryGroup))
+				{
+					const TManagedArray<int32>& TransformIndices = InCollection.GetAttribute<int32>("TransformIndex", FGeometryCollection::GeometryGroup);
+
+					NewSelection.Initialize(TransformIndices.Num(), false);
+					OriginalSelection.Initialize(TransformIndices.Num(), false);
+
+					// The newly fractured pieces are added to the end of the transform array (starting position is ResultGeometryIndex)
+					for (int32 Idx = ResultGeometryIndex; Idx < TransformIndices.Num(); ++Idx)
+					{
+						int32 BoneIdx = TransformIndices[Idx];
+						NewSelection.SetSelected(BoneIdx);
+					}
+
+					for (int32 Idx = 0; Idx < InTransformSelection.Num(); ++Idx)
+					{
+						if (InTransformSelection.IsSelected(Idx))
+						{
+							OriginalSelection.SetSelected(Idx);
+						}
+					}
+
+				}
+			}
+
+			SetValue(Context, MoveTemp(InCollection), &Collection);
+			SetValue(Context, OriginalSelection, &TransformSelection);
+			SetValue(Context, NewSelection, &NewGeometryTransformSelection);
+
+			return;
+		}
+
+		SafeForwardInput(Context, &Collection, &Collection);
+		SetValue(Context, InTransformSelection, &TransformSelection);
+		SetValue(Context, FDataflowTransformSelection(), &NewGeometryTransformSelection);
+	}
+}
 
 void FPlaneCutterDataflowNode::Evaluate(Dataflow::FContext& Context, const FDataflowOutput* Out) const
 {
@@ -164,8 +460,11 @@ void FPlaneCutterDataflowNode::Evaluate(Dataflow::FContext& Context, const FData
 				FFractureEngineFracturing::PlaneCutter(InCollection,
 					InTransformSelection,
 					GetValue<FBox>(Context, &BoundingBox),
+					FTransform::Identity,
 					NumPlanes,
-					GetValue<float>(Context, &RandomSeed),
+					(int32)GetValue<float>(Context, &RandomSeed),
+					1.f,
+					true,
 					GetValue<float>(Context, &Grout),
 					GetValue<float>(Context, &Amplitude),
 					GetValue<float>(Context, &Frequency),
@@ -187,6 +486,95 @@ void FPlaneCutterDataflowNode::Evaluate(Dataflow::FContext& Context, const FData
 	}
 }
 
+void FPlaneCutterDataflowNode_v2::Evaluate(Dataflow::FContext& Context, const FDataflowOutput* Out) const
+{
+	if (Out->IsA<FManagedArrayCollection>(&Collection) ||
+		Out->IsA<FDataflowTransformSelection>(&TransformSelection) ||
+		Out->IsA<FDataflowTransformSelection>(&NewGeometryTransformSelection))
+	{
+		FDataflowTransformSelection InTransformSelection = GetValue<FDataflowTransformSelection>(Context, &TransformSelection);
+
+		//
+		// If not connected select everything by default
+		//
+		if (!IsConnected<FDataflowTransformSelection>(&TransformSelection))
+		{
+			const FManagedArrayCollection& InCollection = GetValue<FManagedArrayCollection>(Context, &Collection);
+
+			GeometryCollection::Facades::FCollectionTransformSelectionFacade TransformSelectionFacade(InCollection);
+			const TArray<int32>& SelectionArr = TransformSelectionFacade.SelectAll();
+
+			FDataflowTransformSelection NewTransformSelection;
+			NewTransformSelection.Initialize(InCollection.NumElements(FGeometryCollection::TransformGroup), false);
+			NewTransformSelection.SetFromArray(SelectionArr);
+
+			InTransformSelection = NewTransformSelection;
+		}
+
+		if (InTransformSelection.AnySelected())
+		{
+			FManagedArrayCollection InCollection = GetValue<FManagedArrayCollection>(Context, &Collection);
+
+			int32 ResultGeometryIndex = FFractureEngineFracturing::PlaneCutter(InCollection,
+				InTransformSelection,
+				GetValue<FBox>(Context, &BoundingBox),
+				GetValue<FTransform>(Context, &Transform),
+				NumPlanes,
+				GetValue<int32>(Context, &RandomSeed),
+				GetValue<float>(Context, &ChanceToFracture),
+				SplitIslands,
+				GetValue<float>(Context, &Grout),
+				GetValue<float>(Context, &Amplitude),
+				GetValue<float>(Context, &Frequency),
+				GetValue<float>(Context, &Persistence),
+				GetValue<float>(Context, &Lacunarity),
+				GetValue<int32>(Context, &OctaveNumber),
+				GetValue<float>(Context, &PointSpacing),
+				AddSamplesForCollision,
+				GetValue<float>(Context, &CollisionSampleSpacing));
+
+			FDataflowTransformSelection NewSelection;
+			FDataflowTransformSelection OriginalSelection;
+
+			if (ResultGeometryIndex != INDEX_NONE)
+			{
+				if (InCollection.HasAttribute("TransformIndex", FGeometryCollection::GeometryGroup))
+				{
+					const TManagedArray<int32>& TransformIndices = InCollection.GetAttribute<int32>("TransformIndex", FGeometryCollection::GeometryGroup);
+
+					NewSelection.Initialize(TransformIndices.Num(), false);
+					OriginalSelection.Initialize(TransformIndices.Num(), false);
+
+					// The newly fractured pieces are added to the end of the transform array (starting position is ResultGeometryIndex)
+					for (int32 Idx = ResultGeometryIndex; Idx < TransformIndices.Num(); ++Idx)
+					{
+						int32 BoneIdx = TransformIndices[Idx];
+						NewSelection.SetSelected(BoneIdx);
+					}
+
+					for (int32 Idx = 0; Idx < InTransformSelection.Num(); ++Idx)
+					{
+						if (InTransformSelection.IsSelected(Idx))
+						{
+							OriginalSelection.SetSelected(Idx);
+						}
+					}
+
+				}
+			}
+
+			SetValue(Context, MoveTemp(InCollection), &Collection);
+			SetValue(Context, OriginalSelection, &TransformSelection);
+			SetValue(Context, NewSelection, &NewGeometryTransformSelection);
+
+			return;
+		}
+
+		SafeForwardInput(Context, &Collection, &Collection);
+		SetValue(Context, InTransformSelection, &TransformSelection);
+		SetValue(Context, FDataflowTransformSelection(), &NewGeometryTransformSelection);
+	}
+}
 
 void FExplodedViewDataflowNode::Evaluate(Dataflow::FContext& Context, const FDataflowOutput* Out) const
 {
@@ -200,6 +588,204 @@ void FExplodedViewDataflowNode::Evaluate(Dataflow::FContext& Context, const FDat
 	}
 }
 
+void FSliceCutterDataflowNode::Evaluate(Dataflow::FContext& Context, const FDataflowOutput* Out) const
+{
+	if (Out->IsA<FManagedArrayCollection>(&Collection) ||
+		Out->IsA<FDataflowTransformSelection>(&TransformSelection) ||
+		Out->IsA<FDataflowTransformSelection>(&NewGeometryTransformSelection))
+	{
+		FDataflowTransformSelection InTransformSelection = GetValue<FDataflowTransformSelection>(Context, &TransformSelection);
+
+		//
+		// If not connected select everything by default
+		//
+		if (!IsConnected<FDataflowTransformSelection>(&TransformSelection))
+		{
+			const FManagedArrayCollection& InCollection = GetValue<FManagedArrayCollection>(Context, &Collection);
+
+			GeometryCollection::Facades::FCollectionTransformSelectionFacade TransformSelectionFacade(InCollection);
+			const TArray<int32>& SelectionArr = TransformSelectionFacade.SelectAll();
+
+			FDataflowTransformSelection NewTransformSelection;
+			NewTransformSelection.Initialize(InCollection.NumElements(FGeometryCollection::TransformGroup), false);
+			NewTransformSelection.SetFromArray(SelectionArr);
+
+			InTransformSelection = NewTransformSelection;
+		}
+
+		if (InTransformSelection.AnySelected())
+		{
+			FManagedArrayCollection InCollection = GetValue<FManagedArrayCollection>(Context, &Collection);
+
+			int32 ResultGeometryIndex = FFractureEngineFracturing::SliceCutter(InCollection,
+				InTransformSelection,
+				GetValue<FBox>(Context, &BoundingBox),
+				GetValue<int32>(Context, &SlicesX),
+				GetValue<int32>(Context, &SlicesY),
+				GetValue<int32>(Context, &SlicesZ),
+				GetValue<float>(Context, &SliceAngleVariation),
+				GetValue<float>(Context, &SliceOffsetVariation),
+				GetValue<int32>(Context, &RandomSeed),
+				GetValue<float>(Context, &ChanceToFracture),
+				SplitIslands,
+				GetValue<float>(Context, &Grout),
+				GetValue<float>(Context, &Amplitude),
+				GetValue<float>(Context, &Frequency),
+				GetValue<float>(Context, &Persistence),
+				GetValue<float>(Context, &Lacunarity),
+				GetValue<int32>(Context, &OctaveNumber),
+				GetValue<float>(Context, &PointSpacing),
+				AddSamplesForCollision,
+				GetValue<float>(Context, &CollisionSampleSpacing));
+
+			FDataflowTransformSelection NewSelection;
+			FDataflowTransformSelection OriginalSelection;
+
+			if (ResultGeometryIndex != INDEX_NONE)
+			{
+				if (InCollection.HasAttribute("TransformIndex", FGeometryCollection::GeometryGroup))
+				{
+					const TManagedArray<int32>& TransformIndices = InCollection.GetAttribute<int32>("TransformIndex", FGeometryCollection::GeometryGroup);
+
+					NewSelection.Initialize(TransformIndices.Num(), false);
+					OriginalSelection.Initialize(TransformIndices.Num(), false);
+
+					// The newly fractured pieces are added to the end of the transform array (starting position is ResultGeometryIndex)
+					for (int32 Idx = ResultGeometryIndex; Idx < TransformIndices.Num(); ++Idx)
+					{
+						int32 BoneIdx = TransformIndices[Idx];
+						NewSelection.SetSelected(BoneIdx);
+					}
+
+					for (int32 Idx = 0; Idx < InTransformSelection.Num(); ++Idx)
+					{
+						if (InTransformSelection.IsSelected(Idx))
+						{
+							OriginalSelection.SetSelected(Idx);
+						}
+					}
+
+				}
+			}
+
+			SetValue(Context, MoveTemp(InCollection), &Collection);
+			SetValue(Context, OriginalSelection, &TransformSelection);
+			SetValue(Context, NewSelection, &NewGeometryTransformSelection);
+
+			return;
+		}
+
+		SafeForwardInput(Context, &Collection, &Collection);
+		SetValue(Context, InTransformSelection, &TransformSelection);
+		SetValue(Context, FDataflowTransformSelection(), &NewGeometryTransformSelection);
+	}
+}
+
+void FBrickCutterDataflowNode::Evaluate(Dataflow::FContext& Context, const FDataflowOutput* Out) const
+{
+	if (Out->IsA<FManagedArrayCollection>(&Collection) ||
+		Out->IsA<FDataflowTransformSelection>(&TransformSelection) ||
+		Out->IsA<FDataflowTransformSelection>(&NewGeometryTransformSelection))
+	{
+		FDataflowTransformSelection InTransformSelection = GetValue<FDataflowTransformSelection>(Context, &TransformSelection);
+		//
+		// If not connected select everything by default
+		//
+		if (!IsConnected<FDataflowTransformSelection>(&TransformSelection))
+		{
+			const FManagedArrayCollection& InCollection = GetValue<FManagedArrayCollection>(Context, &Collection);
+
+			GeometryCollection::Facades::FCollectionTransformSelectionFacade TransformSelectionFacade(InCollection);
+			const TArray<int32>& SelectionArr = TransformSelectionFacade.SelectAll();
+
+			FDataflowTransformSelection NewTransformSelection;
+			NewTransformSelection.Initialize(InCollection.NumElements(FGeometryCollection::TransformGroup), false);
+			NewTransformSelection.SetFromArray(SelectionArr);
+
+			InTransformSelection = NewTransformSelection;
+		}
+
+		FBox InBoundingBox = GetValue<FBox>(Context, &BoundingBox);
+		//
+		// If not connected set bounds to collection bounds
+		//
+		if (!IsConnected<FBox>(&BoundingBox))
+		{
+			const FManagedArrayCollection& InCollection = GetValue<FManagedArrayCollection>(Context, &Collection);
+
+			GeometryCollection::Facades::FBoundsFacade BoundsFacade(InCollection);
+			const FBox& BoundingBoxInCollectionSpace = BoundsFacade.GetBoundingBoxInCollectionSpace();
+
+			InBoundingBox = BoundingBoxInCollectionSpace;
+		}
+
+		if (InTransformSelection.AnySelected())
+		{
+			FManagedArrayCollection InCollection = GetValue<FManagedArrayCollection>(Context, &Collection);
+
+			int32 ResultGeometryIndex = FFractureEngineFracturing::BrickCutter(InCollection,
+				InTransformSelection,
+				InBoundingBox,
+				GetValue<FTransform>(Context, &Transform),
+				Bond,
+				GetValue<float>(Context, &BrickLength),
+				GetValue<float>(Context, &BrickHeight),
+				GetValue<float>(Context, &BrickDepth),
+				GetValue<int32>(Context, &RandomSeed),
+				GetValue<float>(Context, &ChanceToFracture),
+				SplitIslands,
+				GetValue<float>(Context, &Grout),
+				GetValue<float>(Context, &Amplitude),
+				GetValue<float>(Context, &Frequency),
+				GetValue<float>(Context, &Persistence),
+				GetValue<float>(Context, &Lacunarity),
+				GetValue<int32>(Context, &OctaveNumber),
+				GetValue<float>(Context, &PointSpacing),
+				AddSamplesForCollision,
+				GetValue<float>(Context, &CollisionSampleSpacing));
+
+			FDataflowTransformSelection NewSelection;
+			FDataflowTransformSelection OriginalSelection;
+
+			if (ResultGeometryIndex != INDEX_NONE)
+			{
+				if (InCollection.HasAttribute("TransformIndex", FGeometryCollection::GeometryGroup))
+				{
+					const TManagedArray<int32>& TransformIndices = InCollection.GetAttribute<int32>("TransformIndex", FGeometryCollection::GeometryGroup);
+
+					NewSelection.Initialize(TransformIndices.Num(), false);
+					OriginalSelection.Initialize(TransformIndices.Num(), false);
+
+					// The newly fractured pieces are added to the end of the transform array (starting position is ResultGeometryIndex)
+					for (int32 Idx = ResultGeometryIndex; Idx < TransformIndices.Num(); ++Idx)
+					{
+						int32 BoneIdx = TransformIndices[Idx];
+						NewSelection.SetSelected(BoneIdx);
+					}
+
+					for (int32 Idx = 0; Idx < InTransformSelection.Num(); ++Idx)
+					{
+						if (InTransformSelection.IsSelected(Idx))
+						{
+							OriginalSelection.SetSelected(Idx);
+						}
+					}
+
+				}
+			}
+
+			SetValue(Context, MoveTemp(InCollection), &Collection);
+			SetValue(Context, OriginalSelection, &TransformSelection);
+			SetValue(Context, NewSelection, &NewGeometryTransformSelection);
+
+			return;
+		}
+
+		SafeForwardInput(Context, &Collection, &Collection);
+		SetValue(Context, InTransformSelection, &TransformSelection);
+		SetValue(Context, FDataflowTransformSelection(), &NewGeometryTransformSelection);
+	}
+}
 
 
 
diff --git a/Engine/Plugins/Experimental/GeometryCollectionPlugin/Source/GeometryCollectionNodes/Public/Dataflow/GeometryCollectionFracturingNodes.h b/Engine/Plugins/Experimental/GeometryCollectionPlugin/Source/GeometryCollectionNodes/Public/Dataflow/GeometryCollectionFracturingNodes.h
--- a/Engine/Plugins/Experimental/GeometryCollectionPlugin/Source/GeometryCollectionNodes/Public/Dataflow/GeometryCollectionFracturingNodes.h
+++ b/Engine/Plugins/Experimental/GeometryCollectionPlugin/Source/GeometryCollectionNodes/Public/Dataflow/GeometryCollectionFracturingNodes.h
@@ -6,23 +6,23 @@
 #include "Dataflow/DataflowEngine.h"
 #include "GeometryCollection/ManagedArrayCollection.h"
 #include "Dataflow/DataflowSelection.h"
+#include "FractureEngineFracturing.h"
 
 #include "GeometryCollectionFracturingNodes.generated.h"
 
 class FGeometryCollection;
 
-
-USTRUCT(meta = (DataflowGeometryCollection))
+USTRUCT(meta = (DataflowGeometryCollection, Deprecated = "5.5"))
 struct FUniformScatterPointsDataflowNode : public FDataflowNode
 {
 	GENERATED_USTRUCT_BODY()
 	DATAFLOW_NODE_DEFINE_INTERNAL(FUniformScatterPointsDataflowNode, "UniformScatterPoints", "Generators|Point", "")
 
 public:
-	UPROPERTY(EditAnywhere, Category = "Scatter", meta = (DataflowInput));
+	UPROPERTY(EditAnywhere, Category = "Scatter", meta = (DataflowInput, UIMin = 1));
 	int32 MinNumberOfPoints = 20;
 
-	UPROPERTY(EditAnywhere, Category = "Scatter", meta = (DataflowInput));
+	UPROPERTY(EditAnywhere, Category = "Scatter", meta = (DataflowInput, UIMin = 1));
 	int32 MaxNumberOfPoints = 20;
 
 	UPROPERTY(EditAnywhere, Category = "Scatter", meta = (DataflowInput));
@@ -48,8 +48,48 @@ public:
 
 };
 
-
 USTRUCT(meta = (DataflowGeometryCollection))
+struct FUniformScatterPointsDataflowNode_v2 : public FDataflowNode
+{
+	GENERATED_USTRUCT_BODY()
+	DATAFLOW_NODE_DEFINE_INTERNAL(FUniformScatterPointsDataflowNode_v2, "UniformScatterPoints", "Generators|Point", "")
+
+public:
+	/** Minimum for the random range */
+	UPROPERTY(EditAnywhere, Category = "Scatter", meta = (DataflowInput, UIMin = 1));
+	int32 MinNumberOfPoints = 20;
+
+	/** Maximum for the random range */
+	UPROPERTY(EditAnywhere, Category = "Scatter", meta = (DataflowInput, UIMin = 1));
+	int32 MaxNumberOfPoints = 20;
+
+	/** Seed for random */
+	UPROPERTY(EditAnywhere, Category = "Scatter", meta = (DataflowInput, UIMin = "0"));
+	int32 RandomSeed = 0;
+
+	/** BoundingBox to generate points inside of */
+	UPROPERTY(meta = (DataflowInput))
+	FBox BoundingBox = FBox(ForceInit);
+
+	/** Generated points */
+	UPROPERTY(meta = (DataflowOutput))
+	TArray<FVector> Points;
+
+	FUniformScatterPointsDataflowNode_v2(const Dataflow::FNodeParameters& InParam, FGuid InGuid = FGuid::NewGuid())
+		: FDataflowNode(InParam, InGuid)
+	{
+		RegisterInputConnection(&BoundingBox);
+		RegisterInputConnection(&MinNumberOfPoints);
+		RegisterInputConnection(&MaxNumberOfPoints);
+		RegisterInputConnection(&RandomSeed);
+		RegisterOutputConnection(&Points);
+	}
+
+	virtual void Evaluate(Dataflow::FContext& Context, const FDataflowOutput* Out) const override;
+
+};
+
+USTRUCT(meta = (DataflowGeometryCollection, Deprecated = "5.5"))
 struct FRadialScatterPointsDataflowNode : public FDataflowNode
 {
 	GENERATED_USTRUCT_BODY()
@@ -101,18 +141,167 @@ public:
 
 };
 
+USTRUCT(meta = (DataflowGeometryCollection))
+struct FRadialScatterPointsDataflowNode_v2 : public FDataflowNode
+{
+	GENERATED_USTRUCT_BODY()
+	DATAFLOW_NODE_DEFINE_INTERNAL(FRadialScatterPointsDataflowNode_v2, "RadialScatterPoints", "Generators|Point", "")
+
+public:
+	/** BoundingBox to generate points inside of */
+	UPROPERTY(meta = (DataflowInput))
+	FBox BoundingBox = FBox(ForceInit);
+
+	/** Center of generated pattern */
+	UPROPERTY(EditAnywhere, Category = "Scatter", meta = (DataflowInput));
+	FVector Center = FVector(0.0);
+
+	/** Normal to plane in which sites are generated */
+	UPROPERTY(EditAnywhere, Category = "Scatter", meta = (DataflowInput));
+	FVector Normal = FVector(0.0, 0.0, 1.0);
+
+	/** Seed for random */
+	UPROPERTY(EditAnywhere, Category = "Scatter", meta = (DataflowInput, UIMin = "0"));
+	int32 RandomSeed = 0;
+
+	/** Number of angular steps */
+	UPROPERTY(EditAnywhere, Category = "Scatter", meta = (DataflowInput, UIMin = 1, UIMax = 50, ClampMin = "1"));
+	int32 AngularSteps = 5;
+
+	/** Angle offset at each radial step (in degrees) */
+	UPROPERTY(EditAnywhere, Category = "Scatter", meta = (DataflowInput));
+	float AngleOffset = 0.f;
+
+	/** Amount of global variation to apply to each angular step (in degrees) */
+	UPROPERTY(EditAnywhere, Category = "Scatter", meta = (DataflowInput));
+	float AngularNoise = 0.f;
+
+	/** Pattern radius (in cm) */
+	UPROPERTY(EditAnywhere, Category = "Scatter", meta = (DataflowInput, UIMin = "0.0", ClampMin = "0.0"));
+	float Radius = 50.f;
+
+	/** Number of radial steps */
+	UPROPERTY(EditAnywhere, Category = "Scatter", meta = (DataflowInput, UIMin = "1", UIMax = "50", ClampMin = "1"));
+	int32 RadialSteps = 5;
+
+	/** Radial steps will follow a distribution based on this exponent, i.e., Pow(distance from center, RadialStepExponent) */
+	UPROPERTY(EditAnywhere, Category = "Scatter", meta = (DataflowInput, UIMin = ".01", UIMax = "10", ClampMin = ".01", ClampMax = "20"));
+	float RadialStepExponent = 1.f;
+
+	/** Minimum radial separation between any two voronoi points (in cm) */
+	UPROPERTY(EditAnywhere, Category = "Scatter", meta = (DataflowInput, UIMin = ".25", ClampMin = ".01"))
+	float RadialMinStep = 1.f;
+
+	/** Amount of global variation to apply to each radial step (in cm) */
+	UPROPERTY(EditAnywhere, Category = "Scatter", meta = (DataflowInput, ClampMin = "0"))
+	float RadialNoise = 0.f;
+
+	/** Amount to randomly displace each Voronoi site radially (in cm) */
+	UPROPERTY(EditAnywhere, Category = "Per Point Variability", meta = (DataflowInput, ClampMin = "0.0"))
+	float RadialVariability = 0.f;
+
+	/** Amount to randomly displace each Voronoi site in angle (in degrees) */
+	UPROPERTY(EditAnywhere, Category = "Per Point Variability", meta = (DataflowInput, ClampMin = "0.0"))
+	float AngularVariability = 0.f;
+
+	/** Amount to randomly displace each Voronoi site in the direction of the rotation axis (in cm) */
+	UPROPERTY(EditAnywhere, Category = "Per Point Variability", meta = (DataflowInput, ClampMin = "0.0"))
+	float AxialVariability = 0.f;
+
+	/** Generated points */
+	UPROPERTY(meta = (DataflowOutput))
+	TArray<FVector> Points;
+
+	FRadialScatterPointsDataflowNode_v2(const Dataflow::FNodeParameters& InParam, FGuid InGuid = FGuid::NewGuid())
+		: FDataflowNode(InParam, InGuid)
+	{
+		RegisterInputConnection(&BoundingBox);
+		RegisterInputConnection(&Center);
+		RegisterInputConnection(&Normal);
+		RegisterInputConnection(&RandomSeed);
+		RegisterInputConnection(&AngularSteps);
+		RegisterInputConnection(&AngleOffset);
+		RegisterInputConnection(&AngularNoise);
+		RegisterInputConnection(&Radius);
+		RegisterInputConnection(&RadialSteps);
+		RegisterInputConnection(&RadialStepExponent);
+		RegisterInputConnection(&RadialMinStep);
+		RegisterInputConnection(&RadialNoise);
+		RegisterInputConnection(&RadialVariability);
+		RegisterInputConnection(&AngularVariability);
+		RegisterInputConnection(&AxialVariability);
+		RegisterOutputConnection(&Points);
+	}
+
+	virtual void Evaluate(Dataflow::FContext& Context, const FDataflowOutput* Out) const override;
+};
+
+USTRUCT(meta = (DataflowGeometryCollection))
+struct FGridScatterPointsDataflowNode : public FDataflowNode
+{
+	GENERATED_USTRUCT_BODY()
+	DATAFLOW_NODE_DEFINE_INTERNAL(FGridScatterPointsDataflowNode, "GridScatterPoints", "Generators|Point", "")
 
-//
-// GridScatterPoints
-//
+public:
+	/** Number of points in X direction */
+	UPROPERTY(EditAnywhere, Category = "Scatter", meta = (DataflowInput, UIMin = "0"));
+	int32 NumberOfPointsInX = 5;
 
+	/** Number of points in Y direction */
+	UPROPERTY(EditAnywhere, Category = "Scatter", meta = (DataflowInput, UIMin = "0"));
+	int32 NumberOfPointsInY = 5;
+
+	/** Number of points in Z direction */
+	UPROPERTY(EditAnywhere, Category = "Scatter", meta = (DataflowInput, UIMin = "0"));
+	int32 NumberOfPointsInZ = 5;
+
+	/** Seed for random */
+	UPROPERTY(EditAnywhere, Category = "Scatter", meta = (DataflowInput, UIMin = "0"));
+	int32 RandomSeed = 0;
+
+	/** Random displacement in X direction will be in the range (-MaxRandomDisplacementX, MaxRandomDisplacementX) */
+	UPROPERTY(EditAnywhere, Category = "Scatter", meta = (DataflowInput, UIMin = "0.0"));
+	float MaxRandomDisplacementX = 0.f;
+
+	/** Random displacement in Y direction will be in the range (-MaxRandomDisplacementY, MaxRandomDisplacementY) */
+	UPROPERTY(EditAnywhere, Category = "Scatter", meta = (DataflowInput, UIMin = "0.0"));
+	float MaxRandomDisplacementY = 0.f;
+
+	/** Random displacement in Z direction will be in the range (-MaxRandomDisplacementZ, MaxRandomDisplacementZ) */
+	UPROPERTY(EditAnywhere, Category = "Scatter", meta = (DataflowInput, UIMin = "0.0"));
+	float MaxRandomDisplacementZ = 0.f;
+
+	/** BoundingBox to generate points inside of */
+	UPROPERTY(meta = (DataflowInput))
+	FBox BoundingBox = FBox(ForceInit);
+
+	/** Generated points */
+	UPROPERTY(meta = (DataflowOutput))
+	TArray<FVector> Points;
+
+	FGridScatterPointsDataflowNode(const Dataflow::FNodeParameters& InParam, FGuid InGuid = FGuid::NewGuid())
+		: FDataflowNode(InParam, InGuid)
+	{
+		RegisterInputConnection(&BoundingBox);
+		RegisterInputConnection(&NumberOfPointsInX);
+		RegisterInputConnection(&NumberOfPointsInY);
+		RegisterInputConnection(&NumberOfPointsInZ);
+		RegisterInputConnection(&RandomSeed);
+		RegisterInputConnection(&MaxRandomDisplacementX);
+		RegisterInputConnection(&MaxRandomDisplacementY);
+		RegisterInputConnection(&MaxRandomDisplacementZ);
+		RegisterOutputConnection(&Points);
+	}
+
+	virtual void Evaluate(Dataflow::FContext& Context, const FDataflowOutput* Out) const override;
+};
 
 /**
  *
  * Generates a Voronoi fracture
  *
  */
-USTRUCT(meta = (DataflowGeometryCollection))
+USTRUCT(meta = (DataflowGeometryCollection, Deprecated = "5.5"))
 struct FVoronoiFractureDataflowNode : public FDataflowNode
 {
 	GENERATED_USTRUCT_BODY()
@@ -188,13 +377,123 @@ public:
 
 };
 
+USTRUCT(meta = (DataflowGeometryCollection))
+struct FVoronoiFractureDataflowNode_v2 : public FDataflowNode
+{
+	GENERATED_USTRUCT_BODY()
+	DATAFLOW_NODE_DEFINE_INTERNAL(FVoronoiFractureDataflowNode_v2, "VoronoiFracture", "GeometryCollection|Fracture", "")
+
+public:
+	/** Collection to fracture */
+	UPROPERTY(meta = (DataflowInput, DataflowOutput, DataflowPassthrough = "Collection", DataflowIntrinsic))
+	FManagedArrayCollection Collection;
+
+	/** Voronoi source points */
+	UPROPERTY(meta = (DataflowInput, DataflowIntrinsic))
+	TArray<FVector> Points;
+
+	/** Pieces to fracture */
+	UPROPERTY(meta = (DataflowInput, DataflowOutput, DataflowPassthrough = "TransformSelection", DisplayName = "TransformSelection"))
+	FDataflowTransformSelection TransformSelection;
+
+	/** Transform to apply to cut planes */
+	UPROPERTY(meta = (DataflowInput, DisplayName = "Transform"))
+	FTransform Transform = FTransform::Identity;
+
+	/** Seed for random */
+	UPROPERTY(EditAnywhere, Category = "Common Fracture", meta = (DataflowInput, UIMin = "0"));
+	int32 RandomSeed = 0;
+
+	/** Chance to fracture each selected bone. If 0, no bones will fracture; if 1, all bones will fracture. */
+	UPROPERTY(EditAnywhere, Category = "Common Fracture", meta = (DataflowInput, UIMin = "0.0", UIMax = "1.0", ClampMin = "0.0", ClampMax = "1.0"));
+	float ChanceToFracture = 1.f;
+
+	// Whether to split the fractured mesh pieces based on geometric connectivity after fracturing
+	UPROPERTY(EditAnywhere, Category = "Common Fracture");
+	bool SplitIslands = true;
+
+	/** Amount of space to leave between cut pieces */
+	UPROPERTY(EditAnywhere, Category = "Common Fracture", meta = (DataflowInput, UIMin = 0.f));
+	float Grout = 0.f;
+
+	/** Size of the Perlin noise displacement (in cm). If 0, no noise will be applied */
+	UPROPERTY(EditAnywhere, Category = "Noise", meta = (DataflowInput, UIMin = 0.f));
+	float Amplitude = 0.f;
+
+	/** Period of the Perlin noise.  Smaller values will create a smoother noise pattern */
+	UPROPERTY(EditAnywhere, Category = "Noise", meta = (DataflowInput, UIMin = 0.00001f));
+	float Frequency = 0.1f;
+
+	/** Persistence of the layers of Perlin noise. At each layer (octave) after the first, the amplitude of the Perlin noise is scaled by this factor */
+	UPROPERTY(EditAnywhere, Category = "Noise", meta = (DataflowInput, UIMin = 0.f));
+	float Persistence = 0.5f;
+
+	/** Lacunarity of the layers of Perlin noise. At each layer (octave) after the first, the frequency of the Perlin noise is scaled by this factor */
+	UPROPERTY(EditAnywhere, Category = "Noise", meta = (DataflowInput, UIMin = 0.f));
+	float Lacunarity = 2.f;
+
+	/**
+	 * Number of fractal layers of Perlin noise to apply. Each layer is additive, with Amplitude and Frequency parameters scaled by Persistence and Lacunarity
+	 * Smaller values (1 or 2) will create noise that looks like gentle rolling hills, while larger values (> 4) will tend to look more like craggy mountains
+	 */
+	UPROPERTY(EditAnywhere, Category = "Noise", meta = (DataflowInput, UIMin = 0.f));
+	int32 OctaveNumber = 4;
+
+	/** Distance (in cm) between vertices on cut surfaces where noise is added.  Larger spacing between vertices will create more efficient meshes with fewer triangles, but less resolution to see the shape of the added noise  */
+	UPROPERTY(EditAnywhere, Category = "Noise", meta = (DataflowInput, UIMin = 0.f));
+	float PointSpacing = 10.f;
+
+	/**
+	 * If enabled, add extra vertices (without triangles) to the geometry in regions where vertices are spaced too far apart (e.g. across large triangles)
+	 * These extra vertices will be used as collision samples in particle-implicit collisions, and can help the physics system detect collisions more accurately
+	 *
+	 * Note this is *only* useful for simulations that use particle-implicit collisions
+	 */
+	UPROPERTY(EditAnywhere, Category = "Collision");
+	bool AddSamplesForCollision = false;
+
+	/**
+	 * The number of centimeters to allow between vertices on the mesh surface: If there are gaps larger than this, add additional vertices (without triangles) to help support particle-implicit collisions
+	 * Only used if Add Samples For Collision is enabled
+	 */
+	UPROPERTY(EditAnywhere, Category = "Collision", meta = (DataflowInput, UIMin = 0.f, EditCondition = "AddSamplesForCollision"));
+	float CollisionSampleSpacing = 50.f;
+
+	/** Fractured Pieces */
+	UPROPERTY(meta = (DataflowOutput, DisplayName = "FracturedTransformSelection "))
+	FDataflowTransformSelection NewGeometryTransformSelection;
+
+	FVoronoiFractureDataflowNode_v2(const Dataflow::FNodeParameters& InParam, FGuid InGuid = FGuid::NewGuid())
+		: FDataflowNode(InParam, InGuid)
+	{
+		RegisterInputConnection(&Collection);
+		RegisterInputConnection(&Points);
+		RegisterInputConnection(&TransformSelection);
+		RegisterInputConnection(&Transform);
+		RegisterInputConnection(&RandomSeed);
+		RegisterInputConnection(&ChanceToFracture);
+		RegisterInputConnection(&Grout);
+		RegisterInputConnection(&Amplitude);
+		RegisterInputConnection(&Frequency);
+		RegisterInputConnection(&Persistence);
+		RegisterInputConnection(&Lacunarity);
+		RegisterInputConnection(&OctaveNumber);
+		RegisterInputConnection(&PointSpacing);
+		RegisterInputConnection(&CollisionSampleSpacing);
+		RegisterOutputConnection(&Collection, &Collection);
+		RegisterOutputConnection(&TransformSelection, &TransformSelection);
+		RegisterOutputConnection(&NewGeometryTransformSelection);
+	}
+
+	virtual void Evaluate(Dataflow::FContext& Context, const FDataflowOutput* Out) const override;
+};
 
 /**
  *
  * Cuts geometry using a set of noised up planes
  *
  */
-USTRUCT(meta = (DataflowGeometryCollection))
+USTRUCT(meta = (DataflowGeometryCollection, Deprecated = "5.5"))
 struct FPlaneCutterDataflowNode : public FDataflowNode
 {
 	GENERATED_USTRUCT_BODY()
@@ -265,9 +564,123 @@ public:
 	}
 
 	virtual void Evaluate(Dataflow::FContext& Context, const FDataflowOutput* Out) const override;
-
 };
 
+USTRUCT(meta = (DataflowGeometryCollection))
+struct FPlaneCutterDataflowNode_v2 : public FDataflowNode
+{
+	GENERATED_USTRUCT_BODY()
+	DATAFLOW_NODE_DEFINE_INTERNAL(FPlaneCutterDataflowNode_v2, "PlaneCutter", "GeometryCollection|Fracture", "")
+
+public:
+	/** Collection to fracture */
+	UPROPERTY(meta = (DataflowInput, DataflowOutput, DataflowPassthrough = "Collection", DataflowIntrinsic))
+	FManagedArrayCollection Collection;
+
+	/** Bound for plane centers */
+	UPROPERTY(meta = (DataflowInput, DataflowIntrinsic))
+	FBox BoundingBox = FBox(ForceInit);
+
+	/** Pieces to cut */
+	UPROPERTY(meta = (DataflowInput, DataflowOutput, DataflowPassthrough = "TransformSelection", DisplayName = "TransformSelection"))
+	FDataflowTransformSelection TransformSelection;
+
+	/** Transform to apply to cut planes */
+	UPROPERTY(meta = (DataflowInput, DisplayName = "Transform"))
+	FTransform Transform = FTransform::Identity;
+
+	/** Number of cutting planes */
+	UPROPERTY(EditAnywhere, Category = "Plane Cut", meta = (DataflowInput, DisplayName = "Number of Cuts", UIMin = 1))
+	int32 NumPlanes = 1;
+
+	/** Seed for random */
+	UPROPERTY(EditAnywhere, Category = "Common Fracture", meta = (DataflowInput, UIMin = "0"))
+	int32 RandomSeed = 0;
+
+	/** Chance to fracture each selected bone. If 0, no bones will fracture; if 1, all bones will fracture. */
+	UPROPERTY(EditAnywhere, Category = "Common Fracture", meta = (DataflowInput, UIMin = 0.f, DisplayName = "Chance To Fracture Per Bone", UIMax = 1.f));
+	float ChanceToFracture = 1.f;
+
+	// Whether to split the fractured mesh pieces based on geometric connectivity after fracturing
+	UPROPERTY(EditAnywhere, Category = "Common Fracture");
+	bool SplitIslands = true;
+
+	/** Amount of space to leave between cut pieces */
+	UPROPERTY(EditAnywhere, Category = "Common Fracture", meta = (DataflowInput, UIMin = 0.f));
+	float Grout = 0.f;
+
+	/** Size of the Perlin noise displacement (in cm). If 0, no noise will be applied */
+	UPROPERTY(EditAnywhere, Category = "Noise", meta = (DataflowInput, UIMin = 0.f));
+	float Amplitude = 0.f;
+
+	/** Period of the Perlin noise.  Smaller values will create a smoother noise pattern */
+	UPROPERTY(EditAnywhere, Category = "Noise", meta = (DataflowInput, UIMin = 0.00001f));
+	float Frequency = 0.1f;
+
+	/** Persistence of the layers of Perlin noise. At each layer (octave) after the first, the amplitude of the Perlin noise is scaled by this factor */
+	UPROPERTY(EditAnywhere, Category = "Noise", meta = (DataflowInput, UIMin = 0.f));
+	float Persistence = 0.5f;
+
+	/** Lacunarity of the layers of Perlin noise. At each layer (octave) after the first, the frequency of the Perlin noise is scaled by this factor */
+	UPROPERTY(EditAnywhere, Category = "Noise", meta = (DataflowInput, UIMin = 0.f));
+	float Lacunarity = 2.f;
+
+	/**
+	 * Number of fractal layers of Perlin noise to apply. Each layer is additive, with Amplitude and Frequency parameters scaled by Persistence and Lacunarity
+	 * Smaller values (1 or 2) will create noise that looks like gentle rolling hills, while larger values (> 4) will tend to look more like craggy mountains
+	 */
+	UPROPERTY(EditAnywhere, Category = "Noise", meta = (DataflowInput, UIMin = 0.f));
+	int32 OctaveNumber = 4;
+
+	/** Distance (in cm) between vertices on cut surfaces where noise is added.  Larger spacing between vertices will create more efficient meshes with fewer triangles, but less resolution to see the shape of the added noise  */
+	UPROPERTY(EditAnywhere, Category = "Noise", meta = (DataflowInput, UIMin = 0.f));
+	float PointSpacing = 10.f;
+
+	/**
+	 * If enabled, add extra vertices (without triangles) to the geometry in regions where vertices are spaced too far apart (e.g. across large triangles)
+	 * These extra vertices will be used as collision samples in particle-implicit collisions, and can help the physics system detect collisions more accurately
+	 *
+	 * Note this is *only* useful for simulations that use particle-implicit collisions
+	 */
+	UPROPERTY(EditAnywhere, Category = "Collision");
+	bool AddSamplesForCollision = false;
+
+	/**
+	 * The number of centimeters to allow between vertices on the mesh surface: If there are gaps larger than this, add additional vertices (without triangles) to help support particle-implicit collisions
+	 * Only used if Add Samples For Collision is enabled
+	 */
+	UPROPERTY(EditAnywhere, Category = "Collision", meta = (DataflowInput, UIMin = 0.f, EditCondition = "AddSamplesForCollision"));
+	float CollisionSampleSpacing = 50.f;
+
+	/** Fractured Pieces */
+	UPROPERTY(meta = (DataflowOutput, DisplayName = "FracturedTransformSelection "))
+	FDataflowTransformSelection NewGeometryTransformSelection;
+
+	FPlaneCutterDataflowNode_v2(const Dataflow::FNodeParameters& InParam, FGuid InGuid = FGuid::NewGuid())
+		: FDataflowNode(InParam, InGuid)
+	{
+		RegisterInputConnection(&Collection);
+		RegisterInputConnection(&BoundingBox);
+		RegisterInputConnection(&TransformSelection);
+		RegisterInputConnection(&Transform);
+		RegisterInputConnection(&NumPlanes);
+		RegisterInputConnection(&RandomSeed);
+		RegisterInputConnection(&ChanceToFracture);
+		RegisterInputConnection(&Grout);
+		RegisterInputConnection(&Amplitude);
+		RegisterInputConnection(&Frequency);
+		RegisterInputConnection(&Persistence);
+		RegisterInputConnection(&Lacunarity);
+		RegisterInputConnection(&OctaveNumber);
+		RegisterInputConnection(&PointSpacing);
+		RegisterInputConnection(&CollisionSampleSpacing);
+		RegisterOutputConnection(&Collection, &Collection);
+		RegisterOutputConnection(&TransformSelection, &TransformSelection);
+		RegisterOutputConnection(&NewGeometryTransformSelection);
+	}
+
+	virtual void Evaluate(Dataflow::FContext& Context, const FDataflowOutput* Out) const override;
+};
 
 /**
  *
@@ -282,12 +695,15 @@ struct FExplodedViewDataflowNode : public FDataflowNode
 	DATAFLOW_NODE_RENDER_TYPE("SurfaceRender",FGeometryCollection::StaticType(),  "Collection")
 
 public:
+	/** Collection to explode */
 	UPROPERTY(meta = (DataflowInput, DataflowOutput, DataflowPassthrough = "Collection", DataflowIntrinsic))
 	FManagedArrayCollection Collection;
 
+	/** Scale amount to expand the pieces uniformly in all directions */
 	UPROPERTY(EditAnywhere, Category = "Scale", meta = (DataflowInput))
 	float UniformScale = 1.f;
 
+	/** Scale amounts to expand the pieces in all 3 directions */
 	UPROPERTY(EditAnywhere, Category = "Scale", meta = (DataflowInput))
 	FVector Scale = FVector(1.0);
 
@@ -305,9 +721,288 @@ public:
 private:
 	// todo(chaos) this is a copy of a function in FractureEditorModeToolkit, we should move this to a common place  
 	static bool GetValidGeoCenter(FGeometryCollection* Collection, const TManagedArray<int32>& TransformToGeometryIndex, const TArray<FTransform>& Transforms, const TManagedArray<TSet<int32>>& Children, const TManagedArray<FBox>& BoundingBox, int32 TransformIndex, FVector& OutGeoCenter);
+};
+
+/**
+ *
+ * Slice tool
+ *
+ */
+USTRUCT(meta = (DataflowGeometryCollection))
+struct FSliceCutterDataflowNode : public FDataflowNode
+{
+	GENERATED_USTRUCT_BODY()
+	DATAFLOW_NODE_DEFINE_INTERNAL(FSliceCutterDataflowNode, "SliceCutter", "GeometryCollection|Fracture", "")
+
+public:
+	/** Collection to fracture */
+	UPROPERTY(meta = (DataflowInput, DataflowOutput, DataflowPassthrough = "Collection", DataflowIntrinsic))
+	FManagedArrayCollection Collection;
+
+	UPROPERTY(meta = (DataflowInput, DataflowIntrinsic))
+	FBox BoundingBox = FBox(ForceInit);
+
+	/** The selected pieces to cut */
+	UPROPERTY(meta = (DataflowInput, DataflowOutput, DataflowPassthrough = "TransformSelection", DisplayName = "TransformSelection"))
+	FDataflowTransformSelection TransformSelection;
+
+	/** Number of slices along the X axis */
+	UPROPERTY(EditAnywhere, Category = Slicing, meta = (DataflowInput, UIMin = "0"))
+	int32 SlicesX = 3;
+
+	/** Number of slices along the Y axis */
+	UPROPERTY(EditAnywhere, Category = Slicing, meta = (DataflowInput, UIMin = "0"))
+	int32 SlicesY = 3;
+
+	/** Number of slices along the Z axis */
+	UPROPERTY(EditAnywhere, Category = Slicing, meta = (DataflowInput, UIMin = "0"))
+	int32 SlicesZ = 1;
+
+	/** Maximum angle (in degrees) to randomly rotate each slicing plane */
+	UPROPERTY(EditAnywhere, Category = Slicing, meta = (DataflowInput, DisplayName = "Random Angle Variation", UIMin = "0.0", UIMax = "90.0"))
+	float SliceAngleVariation = 0.f;
+
+	/** Maximum distance (in cm) to randomly shift each slicing plane */
+	UPROPERTY(EditAnywhere, Category = Slicing, meta = (DataflowInput, DisplayName = "Random Offset Variation", UIMin = "0.0"))
+	float SliceOffsetVariation = 0.f;
+
+	/** Seed for random */
+	UPROPERTY(EditAnywhere, Category = "Common Fracture", meta = (DataflowInput, UIMin = "0"))
+	int32 RandomSeed = 0;
+
+	/** Chance to fracture each selected bone. If 0, no bones will fracture; if 1, all bones will fracture. */
+	UPROPERTY(EditAnywhere, Category = "Common Fracture", meta = (DataflowInput, UIMin = 0.f, UIMax = 1.f));
+	float ChanceToFracture = 1.f;
+
+	// Whether to split the fractured mesh pieces based on geometric connectivity after fracturing
+	UPROPERTY(EditAnywhere, Category = "Common Fracture");
+	bool SplitIslands = true;
+
+	/** Amount of space to leave between cut pieces */
+	UPROPERTY(EditAnywhere, Category = "Common Fracture", meta = (DataflowInput, UIMin = 0.f));
+	float Grout = 0.f;
+
+	/** Size of the Perlin noise displacement (in cm). If 0, no noise will be applied */
+	UPROPERTY(EditAnywhere, Category = "Noise", meta = (DataflowInput, UIMin = 0.f));
+	float Amplitude = 0.f;
+
+	/** Period of the Perlin noise.  Smaller values will create a smoother noise pattern */
+	UPROPERTY(EditAnywhere, Category = "Noise", meta = (DataflowInput, UIMin = 0.00001f));
+	float Frequency = 0.1f;
+
+	/** Persistence of the layers of Perlin noise. At each layer (octave) after the first, the amplitude of the Perlin noise is scaled by this factor */
+	UPROPERTY(EditAnywhere, Category = "Noise", meta = (DataflowInput, UIMin = 0.f));
+	float Persistence = 0.5f;
+
+	/** Lacunarity of the layers of Perlin noise. At each layer (octave) after the first, the frequency of the Perlin noise is scaled by this factor */
+	UPROPERTY(EditAnywhere, Category = "Noise", meta = (DataflowInput, UIMin = 0.f));
+	float Lacunarity = 2.f;
+
+	/**
+	 * Number of fractal layers of Perlin noise to apply. Each layer is additive, with Amplitude and Frequency parameters scaled by Persistence and Lacunarity
+	 * Smaller values (1 or 2) will create noise that looks like gentle rolling hills, while larger values (> 4) will tend to look more like craggy mountains
+	 */
+	UPROPERTY(EditAnywhere, Category = "Noise", meta = (DataflowInput, UIMin = 0.f));
+	int32 OctaveNumber = 4;
 
+	/** Distance (in cm) between vertices on cut surfaces where noise is added.  Larger spacing between vertices will create more efficient meshes with fewer triangles, but less resolution to see the shape of the added noise  */
+	UPROPERTY(EditAnywhere, Category = "Noise", meta = (DataflowInput, UIMin = 0.f));
+	float PointSpacing = 10.f;
+
+	/**
+	 * If enabled, add extra vertices (without triangles) to the geometry in regions where vertices are spaced too far apart (e.g. across large triangles)
+	 * These extra vertices will be used as collision samples in particle-implicit collisions, and can help the physics system detect collisions more accurately
+	 *
+	 * Note this is *only* useful for simulations that use particle-implicit collisions
+	 */
+	UPROPERTY(EditAnywhere, Category = "Collision");
+	bool AddSamplesForCollision = false;
+
+	/**
+	 * The number of centimeters to allow between vertices on the mesh surface: If there are gaps larger than this, add additional vertices (without triangles) to help support particle-implicit collisions
+	 * Only used if Add Samples For Collision is enabled
+	 */
+	UPROPERTY(EditAnywhere, Category = "Collision", meta = (DataflowInput, UIMin = 0.f, EditCondition = "AddSamplesForCollision"));
+	float CollisionSampleSpacing = 50.f;
+
+	/** Fractured Pieces */
+	UPROPERTY(meta = (DataflowOutput, DisplayName = "FracturedTransformSelection "))
+	FDataflowTransformSelection NewGeometryTransformSelection;
+
+	FSliceCutterDataflowNode(const Dataflow::FNodeParameters& InParam, FGuid InGuid = FGuid::NewGuid())
+		: FDataflowNode(InParam, InGuid)
+	{
+		RegisterInputConnection(&Collection);
+		RegisterInputConnection(&BoundingBox);
+		RegisterInputConnection(&TransformSelection);
+		RegisterInputConnection(&SlicesX);
+		RegisterInputConnection(&SlicesY);
+		RegisterInputConnection(&SlicesZ);
+		RegisterInputConnection(&SliceAngleVariation);
+		RegisterInputConnection(&SliceOffsetVariation);
+		RegisterInputConnection(&RandomSeed);
+		RegisterInputConnection(&ChanceToFracture);
+		RegisterInputConnection(&Grout);
+		RegisterInputConnection(&Amplitude);
+		RegisterInputConnection(&Frequency);
+		RegisterInputConnection(&Persistence);
+		RegisterInputConnection(&Lacunarity);
+		RegisterInputConnection(&OctaveNumber);
+		RegisterInputConnection(&PointSpacing);
+		RegisterInputConnection(&CollisionSampleSpacing);
+		RegisterOutputConnection(&Collection, &Collection);
+		RegisterOutputConnection(&TransformSelection, &TransformSelection);
+		RegisterOutputConnection(&NewGeometryTransformSelection);
+	}
+
+	virtual void Evaluate(Dataflow::FContext& Context, const FDataflowOutput* Out) const override;
 };
 
+/**
+ *
+ * Brick tool
+ *
+ */
+USTRUCT(meta = (DataflowGeometryCollection))
+struct FBrickCutterDataflowNode : public FDataflowNode
+{
+	GENERATED_USTRUCT_BODY()
+	DATAFLOW_NODE_DEFINE_INTERNAL(FBrickCutterDataflowNode, "BrickCutter", "GeometryCollection|Fracture", "")
+
+public:
+	/** Collection to cut */
+	UPROPERTY(meta = (DataflowInput, DataflowOutput, DataflowPassthrough = "Collection", DataflowIntrinsic))
+	FManagedArrayCollection Collection;
+
+	/** Boundingbox to create the cutting planes in */
+	UPROPERTY(meta = (DataflowInput))
+	FBox BoundingBox = FBox(ForceInit);
+
+	/** The selected pieces to cut */
+	UPROPERTY(meta = (DataflowInput, DataflowOutput, DataflowPassthrough = "TransformSelection", DisplayName = "TransformSelection"))
+	FDataflowTransformSelection TransformSelection;
+
+	/** Transform to apply to cut planes */
+	UPROPERTY(meta = (DataflowInput, DisplayName = "Transform"))
+	FTransform Transform = FTransform::Identity;
+
+	/** The brick bond pattern defines how the bricks are arranged */
+	UPROPERTY(EditAnywhere, Category = Brick)
+	EFractureBrickBondEnum Bond = EFractureBrickBondEnum::Dataflow_FractureBrickBond_Stretcher;
+
+	/** Brick length (in cm) */
+	UPROPERTY(EditAnywhere, Category = Brick, meta = (DataflowInput, UIMin = "0.1", UIMax = "500.0", ClampMin = "0.001"))
+	float BrickLength = 194.f;
+
+	/** Brick height (in cm) */
+	UPROPERTY(EditAnywhere, Category = Brick, meta = (DataflowInput, UIMin = "0.1", UIMax = "500.0", ClampMin = "0.001"))
+	float BrickHeight = 57.f;
+
+	/** Brick depth (in cm) */
+	UPROPERTY(EditAnywhere, Category = Brick, meta = (DataflowInput, UIMin = "0.1", UIMax = "500.0", ClampMin = "0.001"))
+	float BrickDepth = 92.f;
+
+	/** Seed for random */
+	UPROPERTY(EditAnywhere, Category = "Common Fracture", meta = (DataflowInput, UIMin = "0"))
+	int32 RandomSeed = 0;
+
+	/** Chance to fracture each selected bone. If 0, no bones will fracture; if 1, all bones will fracture. */
+	UPROPERTY(EditAnywhere, Category = "Common Fracture", meta = (DataflowInput, UIMin = 0.f, UIMax = 1.f));
+	float ChanceToFracture = 1.f;
+
+	// Whether to split the fractured mesh pieces based on geometric connectivity after fracturing
+	UPROPERTY(EditAnywhere, Category = "Common Fracture");
+	bool SplitIslands = true;
+
+	/** Amount of space to leave between cut pieces */
+	UPROPERTY(EditAnywhere, Category = "Common Fracture", meta = (DataflowInput, UIMin = 0.f));
+	float Grout = 0.f;
+
+	/** Size of the Perlin noise displacement (in cm). If 0, no noise will be applied */
+	UPROPERTY(EditAnywhere, Category = "Noise", meta = (DataflowInput, UIMin = 0.f));
+	float Amplitude = 0.f;
+
+	/** Period of the Perlin noise.  Smaller values will create a smoother noise pattern */
+	UPROPERTY(EditAnywhere, Category = "Noise", meta = (DataflowInput, UIMin = 0.00001f));
+	float Frequency = 0.1f;
+
+	/** Persistence of the layers of Perlin noise. At each layer (octave) after the first, the amplitude of the Perlin noise is scaled by this factor */
+	UPROPERTY(EditAnywhere, Category = "Noise", meta = (DataflowInput, UIMin = 0.f));
+	float Persistence = 0.5f;
+
+	/** Lacunarity of the layers of Perlin noise. At each layer (octave) after the first, the frequency of the Perlin noise is scaled by this factor */
+	UPROPERTY(EditAnywhere, Category = "Noise", meta = (DataflowInput, UIMin = 0.f));
+	float Lacunarity = 2.f;
+
+	/**
+	 * Number of fractal layers of Perlin noise to apply. Each layer is additive, with Amplitude and Frequency parameters scaled by Persistence and Lacunarity
+	 * Smaller values (1 or 2) will create noise that looks like gentle rolling hills, while larger values (> 4) will tend to look more like craggy mountains
+	 */
+	UPROPERTY(EditAnywhere, Category = "Noise", meta = (DataflowInput, UIMin = 0.f));
+	int32 OctaveNumber = 4;
+
+	/** Distance (in cm) between vertices on cut surfaces where noise is added.  Larger spacing between vertices will create more efficient meshes with fewer triangles, but less resolution to see the shape of the added noise  */
+	UPROPERTY(EditAnywhere, Category = "Noise", meta = (DataflowInput, UIMin = 0.f));
+	float PointSpacing = 10.f;
+
+	/**
+	 * If enabled, add extra vertices (without triangles) to the geometry in regions where vertices are spaced too far apart (e.g. across large triangles)
+	 * These extra vertices will be used as collision samples in particle-implicit collisions, and can help the physics system detect collisions more accurately
+	 *
+	 * Note this is *only* useful for simulations that use particle-implicit collisions
+	 */
+	UPROPERTY(EditAnywhere, Category = "Collision");
+	bool AddSamplesForCollision = false;
+
+	/**
+	 * The number of centimeters to allow between vertices on the mesh surface: If there are gaps larger than this, add additional vertices (without triangles) to help support particle-implicit collisions
+	 * Only used if Add Samples For Collision is enabled
+	 */
+	UPROPERTY(EditAnywhere, Category = "Collision", meta = (DataflowInput, UIMin = 0.f, EditCondition = "AddSamplesForCollision"));
+	float CollisionSampleSpacing = 50.f;
+
+	/** Fractured Pieces */
+	UPROPERTY(meta = (DataflowOutput, DisplayName = "FracturedTransformSelection "))
+	FDataflowTransformSelection NewGeometryTransformSelection;
+
+	FBrickCutterDataflowNode(const Dataflow::FNodeParameters& InParam, FGuid InGuid = FGuid::NewGuid())
+		: FDataflowNode(InParam, InGuid)
+	{
+		RegisterInputConnection(&Collection);
+		RegisterInputConnection(&BoundingBox);
+		RegisterInputConnection(&TransformSelection);
+		RegisterInputConnection(&Transform);
+		RegisterInputConnection(&BrickLength);
+		RegisterInputConnection(&BrickHeight);
+		RegisterInputConnection(&BrickDepth);
+		RegisterInputConnection(&RandomSeed);
+		RegisterInputConnection(&ChanceToFracture);
+		RegisterInputConnection(&Grout);
+		RegisterInputConnection(&Amplitude);
+		RegisterInputConnection(&Frequency);
+		RegisterInputConnection(&Persistence);
+		RegisterInputConnection(&Lacunarity);
+		RegisterInputConnection(&OctaveNumber);
+		RegisterInputConnection(&PointSpacing);
+		RegisterInputConnection(&CollisionSampleSpacing);
+		RegisterOutputConnection(&Collection, &Collection);
+		RegisterOutputConnection(&TransformSelection, &TransformSelection);
+		RegisterOutputConnection(&NewGeometryTransformSelection);
+	}
+
+	virtual void Evaluate(Dataflow::FContext& Context, const FDataflowOutput* Out) const override;
+};
+
+
+
+
+
+
+
+
+
+
+
 
 
 
diff --git a/Engine/Source/Runtime/Experimental/Dataflow/Core/Private/Dataflow/DataflowNode.cpp b/Engine/Source/Runtime/Experimental/Dataflow/Core/Private/Dataflow/DataflowNode.cpp
--- a/Engine/Source/Runtime/Experimental/Dataflow/Core/Private/Dataflow/DataflowNode.cpp
+++ b/Engine/Source/Runtime/Experimental/Dataflow/Core/Private/Dataflow/DataflowNode.cpp
@@ -796,6 +796,7 @@ FDataflowInput& FDataflowNode::RegisterInputConnectionInternal(const Dataflow::F
 	check(Input->RealAddress() == Reference.Reference);
 	AddInput(Input);
 	check(FindInput(Reference) == Input);
+
 	return *Input;
 }
 
@@ -976,6 +977,7 @@ bool FDataflowNode::ValidateConnections()
 				{
 					if (!FindInput(PropertyIt.Value()))
 					{
+						ensure(false);
 						UE_LOG(LogChaos, Warning, TEXT("Missing dataflow RegisterInputConnection in constructor for (%s:%s)"), *GetName().ToString(), *PropName.ToString())
 							bHasValidConnections = false;
 					}
@@ -985,6 +987,7 @@ bool FDataflowNode::ValidateConnections()
 					const FDataflowOutput* const OutputConnection = FindOutput(PropertyIt.Value());
 					if(!OutputConnection)
 					{
+						ensure(false);
 						UE_LOG(LogChaos, Warning, TEXT("Missing dataflow RegisterOutputConnection in constructor for (%s:%s)"), *GetName().ToString(),*PropName.ToString());
 						bHasValidConnections = false;
 					}
@@ -1006,6 +1009,7 @@ bool FDataflowNode::ValidateConnections()
 						const FDataflowInput* const PassthroughConnectionInput = OutputConnection->GetPassthroughInput();
 						if (PassthroughConnectionInput == nullptr)
 						{
+							ensure(false);
 							UE_LOG(LogChaos, Warning, TEXT("Missing DataflowPassthrough registration for (%s:%s)"), *GetName().ToString(), *PropName.ToString());
 							bHasValidConnections = false;
 						}
@@ -1014,29 +1018,63 @@ bool FDataflowNode::ValidateConnections()
 
 						if(PassthroughConnectionInput != PassthroughConnectionInputFromMetadata)
 						{
+							ensure(false);
 							UE_LOG(LogChaos, Warning, TEXT("Mismatch in declared and registered DataflowPassthrough connection; (%s:%s vs %s)"), *GetName().ToString(), *FullPassthroughName, *PassthroughConnectionInput->GetName().ToString());
 							bHasValidConnections = false;
 						}
 
 						if(!PassthroughConnectionInputFromMetadata)
 						{
+							ensure(false);
 							UE_LOG(LogChaos, Warning, TEXT("Incorrect DataflowPassthrough Connection set for (%s:%s)"), *GetName().ToString(), *PropName.ToString());
 							bHasValidConnections = false;
 						}
 
 						else if(OutputConnection->GetType() != PassthroughConnectionInput->GetType())
 						{
+							ensure(false);
 							UE_LOG(LogChaos, Warning, TEXT("DataflowPassthrough connection types mismatch for (%s:%s)"), *GetName().ToString(), *PropName.ToString());
 							bHasValidConnections = false;
 						}
 					}
 					else if(OutputConnection->GetPassthroughInput()) 
 					{
+						ensure(false);
 						UE_LOG(LogChaos, Warning, TEXT("Missing DataflowPassthrough declaration for (%s:%s)"), *GetName().ToString(), *PropName.ToString());
 						bHasValidConnections = false;
 					}
 				}
 			}
+		
+			for (const TPair<Dataflow::FConnectionKey, FDataflowInput*>& ExpandedInput : ExpandedInputs)
+			{
+				const FDataflowInput* Input = ExpandedInput.Value;
+
+				if (const FProperty* Property = Input->GetProperty())
+				{
+					if (!Property->HasMetaData(FDataflowNode::DataflowInput))
+					{
+						ensure(false);
+						UE_LOG(LogChaos, Warning, TEXT("Missing dataflow DataflowInput declaration for (%s:%s)"), *GetName().ToString(), *Property->GetFName().ToString());
+					}
+				}
+			}
+
+PRAGMA_DISABLE_DEPRECATION_WARNINGS  // Until Outputs becomes private
+			for (const TPair<int, FDataflowOutput*>& Output : Outputs)
+PRAGMA_ENABLE_DEPRECATION_WARNINGS
+			{
+				FDataflowOutput* OutputVal = Output.Value;
+
+				if (const FProperty* Property = OutputVal->GetProperty())
+				{
+					if (!Property->HasMetaData(FDataflowNode::DataflowOutput))
+					{
+						ensure(false);
+						UE_LOG(LogChaos, Warning, TEXT("Missing dataflow DataflowOutput declaration for (%s:%s)"), *GetName().ToString(), *Property->GetFName().ToString());
+					}
+				}
+			}
 		}
 	}
 #endif
diff --git a/Engine/Source/Runtime/Experimental/Dataflow/Core/Private/Dataflow/DataflowSelection.cpp b/Engine/Source/Runtime/Experimental/Dataflow/Core/Private/Dataflow/DataflowSelection.cpp
--- a/Engine/Source/Runtime/Experimental/Dataflow/Core/Private/Dataflow/DataflowSelection.cpp
+++ b/Engine/Source/Runtime/Experimental/Dataflow/Core/Private/Dataflow/DataflowSelection.cpp
@@ -15,6 +15,11 @@ void FDataflowSelection::Initialize(const FDataflowSelection& Other)
 	SelectionArray = Other.SelectionArray;
 }
 
+void FDataflowSelection::Clear()
+{
+	SelectionArray.Init(false, SelectionArray.Num());
+}
+
 void FDataflowSelection::AsArray(TArray<int32>& SelectionArr) const
 {
 	SelectionArr.Reset(SelectionArray.Num());
diff --git a/Engine/Source/Runtime/Experimental/Dataflow/Core/Public/Dataflow/DataflowSelection.h b/Engine/Source/Runtime/Experimental/Dataflow/Core/Public/Dataflow/DataflowSelection.h
--- a/Engine/Source/Runtime/Experimental/Dataflow/Core/Public/Dataflow/DataflowSelection.h
+++ b/Engine/Source/Runtime/Experimental/Dataflow/Core/Public/Dataflow/DataflowSelection.h
@@ -11,6 +11,7 @@ struct FDataflowSelection
 
 	DATAFLOWCORE_API void Initialize(int32 NumBits, bool Value);
 	DATAFLOWCORE_API void Initialize(const FDataflowSelection& Other);
+	DATAFLOWCORE_API void Clear();
 	int32 Num() const { return SelectionArray.Num(); }
 	DATAFLOWCORE_API int32 NumSelected() const;
 	DATAFLOWCORE_API bool AnySelected() const;
