make more dataflow clustering nodes compute the cluster levels, since the clustering code expects levels to be available

diff --git a/Engine/Plugins/Experimental/GeometryCollectionPlugin/Source/GeometryCollectionNodes/Private/Dataflow/GeometryCollectionClusteringNodes.cpp b/Engine/Plugins/Experimental/GeometryCollectionPlugin/Source/GeometryCollectionNodes/Private/Dataflow/GeometryCollectionClusteringNodes.cpp
--- a/Engine/Plugins/Experimental/GeometryCollectionPlugin/Source/GeometryCollectionNodes/Private/Dataflow/GeometryCollectionClusteringNodes.cpp
+++ b/Engine/Plugins/Experimental/GeometryCollectionPlugin/Source/GeometryCollectionNodes/Private/Dataflow/GeometryCollectionClusteringNodes.cpp
@@ -63,6 +63,9 @@ void FAutoClusterDataflowNode::Evaluate(Dataflow::FContext& Context, const FData
 
 		if (TUniquePtr<FGeometryCollection> GeomCollection = TUniquePtr<FGeometryCollection>(InCollection.NewCopy<FGeometryCollection>()))
 		{
+			Chaos::Facades::FCollectionHierarchyFacade HierarchyFacade(*GeomCollection);
+			HierarchyFacade.GenerateLevelAttribute();
+
 			EClusterSizeMethodEnum InClusterSizeMethod = ClusterSizeMethod;
 			int32 InClusterSites = GetValue<int32>(Context, &ClusterSites);
 			float InClusterFraction = GetValue<float>(Context, &ClusterFraction);
@@ -199,6 +202,9 @@ void FClusterMergeToNeighborsDataflowNode::Evaluate(Dataflow::FContext& Context,
 		}
 		if (TUniquePtr<FGeometryCollection> GeomCollection = TUniquePtr<FGeometryCollection>(InCollection.NewCopy<FGeometryCollection>()))
 		{
+			Chaos::Facades::FCollectionHierarchyFacade HierarchyFacade(*GeomCollection);
+			HierarchyFacade.GenerateLevelAttribute();
+
 			double InMinVolumeCubeRoot = (double)GetValue(Context, &MinVolumeCubeRoot);
 			double InMinVolume = InMinVolumeCubeRoot * InMinVolumeCubeRoot * InMinVolumeCubeRoot;
 			bool bInOnlyToConnected = GetValue(Context, &bOnlyToConnected);
@@ -238,6 +244,9 @@ void FClusterMergeDataflowNode::Evaluate(Dataflow::FContext& Context, const FDat
 		const FDataflowTransformSelection& InTransformSelection = GetValue(Context, &TransformSelection);
 		if (TUniquePtr<FGeometryCollection> GeomCollection = TUniquePtr<FGeometryCollection>(InCollection.NewCopy<FGeometryCollection>()))
 		{
+			Chaos::Facades::FCollectionHierarchyFacade HierarchyFacade(*GeomCollection);
+			HierarchyFacade.GenerateLevelAttribute();
+
 			TArray<int32> Selection = InTransformSelection.AsArray();
 			FFractureEngineClustering::MergeSelectedClusters(*GeomCollection, Selection);
 			SetValue(Context, (const FManagedArrayCollection&)(*GeomCollection), &Collection);
