Fix UHT warnings with AutoCluster dataflow node

diff --git a/Engine/Plugins/Experimental/GeometryCollectionPlugin/Source/GeometryCollectionNodes/Public/Dataflow/GeometryCollectionClusteringNodes.h b/Engine/Plugins/Experimental/GeometryCollectionPlugin/Source/GeometryCollectionNodes/Public/Dataflow/GeometryCollectionClusteringNodes.h
--- a/Engine/Plugins/Experimental/GeometryCollectionPlugin/Source/GeometryCollectionNodes/Public/Dataflow/GeometryCollectionClusteringNodes.h
+++ b/Engine/Plugins/Experimental/GeometryCollectionPlugin/Source/GeometryCollectionNodes/Public/Dataflow/GeometryCollectionClusteringNodes.h
@@ -73,11 +73,11 @@ public:
 	float MinimumSize = 0;
 
 	/** Whether to favor clusters that have a convex shape. (Note: Does not support ByGrid clustering.)  */
-	UPROPERTY(EditAnywhere, Category = AutoCluster, meta = (EditCondition = "ClusterSizeMethod != EClusterSizeMethod::ByGrid"))
+	UPROPERTY(EditAnywhere, Category = AutoCluster, meta = (EditCondition = "ClusterSizeMethod != EClusterSizeMethodEnum::Dataflow_ClusterSizeMethod_ByGrid"))
 	bool bPreferConvexity = false;
 
 	/** If > 0, cube root of maximum concave volume to add per cluster (ignoring concavity of individual parts) */
-	UPROPERTY(EditAnywhere, Category = AutoCluster, meta = (EditCondition = "bPreferConvexity && ClusterSizeMethod != EClusterSizeMethod::ByGrid"))
+	UPROPERTY(EditAnywhere, Category = AutoCluster, meta = (EditCondition = "bPreferConvexity && ClusterSizeMethod != EClusterSizeMethodEnum::Dataflow_ClusterSizeMethod_ByGrid"))
 	float ConcavityTolerance = 0;
 
 	/** If true, bones will only be added to the same cluster if they are physically connected (either directly, or via other bones in the same cluster) */
@@ -85,7 +85,7 @@ public:
 	bool AutoCluster = true;
 
 	/** If true, make sure the site parameters are matched as close as possible ( bEnforceConnectivity can make the number of site larger than the requested input may produce without it ) */
-	UPROPERTY(EditAnywhere, Category = AutoCluster, meta = (EditCondition = "bEnforceConnectivity == true"))
+	UPROPERTY(EditAnywhere, Category = AutoCluster, meta = (EditCondition = "AutoCluster == true"))
 	bool EnforceSiteParameters = true;
 
 	/** If true, prevent the creation of clusters with only a single child. Either by merging into a neighboring cluster, or not creating the cluster. */
