Dataflow:
- Added FMeshCutterDataflowNode, FFixTinyGeoDataflowNode, FRecomputeNormalsInGeometryCollectionDataflowNode, FResampleGeometryCollectionDataflowNode, FValidateGeometryCollectionDataflowNode
- Created FractureEngineSampling module for sampling code
- Added FUniformPointSamplingDataflowNode, FNonUniformPointSamplingDataflowNode, FVertexWeightedPointSamplingDataflowNode
- Added FMakeFloatArrayDataflowNode, FRandomizeFloatArrayDataflowNode, FCollectionSelectionByAttrDataflowNode
- Bug fixes

diff --git a/Engine/Plugins/Experimental/Fracture/Source/FractureEngine/Private/FractureEngineFracturing.cpp b/Engine/Plugins/Experimental/Fracture/Source/FractureEngine/Private/FractureEngineFracturing.cpp
--- a/Engine/Plugins/Experimental/Fracture/Source/FractureEngine/Private/FractureEngineFracturing.cpp
+++ b/Engine/Plugins/Experimental/Fracture/Source/FractureEngine/Private/FractureEngineFracturing.cpp
@@ -9,7 +9,9 @@
 #include "Voronoi/Voronoi.h"
 #include "PlanarCut.h"
 #include "GeometryCollection/Facades/CollectionMeshFacade.h"
-
+#include "Algo/RemoveIf.h"
+#include "GeometryCollection/Facades/CollectionTransformSelectionFacade.h"
+#include "GeometryCollection/GeometryCollectionClusteringUtility.h"
 
 static void AddAdditionalAttributesIfRequired(FManagedArrayCollection& InOutCollection)
 {
@@ -896,5 +898,163 @@ int32 FFractureEngineFracturing::BrickCutter(FManagedArrayCollection& InOutColle
 
 	return INDEX_NONE;
 }
+	
+void FFractureEngineFracturing::GenerateMeshTransforms(TArray<FTransform>& MeshTransforms,
+	const FBox& InBoundingBox,
+	const int32 InRandomSeed,
+	const EMeshCutterCutDistribution InCutDistribution,
+	const int32 InNumberToScatter,
+	const int32 InGridX,
+	const int32 InGridY,
+	const int32 InGridZ,
+	const float InVariability,
+	const float InMinScaleFactor,
+	const float InMaxScaleFactor,
+	const bool InRandomOrientation,
+	const float InRollRange,
+	const float InPitchRange,
+	const float InYawRange)
+{
+	FRandomStream RandStream(InRandomSeed);
+
+	FBox Bounds = InBoundingBox;
+	const FVector Extent(Bounds.Max - Bounds.Min);
+
+	TArray<FVector> Positions;
+	if (InCutDistribution == EMeshCutterCutDistribution::UniformRandom)
+	{
+		Positions.Reserve(InNumberToScatter);
+		for (int32 Idx = 0; Idx < InNumberToScatter; ++Idx)
+		{
+			Positions.Emplace(Bounds.Min + FVector(RandStream.FRand(), RandStream.FRand(), RandStream.FRand()) * Extent);
+		}
+	}
+	else if (InCutDistribution == EMeshCutterCutDistribution::Grid)
+	{
+		Positions.Reserve(InGridX * InGridY * InGridZ);
+		auto ToFrac = [](int32 Val, int32 NumVals) -> FVector::FReal
+		{
+			return (FVector::FReal(Val) + FVector::FReal(.5)) / FVector::FReal(NumVals);
+		};
+		for (int32 X = 0; X < InGridX; ++X)
+		{
+			FVector::FReal XFrac = ToFrac(X, InGridX);
+			for (int32 Y = 0; Y < InGridY; ++Y)
+			{
+				FVector::FReal YFrac = ToFrac(Y, InGridY);
+				for (int32 Z = 0; Z < InGridZ; ++Z)
+				{
+					FVector::FReal ZFrac = ToFrac(Z, InGridZ);
+					Positions.Emplace(Bounds.Min + FVector(XFrac, YFrac, ZFrac) * Extent);
+				}
+			}
+		}
+
+		for (FVector& Position : Positions)
+		{
+			Position += (RandStream.VRand() * RandStream.FRand() * InVariability);
+		}
+	}
+
+	MeshTransforms.Reserve(MeshTransforms.Num() + Positions.Num());
+	for (const FVector& Position : Positions)
+	{
+		const FVector ScaleVec(RandStream.FRandRange(InMinScaleFactor, InMaxScaleFactor));
+		FRotator Orientation = FRotator::ZeroRotator;
+		if (InRandomOrientation)
+		{
+			Orientation = FRotator(
+				RandStream.FRandRange(-InPitchRange, InPitchRange),
+				RandStream.FRandRange(-InYawRange, InYawRange),
+				RandStream.FRandRange(-InRollRange, InRollRange)
+			);
+		}
+		MeshTransforms.Emplace(FTransform(Orientation, Position, ScaleVec));
+	}
+}
+
+int32 FFractureEngineFracturing::MeshCutter(TArray<FTransform>& MeshTransforms,
+	FManagedArrayCollection& InOutCollection,
+	FDataflowTransformSelection InTransformSelection,
+	const UE::Geometry::FDynamicMesh3& InDynCuttingMesh,
+	const FTransform& InTransform,
+	const int32 InRandomSeed,
+	const float InChanceToFracture,
+	const bool InSplitIslands,
+	const float InCollisionSampleSpacing)
+{
+	if (TUniquePtr<FGeometryCollection> GeomCollection = TUniquePtr<FGeometryCollection>(InOutCollection.NewCopy<FGeometryCollection>()))
+	{
+		// Note: Noise not currently supported
+		FInternalSurfaceMaterials InternalSurfaceMaterials;
+
+		int32 ResultGeometryIndex = -1;
+
+		const int32 OriginalNumTransforms = InOutCollection.NumElements(FGeometryCollection::TransformGroup);
+
+		RandomReduceSelection(InTransformSelection, InRandomSeed, InChanceToFracture);
+		TArray<int32> TransformSelectionArr = InTransformSelection.AsArray();
+
+		for (const FTransform& ScatterTransform : MeshTransforms)
+		{
+			constexpr bool bSetDefaultInternalMaterialsFromCollection = true;
+			int32 Index = CutWithMesh(InDynCuttingMesh, ScatterTransform, InternalSurfaceMaterials, *GeomCollection, TransformSelectionArr, InCollisionSampleSpacing, InTransform, bSetDefaultInternalMaterialsFromCollection, nullptr, InSplitIslands);
+
+			int32 NewLen = Algo::RemoveIf(TransformSelectionArr, [&](int32 Bone)
+				{
+					return !GeomCollection->IsVisible(Bone); // remove already-fractured pieces from the to-cut list
+				});
+			TransformSelectionArr.SetNum(NewLen);
+
+			if (ResultGeometryIndex == -1)
+			{
+				ResultGeometryIndex = Index;
+			}
+			if (Index > -1)
+			{
+				const int32 TransformIdx = GeomCollection->TransformIndex[Index];
+				// after a successful cut, also consider any new bones added by the cut
+				for (int32 NewBoneIdx = TransformIdx; NewBoneIdx < GeomCollection->NumElements(FGeometryCollection::TransformGroup); NewBoneIdx++)
+				{
+					TransformSelectionArr.Add(NewBoneIdx);
+				}
+			}
+		}
+
+		if (ResultGeometryIndex > -1)
+		{
+			TArray<int32> ToRemove;
+			for (int32 NewIdx = OriginalNumTransforms; NewIdx < GeomCollection->NumElements(FGeometryCollection::TransformGroup); ++NewIdx)
+			{
+				if (GeomCollection->IsRigid(NewIdx))
+				{
+					int32 ParentIdx = GeomCollection->Parent[NewIdx];
+					if (ParentIdx >= OriginalNumTransforms)
+					{
+						do
+						{
+							ParentIdx = GeomCollection->Parent[ParentIdx];
+						} while (GeomCollection->Parent[ParentIdx] >= OriginalNumTransforms);
+						GeomCollection->ParentTransforms(ParentIdx, NewIdx);
+					}
+				}
+				else
+				{
+					ToRemove.Add(NewIdx);
+				}
+			}
+			FManagedArrayCollection::FProcessingParameters ProcessingParams;
+			ProcessingParams.bDoValidation = false;
+			GeomCollection->RemoveElements(FGeometryCollection::TransformGroup, ToRemove, ProcessingParams);
+		}
+
+		InOutCollection = (const FManagedArrayCollection&)(*GeomCollection);
+
+		return ResultGeometryIndex;
+	}
+
+	return INDEX_NONE;
+}
+
 
 
diff --git a/Engine/Plugins/Experimental/Fracture/Source/FractureEngine/Private/FractureEngineSampling.cpp b/Engine/Plugins/Experimental/Fracture/Source/FractureEngine/Private/FractureEngineSampling.cpp
new file mode 100644
--- /dev/null
+++ b/Engine/Plugins/Experimental/Fracture/Source/FractureEngine/Private/FractureEngineSampling.cpp
@@ -0,0 +1,168 @@
+// Copyright Epic Games, Inc. All Rights Reserved.
+
+#include "FractureEngineSampling.h"
+#include "DynamicMesh/DynamicMesh3.h"
+#include "FractureEngineSelection.h"
+#include "Dataflow/DataflowSelection.h"
+#include "GeometryCollection/GeometryCollectionClusteringUtility.h"
+#include "GeometryCollection/GeometryCollectionAlgo.h"
+#include "GeometryCollection/TransformCollection.h"
+#include "Voronoi/Voronoi.h"
+#include "PlanarCut.h"
+#include "GeometryCollection/Facades/CollectionMeshFacade.h"
+#include "Algo/RemoveIf.h"
+#include "GeometryCollection/Facades/CollectionTransformSelectionFacade.h"
+#include "GeometryCollection/GeometryCollectionClusteringUtility.h"
+#include "Sampling/MeshSurfacePointSampling.h"
+
+void FFractureEngineSampling::ComputeUniformPointSampling(const UE::Geometry::FDynamicMesh3& InMesh,
+	const float InSamplingRadius,
+	const int32 InMaxNumSamples,
+	const float InSubSampleDensity,
+	const int32 InRandomSeed,
+	TArray<FTransform>& OutSamples,
+	TArray<int32>& OutTriangleIDs,
+	TArray<FVector>& OutBarycentricCoords)
+{
+	if (InMesh.VertexCount())
+	{
+		UE::Geometry::FMeshSurfacePointSampling Sampler;
+		Sampler.SampleRadius = InSamplingRadius;
+		Sampler.MaxSamples = InMaxNumSamples;
+		Sampler.RandomSeed = InRandomSeed;
+		Sampler.SubSampleDensity = InSubSampleDensity;
+		Sampler.bComputeBarycentrics = true;
+
+		Sampler.ComputePoissonSampling(InMesh);
+
+		int32 NumSamples = Sampler.Samples.Num();
+		OutSamples.SetNumUninitialized(NumSamples);
+		OutTriangleIDs.SetNumUninitialized(NumSamples);
+		OutBarycentricCoords.SetNumUninitialized(NumSamples);
+
+		for (int32 Idx = 0; Idx < NumSamples; ++Idx)
+		{
+			const UE::Geometry::FFrame3d& Frame = Sampler.Samples[Idx];
+			OutSamples[Idx] = Frame.ToFTransform();
+			OutTriangleIDs[Idx] = Sampler.TriangleIDs[Idx];
+			OutBarycentricCoords[Idx] = Sampler.BarycentricCoords[Idx];
+		}
+	}
+}
+
+void FFractureEngineSampling::ComputeNonUniformPointSampling(const UE::Geometry::FDynamicMesh3& InMesh,
+	const float InSamplingRadius,
+	const int32 InMaxNumSamples,
+	const float InSubSampleDensity,
+	const int32 InRandomSeed,
+	const float InMaxSamplingRadius,
+	const ENonUniformSamplingDistributionMode InSizeDistribution,
+	const float InSizeDistributionPower,
+	TArray<FTransform>& OutSamples,
+	TArray<float>& OutSampleRadii,
+	TArray<int32>& OutTriangleIDs,
+	TArray<FVector>& OutBarycentricCoords)
+{
+	if (InMesh.VertexCount())
+	{
+		UE::Geometry::FMeshSurfacePointSampling Sampler;
+		Sampler.SampleRadius = InSamplingRadius;
+		Sampler.MaxSamples = InMaxNumSamples;
+		Sampler.RandomSeed = InRandomSeed;
+		Sampler.SubSampleDensity = InSubSampleDensity;
+		Sampler.bComputeBarycentrics = true;
+
+		if (InMaxSamplingRadius > InSamplingRadius)
+		{
+			Sampler.MaxSampleRadius = InMaxSamplingRadius;
+			Sampler.SizeDistribution = (UE::Geometry::FMeshSurfacePointSampling::ESizeDistribution)(int)(InSizeDistribution);
+			Sampler.SizeDistributionPower = FMath::Clamp(InSizeDistributionPower, 1.f, 10.f);
+		}
+
+		Sampler.ComputePoissonSampling(InMesh);
+
+		int32 NumSamples = Sampler.Samples.Num();
+		OutSamples.SetNumUninitialized(NumSamples);
+		OutSampleRadii.SetNumUninitialized(NumSamples);
+		OutTriangleIDs.SetNumUninitialized(NumSamples);
+		OutBarycentricCoords.SetNumUninitialized(NumSamples);
+
+		for (int32 Idx = 0; Idx < NumSamples; ++Idx)
+		{
+			const UE::Geometry::FFrame3d& Frame = Sampler.Samples[Idx];
+			OutSamples[Idx] = Frame.ToFTransform();
+			OutSampleRadii[Idx] = Sampler.Radii[Idx];
+			OutTriangleIDs[Idx] = Sampler.TriangleIDs[Idx];
+			OutBarycentricCoords[Idx] = Sampler.BarycentricCoords[Idx];
+		}
+	}
+}
+
+void FFractureEngineSampling::ComputeVertexWeightedPointSampling(const UE::Geometry::FDynamicMesh3& InMesh,
+	const TArray<float>& InVertexWeights,
+	const float InSamplingRadius,
+	const int32 InMaxNumSamples,
+	const float InSubSampleDensity,
+	const int32 InRandomSeed,
+	const float InMaxSamplingRadius,
+	const ENonUniformSamplingDistributionMode InSizeDistribution,
+	const float InSizeDistributionPower,
+	const ENonUniformSamplingWeightMode InWeightMode,
+	const bool InInvertWeights,
+	TArray<FTransform>& OutSamples,
+	TArray<float>& OutSampleRadii,
+	TArray<int32>& OutTriangleIDs,
+	TArray<FVector>& OutBarycentricCoords)
+{
+	if (InMesh.VertexCount())
+	{
+		UE::Geometry::FMeshSurfacePointSampling Sampler;
+		Sampler.SampleRadius = InSamplingRadius;
+		Sampler.MaxSamples = InMaxNumSamples;
+		Sampler.RandomSeed = InRandomSeed;
+		Sampler.SubSampleDensity = InSubSampleDensity;
+		Sampler.bComputeBarycentrics = true;
+
+		bool bIsNonUniform = false;
+
+		if (InMaxSamplingRadius > InSamplingRadius)
+		{
+			bIsNonUniform = true;
+
+			Sampler.MaxSampleRadius = InMaxSamplingRadius;
+			Sampler.SizeDistribution = (UE::Geometry::FMeshSurfacePointSampling::ESizeDistribution)(int)(InSizeDistribution);
+			Sampler.SizeDistributionPower = FMath::Clamp(InSizeDistributionPower, 1.f, 10.f);
+		}
+
+		const int32 NumVertices = InMesh.VertexCount();
+
+		if (bIsNonUniform && InVertexWeights.Num() == NumVertices)
+		{
+			TArray<double> WeightArr; WeightArr.SetNumUninitialized(NumVertices);
+			for (int32 Idx = 0; Idx < NumVertices; ++Idx) {	WeightArr[Idx] = (double)InVertexWeights[Idx]; }
+
+			Sampler.VertexWeights = WeightArr;
+			Sampler.bUseVertexWeights = true;
+			Sampler.InterpretWeightMode = (UE::Geometry::FMeshSurfacePointSampling::EInterpretWeightMode)(int)(InWeightMode);
+			Sampler.bInvertWeights = InInvertWeights;
+		}
+
+		Sampler.ComputePoissonSampling(InMesh);
+
+		int32 NumSamples = Sampler.Samples.Num();
+		OutSamples.SetNumUninitialized(NumSamples);
+		OutSampleRadii.SetNumUninitialized(NumSamples);
+		OutTriangleIDs.SetNumUninitialized(NumSamples);
+		OutBarycentricCoords.SetNumUninitialized(NumSamples);
+
+		for (int32 Idx = 0; Idx < NumSamples; ++Idx)
+		{
+			const UE::Geometry::FFrame3d& Frame = Sampler.Samples[Idx];
+			OutSamples[Idx] = Frame.ToFTransform();
+			OutSampleRadii[Idx] = Sampler.Radii[Idx];
+			OutTriangleIDs[Idx] = Sampler.TriangleIDs[Idx];
+			OutBarycentricCoords[Idx] = Sampler.BarycentricCoords[Idx];
+		}
+	}
+}
+
diff --git a/Engine/Plugins/Experimental/Fracture/Source/FractureEngine/Private/FractureEngineUtility.cpp b/Engine/Plugins/Experimental/Fracture/Source/FractureEngine/Private/FractureEngineUtility.cpp
--- a/Engine/Plugins/Experimental/Fracture/Source/FractureEngine/Private/FractureEngineUtility.cpp
+++ b/Engine/Plugins/Experimental/Fracture/Source/FractureEngine/Private/FractureEngineUtility.cpp
@@ -2,7 +2,17 @@
 
 #include "FractureEngineUtility.h"
 #include "DynamicMesh/DynamicMesh3.h"
-
+#include "FractureEngineSelection.h"
+#include "Dataflow/DataflowSelection.h"
+#include "GeometryCollection/GeometryCollectionClusteringUtility.h"
+#include "GeometryCollection/GeometryCollectionAlgo.h"
+#include "GeometryCollection/TransformCollection.h"
+#include "Voronoi/Voronoi.h"
+#include "PlanarCut.h"
+#include "GeometryCollection/Facades/CollectionMeshFacade.h"
+#include "Algo/RemoveIf.h"
+#include "GeometryCollection/Facades/CollectionTransformSelectionFacade.h"
+#include "GeometryCollection/GeometryCollectionClusteringUtility.h"
 
 void FFractureEngineUtility::ConvertBoxToVertexAndTriangleData(const FBox& InBox, TArray<FVector3f>& OutVertices, TArray<FIntVector>& OutTriangles)
 {
@@ -77,3 +87,334 @@ void FFractureEngineUtility::DeconstructMesh(const UE::Geometry::FDynamicMesh3&
 		}
 	}
 }
+
+namespace Dataflow::Private
+{
+	static constexpr double VolDimScale = .01; // compute volumes in meters instead of cm, for saner units at typical scales
+}
+
+static double GetTotalVolume(const FGeometryCollection& Collection, const TArray<double>& Volumes)
+{
+	double Sum = 0.0;
+	if (Volumes.Num() != Collection.Transform.Num())
+	{
+		return 0.0f;
+	}
+	for (int32 BoneIdx = 0, NumBones = Collection.Transform.Num(); BoneIdx < NumBones; ++BoneIdx)
+	{
+		if (Collection.SimulationType[BoneIdx] == FGeometryCollection::ESimulationTypes::FST_Rigid)
+		{
+			Sum += Volumes[BoneIdx];
+		}
+	}
+
+	return Sum;
+}
+
+static double GetMinVolume(double TotalVolume,
+	const EFixTinyGeoGeometrySelectionMethod InSelectionMethod,
+	const float InMinVolumeCubeRoot,
+	const float InRelativeVolume)
+{
+	double MinVolume = 0.0;
+	if (InSelectionMethod == EFixTinyGeoGeometrySelectionMethod::VolumeCubeRoot)
+	{
+		MinVolume = InMinVolumeCubeRoot * Dataflow::Private::VolDimScale;
+		MinVolume = MinVolume * MinVolume * MinVolume;
+	}
+	else // EGeometrySelectionMethod::RelativeVolume
+	{
+		MinVolume = FMath::Pow(TotalVolume, 1.0 / 3.0) * InRelativeVolume;
+		MinVolume = MinVolume * MinVolume * MinVolume;
+	}
+
+	return MinVolume;
+}
+
+static bool CollectTargetBones(FGeometryCollection& Collection,
+	const TArray<int32>& Selection,
+	TArray<int32>& OutSmallIndices,
+	TArray<double>& OutVolumes,
+	double& OutMinVolume,
+	const EFixTinyGeoMergeType InMergeType,
+	const bool InOnFractureLevel,
+	const EFixTinyGeoGeometrySelectionMethod InSelectionMethod,
+	const float InMinVolumeCubeRoot,
+	const float InRelativeVolume,
+	const EFixTinyGeoUseBoneSelection InUseBoneSelection,
+	const bool InOnlyClusters)
+{
+	const bool bClusterMode = InMergeType == EFixTinyGeoMergeType::MergeClusters;
+	const bool bRestrictToLevel = bClusterMode && InOnFractureLevel;
+	int32 InFractureLevel = -1;
+
+	FindBoneVolumes(
+		Collection,
+		TArrayView<int32>(), /*Empty array => use all transforms*/
+		OutVolumes,
+		Dataflow::Private::VolDimScale,
+		bClusterMode
+	);
+
+	double TotalVolume = GetTotalVolume(Collection, OutVolumes);
+	OutMinVolume = GetMinVolume(TotalVolume, InSelectionMethod, InMinVolumeCubeRoot, InRelativeVolume);
+
+	TArray<int32> FilteredSelection;
+
+	GeometryCollection::Facades::FCollectionTransformSelectionFacade SelectionFacade(Collection);
+
+	if (InUseBoneSelection == EFixTinyGeoUseBoneSelection::OnlyMergeSelected)
+	{
+		if (!bClusterMode)
+		{
+			OutSmallIndices = Selection;
+			SelectionFacade.ConvertSelectionToRigidNodes(OutSmallIndices);
+		}
+		else
+		{
+			OutSmallIndices = Selection;
+		}
+		return !OutSmallIndices.IsEmpty();
+	}
+
+	FindSmallBones(
+		Collection,
+		TArrayView<int32>(),
+		OutVolumes,
+		OutMinVolume,
+		OutSmallIndices,
+		bClusterMode
+	);
+
+	// Filter bones that aren't at the target level
+	const int32 TargetLevel = InFractureLevel;
+	const TManagedArray<int32>* LevelAttrib = Collection.FindAttribute<int32>("Level", FGeometryCollection::TransformGroup);
+	bool bFilterByLevel = bRestrictToLevel && TargetLevel > -1 && LevelAttrib;
+	if (bFilterByLevel)
+	{
+		OutSmallIndices.SetNum(Algo::RemoveIf(OutSmallIndices, [&](int32 BoneIdx)
+			{
+				return (*LevelAttrib)[BoneIdx] != TargetLevel;
+			}));
+	}
+	// Filter bones that aren't clusters
+	if (bClusterMode && (!bFilterByLevel || InOnlyClusters))
+	{
+		OutSmallIndices.SetNum(Algo::RemoveIf(OutSmallIndices, [&](int32 BoneIdx)
+			{
+				return Collection.SimulationType[BoneIdx] != FGeometryCollection::ESimulationTypes::FST_Clustered;
+			}));
+	}
+
+	if (InUseBoneSelection == EFixTinyGeoUseBoneSelection::AlsoMergeSelected)
+	{
+		TArray<int32> ProcessedSelection = Selection;
+		if (!bClusterMode)
+		{
+			SelectionFacade.ConvertSelectionToRigidNodes(ProcessedSelection);
+		}
+		for (int32 Bone : ProcessedSelection)
+		{
+			OutSmallIndices.AddUnique(Bone);
+		}
+	}
+
+	return !OutSmallIndices.IsEmpty();
+}
+
+void FFractureEngineUtility::FixTinyGeo(FManagedArrayCollection& InOutCollection,
+	FDataflowTransformSelection InTransformSelection,
+	const EFixTinyGeoMergeType InMergeType,
+	const bool InOnFractureLevel,
+	const EFixTinyGeoGeometrySelectionMethod InSelectionMethod,
+	const float InMinVolumeCubeRoot,
+	const float InRelativeVolume,
+	const EFixTinyGeoUseBoneSelection InUseBoneSelection,
+	const bool InOnlyClusters,
+	const EFixTinyGeoNeighborSelectionMethod InNeighborSelection,
+	const bool InOnlyToConnected,
+	const bool InOnlySameParent)
+{
+	if (TUniquePtr<FGeometryCollection> GeomCollection = TUniquePtr<FGeometryCollection>(InOutCollection.NewCopy<FGeometryCollection>()))
+	{
+		const TArray<int32> TransformSelectionArr = InTransformSelection.AsArray();
+
+		TArray<int32> SmallIndices;
+		TArray<double> Volumes;
+		double MinVolume;
+		if (CollectTargetBones(*GeomCollection,
+			TransformSelectionArr,
+			SmallIndices,
+			Volumes,
+			MinVolume,
+			InMergeType,
+			InOnFractureLevel,
+			InSelectionMethod,
+			InMinVolumeCubeRoot,
+			InRelativeVolume,
+			InUseBoneSelection,
+			InOnlyClusters))
+		{
+			// convert ENeighborSelectionMethod to the PlanarCut module version (the only difference is it's not a UENUM)
+			UE::PlanarCut::ENeighborSelectionMethod SelectionMethod = UE::PlanarCut::ENeighborSelectionMethod::LargestNeighbor;
+			if (InNeighborSelection == EFixTinyGeoNeighborSelectionMethod::NearestCenter)
+			{
+				SelectionMethod = UE::PlanarCut::ENeighborSelectionMethod::NearestCenter;
+			}
+
+			if (InMergeType == EFixTinyGeoMergeType::MergeGeometry)
+			{
+				MergeBones(*GeomCollection,
+					TArrayView<const int32>(), // empty view == consider all bones
+					Volumes,
+					MinVolume,
+					SmallIndices,
+					false, /*bUnionJoinedPieces*/ // Note: Union-ing the pieces is nicer in theory, but can leave cracks and non-manifold garbage
+					SelectionMethod);
+			}
+			else
+			{
+				MergeClusters(*GeomCollection,
+					Volumes,
+					MinVolume,
+					SmallIndices,
+					SelectionMethod,
+					InOnlyToConnected,
+					InOnlySameParent);
+			}
+
+			InOutCollection = (const FManagedArrayCollection&)(*GeomCollection);
+		}
+	}
+}
+
+void FFractureEngineUtility::RecomputeNormalsInGeometryCollection(FManagedArrayCollection& InOutCollection,
+	FDataflowTransformSelection InTransformSelection,
+	const bool InOnlyTangents,
+	const bool InRecomputeSharpEdges,
+	const float InSharpEdgeAngleThreshold,
+	const bool InOnlyInternalSurfaces)
+{
+	if (TUniquePtr<FGeometryCollection> GeomCollection = TUniquePtr<FGeometryCollection>(InOutCollection.NewCopy<FGeometryCollection>()))
+	{
+		TArray<int32> TransformSelectionArr = InTransformSelection.AsArray();
+
+		RecomputeNormalsAndTangents(InOnlyTangents, InRecomputeSharpEdges, InSharpEdgeAngleThreshold, *GeomCollection, TransformSelectionArr, InOnlyInternalSurfaces);
+
+		InOutCollection = (const FManagedArrayCollection&)(*GeomCollection);
+	}
+}
+
+int32 FFractureEngineUtility::ResampleGeometryCollection(FManagedArrayCollection& InOutCollection,
+	FDataflowTransformSelection InTransformSelection,
+	const float InCollisionSampleSpacing)
+{
+	if (TUniquePtr<FGeometryCollection> GeomCollection = TUniquePtr<FGeometryCollection>(InOutCollection.NewCopy<FGeometryCollection>()))
+	{
+		TArray<int32> TransformSelectionArr = InTransformSelection.AsArray();
+
+		AddCollisionSampleVertices(InCollisionSampleSpacing, *GeomCollection, TransformSelectionArr);
+
+		InOutCollection = (const FManagedArrayCollection&)(*GeomCollection);
+	}
+
+	return INDEX_NONE;
+}
+
+static void AddSingleRootNodeIfRequired(TUniquePtr<FGeometryCollection>& InOutCollection)
+{
+	if (InOutCollection)
+	{
+		if (FGeometryCollectionClusteringUtility::ContainsMultipleRootBones(InOutCollection.Get()))
+		{
+			FGeometryCollectionClusteringUtility::ClusterAllBonesUnderNewRoot(InOutCollection.Get());
+		}
+	}
+}
+
+DECLARE_LOG_CATEGORY_EXTERN(LogFractureTool, Log, All);
+DEFINE_LOG_CATEGORY(LogFractureTool);
+
+void FFractureEngineUtility::ValidateGeometryCollection(FManagedArrayCollection& InOutCollection,
+	const bool InRemoveUnreferencedGeometry,
+	const bool InRemoveClustersOfOne,
+	const bool InRemoveDanglingClusters)
+{
+	if (TUniquePtr<FGeometryCollection> GeometryCollection = TUniquePtr<FGeometryCollection>(InOutCollection.NewCopy<FGeometryCollection>()))
+	{
+		bool bDirty = false;
+
+		TManagedArray<int32>& TransformToGeometry = GeometryCollection->TransformToGeometryIndex;
+		constexpr bool bClustersCanHaveGeometry = true;
+		if (!bClustersCanHaveGeometry)
+		{
+			// Optionally ensure that clusters do not point to geometry (currently disabled; keeping the geometry can be useful)
+			const int32 ElementCount = TransformToGeometry.Num();
+			for (int32 Idx = 0; Idx < ElementCount; ++Idx)
+			{
+				if (GeometryCollection->IsClustered(Idx) && TransformToGeometry[Idx] != INDEX_NONE)
+				{
+					TransformToGeometry[Idx] = INDEX_NONE;
+					UE_LOG(LogFractureTool, Verbose, TEXT("Removed geometry index from cluster %d."), Idx);
+					bDirty = true;
+				}
+			}
+		}
+
+		// Remove any unreferenced geometry
+		if (InRemoveUnreferencedGeometry)
+		{
+			TManagedArray<int32>& TransformIndex = GeometryCollection->TransformIndex;
+			const int32 GeometryCount = TransformIndex.Num();
+
+			TArray<int32> RemoveGeometry;
+			RemoveGeometry.Reserve(GeometryCount);
+
+			for (int32 Idx = 0; Idx < GeometryCount; ++Idx)
+			{
+				if ((TransformIndex[Idx] == INDEX_NONE) || (TransformToGeometry[TransformIndex[Idx]] != Idx))
+				{
+					RemoveGeometry.Add(Idx);
+					UE_LOG(LogFractureTool, Verbose, TEXT("Removed dangling geometry at index %d."), Idx);
+					bDirty = true;
+				}
+			}
+
+			if (RemoveGeometry.Num() > 0)
+			{
+				FManagedArrayCollection::FProcessingParameters Params;
+				Params.bDoValidation = false; // for perf reasons
+				GeometryCollection->RemoveElements(FGeometryCollection::GeometryGroup, RemoveGeometry);
+			}
+		}
+
+		if (InRemoveClustersOfOne)
+		{
+			if (FGeometryCollectionClusteringUtility::RemoveClustersOfOnlyOneChild(GeometryCollection.Get()))
+			{
+				UE_LOG(LogFractureTool, Verbose, TEXT("Removed one or more clusters of only one child."));
+				bDirty = true;
+			}
+		}
+
+		if (InRemoveDanglingClusters)
+		{
+			if (FGeometryCollectionClusteringUtility::RemoveDanglingClusters(GeometryCollection.Get()))
+			{
+				UE_LOG(LogFractureTool, Verbose, TEXT("Removed one or more dangling clusters."));
+				bDirty = true;
+			}
+		}
+
+		if (bDirty)
+		{
+			FGeometryCollectionClusteringUtility::UpdateHierarchyLevelOfChildren(*GeometryCollection, -1);
+			AddSingleRootNodeIfRequired(GeometryCollection);
+		}
+
+		InOutCollection = (const FManagedArrayCollection&)(*GeometryCollection);
+	}
+}
+
+
+
diff --git a/Engine/Plugins/Experimental/Fracture/Source/FractureEngine/Public/FractureEngineFracturing.h b/Engine/Plugins/Experimental/Fracture/Source/FractureEngine/Public/FractureEngineFracturing.h
--- a/Engine/Plugins/Experimental/Fracture/Source/FractureEngine/Public/FractureEngineFracturing.h
+++ b/Engine/Plugins/Experimental/Fracture/Source/FractureEngine/Public/FractureEngineFracturing.h
@@ -6,6 +6,8 @@
 
 #include "FractureEngineFracturing.generated.h"
 
+namespace UE { namespace Geometry { class FDynamicMesh3; } }
+
 struct FDataflowTransformSelection;
 struct FManagedArrayCollection;
 
@@ -17,10 +19,17 @@ enum class EFractureBrickBondEnum : uint8
 	Dataflow_FractureBrickBond_English UMETA(DisplayName = "English"),
 	Dataflow_FractureBrickBond_Header UMETA(DisplayName = "Header"),
 	Dataflow_FractureBrickBond_Flemish UMETA(DisplayName = "Flemish"),
+};
 
-	//~~~
-	//256th entry
-	Dataflow_Max                UMETA(Hidden)
+UENUM(BlueprintType)
+enum class EMeshCutterCutDistribution : uint8
+{
+	// Cut only once, at the cutting mesh's current location in the level
+	SingleCut UMETA(DisplayName = "Single Cut"),
+	// Scatter the cutting mesh in a uniform random distribution around the geometry bounding box
+	UniformRandom UMETA(DisplayName = "Uniform Random"),
+	// Arrange the cutting mesh in a regular grid pattern
+	Grid UMETA(DisplayName = "Grid"),
 };
 
 class FRACTUREENGINE_API FFractureEngineFracturing
@@ -127,6 +136,31 @@ public:
 		bool InAddSamplesForCollision,
 		float InCollisionSampleSpacing);
 
+	static void GenerateMeshTransforms(TArray<FTransform>& MeshTransforms,
+		const FBox& InBoundingBox,
+		const int32 InRandomSeed,
+		const EMeshCutterCutDistribution InCutDistribution,
+		const int32 InNumberToScatter,
+		const int32 InGridX,
+		const int32 InGridY,
+		const int32 InGridZ,
+		const float InVariability,
+		const float InMinScaleFactor,
+		const float InMaxScaleFactor,
+		const bool InRandomOrientation,
+		const float InRollRange,
+		const float InPitchRange,
+		const float InYawRange);
+
+	static int32 MeshCutter(TArray<FTransform>& MeshTransforms,
+		FManagedArrayCollection& InOutCollection,
+		FDataflowTransformSelection InTransformSelection,
+		const UE::Geometry::FDynamicMesh3& InDynCuttingMesh,
+		const FTransform& InTransform,
+		const int32 InRandomSeed,
+		const float InChanceToFracture,
+		const bool InSplitIslands,
+		const float InCollisionSampleSpacing);
 };
 
 #if UE_ENABLE_INCLUDE_ORDER_DEPRECATED_IN_5_2
diff --git a/Engine/Plugins/Experimental/Fracture/Source/FractureEngine/Public/FractureEngineSampling.h b/Engine/Plugins/Experimental/Fracture/Source/FractureEngine/Public/FractureEngineSampling.h
new file mode 100644
--- /dev/null
+++ b/Engine/Plugins/Experimental/Fracture/Source/FractureEngine/Public/FractureEngineSampling.h
@@ -0,0 +1,97 @@
+// Copyright Epic Games, Inc. All Rights Reserved.
+
+#pragma once
+
+#include "Math/Box.h"
+#include "Sampling/MeshSurfacePointSampling.h"
+
+#include "FractureEngineSampling.generated.h"
+
+namespace UE::Geometry { class FDynamicMesh3; }
+
+struct FDataflowTransformSelection;
+struct FManagedArrayCollection;
+
+UENUM(BlueprintType)
+enum class ENonUniformSamplingDistributionMode : uint8
+{
+	/** Uniform distribution of sizes, ie all equally likely */
+	ENonUniformSamplingDistributionMode_Uniform UMETA(DisplayName = "Uniform"),
+	/** Distribution is weighted towards smaller points */
+	ENonUniformSamplingDistributionMode_Smaller UMETA(DisplayName = "Smaller"),
+	/** Distribution is weighted towards larger points */
+	ENonUniformSamplingDistributionMode_Larger UMETA(DisplayName = "Larger"),
+};
+
+UENUM(BlueprintType)
+enum class ENonUniformSamplingWeightMode : uint8
+{
+	/**
+	 * Weights are clamped to [0,1] and used to interpolate Min/Max Radius. This is a "hard constraint", ie if the weight
+	 * at a point is 1, only a "max radius" sample may be placed there, otherwise no samples at all (so no "filling in" smaller samples between large ones)
+	 */
+	ENonUniformSamplingWeightMode_WeightToRadius UMETA(DisplayName = "Weight To Radius"),
+	/**
+	 * Weights are clamped to [0,1] and used to interpolate Min/Max Radius, with decay, so that smaller-radius samples will infill between large ones.
+	 * So areas with large weight may still end up with some variable-radius samples, but areas with 0 weight will only ever have min-radius samples.
+	 */
+	ENonUniformSamplingWeightMode_FilledWeightToRadius UMETA(DisplayName = "Filled Weight To Radius"),
+	/**
+	 * Weight is used to create nonuniform random sampling, ie it nudges the random sample-radius distribution but does not directly control it.
+	 * So samples with any radius can still appear at any location, but if weight=1 then max-radius samples are more likely, etc.
+	 */
+	ENonUniformSamplingWeightMode_WeightedRandom UMETA(DisplayName = "Weighted Random"),
+};
+
+class FFractureEngineSampling
+{
+public:
+	/**
+	* Computes a Poisson sampling on the InMesh
+	*/
+	FRACTUREENGINE_API static void ComputeUniformPointSampling(const UE::Geometry::FDynamicMesh3& InMesh,
+		const float InSamplingRadius,
+		const int32 InMaxNumSamples,
+		const float InSubSampleDensity,
+		const int32 InRandomSeed,
+		TArray<FTransform>& OutSamples,
+		TArray<int32>& OutTriangleIDs,
+		TArray<FVector>& OutBarycentricCoords);
+
+	/**
+	* Computes a NonUniform sampling on the InMesh
+	*/
+	FRACTUREENGINE_API static void ComputeNonUniformPointSampling(const UE::Geometry::FDynamicMesh3& InMesh,
+		const float InSamplingRadius,
+		const int32 InMaxNumSamples,
+		const float InSubSampleDensity,
+		const int32 InRandomSeed,
+		const float InMaxSamplingRadius,
+		const ENonUniformSamplingDistributionMode InSizeDistribution,
+		const float InSizeDistributionPower,
+		TArray<FTransform>& OutSamples,
+		TArray<float>& OutSampleRadii,
+		TArray<int32>& OutTriangleIDs,
+		TArray<FVector>& OutBarycentricCoords);
+
+	/**
+	* Computes a vertex weighted sampling on the InMesh
+	*/
+	FRACTUREENGINE_API static void ComputeVertexWeightedPointSampling(const UE::Geometry::FDynamicMesh3& InMesh,
+		const TArray<float>& InVertexWeights,
+		const float InSamplingRadius,
+		const int32 InMaxNumSamples,
+		const float InSubSampleDensity,
+		const int32 InRandomSeed,
+		const float InMaxSamplingRadius,
+		const ENonUniformSamplingDistributionMode InSizeDistribution,
+		const float InSizeDistributionPower,
+		const ENonUniformSamplingWeightMode InWeightMode,
+		const bool InInvertWeights,
+		TArray<FTransform>& OutSamples,
+		TArray<float>& OutSampleRadii,
+		TArray<int32>& OutTriangleIDs,
+		TArray<FVector>& OutBarycentricCoords);
+};
+
+
diff --git a/Engine/Plugins/Experimental/Fracture/Source/FractureEngine/Public/FractureEngineUtility.h b/Engine/Plugins/Experimental/Fracture/Source/FractureEngine/Public/FractureEngineUtility.h
--- a/Engine/Plugins/Experimental/Fracture/Source/FractureEngine/Public/FractureEngineUtility.h
+++ b/Engine/Plugins/Experimental/Fracture/Source/FractureEngine/Public/FractureEngineUtility.h
@@ -4,9 +4,45 @@
 
 #include "Math/Box.h"
 
+#include "FractureEngineUtility.generated.h"
+
 namespace UE::Geometry { class FDynamicMesh3; }
 
-//class FDynamicMesh3;
+struct FDataflowTransformSelection;
+struct FManagedArrayCollection;
+
+UENUM(BlueprintType)
+enum class EFixTinyGeoMergeType : uint8
+{
+	MergeGeometry UMETA(DisplayName = "Merge Geometry"),
+	MergeClusters UMETA(DisplayName = "Merge Clusters"),
+};
+
+UENUM(BlueprintType)
+enum class EFixTinyGeoNeighborSelectionMethod : uint8
+{
+	// Merge to the neighbor with the largest volume
+	LargestNeighbor UMETA(DisplayName = "Largest Neighbor"),
+	// Merge to the neighbor with the closest center
+	NearestCenter UMETA(DisplayName = "Nearest Center"),
+};
+
+UENUM(BlueprintType)
+enum class EFixTinyGeoUseBoneSelection : uint8
+{
+	NoEffect UMETA(DisplayName = "No Effect"),
+	AlsoMergeSelected UMETA(DisplayName = "Also Merge Selected"),
+	OnlyMergeSelected UMETA(DisplayName = "Only Merge Selected"),
+};
+
+UENUM(BlueprintType)
+enum class EFixTinyGeoGeometrySelectionMethod : uint8
+{
+	// Select by cube root of volume
+	VolumeCubeRoot UMETA(DisplayName = "Size"),
+	// Select by cube root of volume relative to the overall shape's cube root of volume
+	RelativeVolume UMETA(DisplayName = "Relative Volume"),
+};
 
 class FRACTUREENGINE_API FFractureEngineUtility
 {
@@ -26,6 +62,35 @@ public:
 	*/
 	static void DeconstructMesh(const UE::Geometry::FDynamicMesh3& InMesh, TArray<FVector3f>& OutVertices, TArray<FIntVector>& OutTriangles);
 
+	static void FixTinyGeo(FManagedArrayCollection& InOutCollection,
+		FDataflowTransformSelection InTransformSelection,
+		const EFixTinyGeoMergeType InMergeType,
+		const bool InOnFractureLevel,
+		const EFixTinyGeoGeometrySelectionMethod InSelectionMethod,
+		const float InMinVolumeCubeRoot,
+		const float InRelativeVolume,
+		const EFixTinyGeoUseBoneSelection InUseBoneSelection,
+		const bool InOnlyClusters,
+		const EFixTinyGeoNeighborSelectionMethod InNeighborSelection,
+		const bool InOnlyToConnected,
+		const bool InOnlySameParent);
+
+	static void RecomputeNormalsInGeometryCollection(FManagedArrayCollection& InOutCollection,
+		FDataflowTransformSelection InTransformSelection,
+		const bool InOnlyTangents,
+		const bool InRecomputeSharpEdges,
+		const float InSharpEdgeAngleThreshold,
+		const bool InOnlyInternalSurfaces);
+
+	static int32 ResampleGeometryCollection(FManagedArrayCollection& InOutCollection,
+		FDataflowTransformSelection InTransformSelection,
+		const float InCollisionSampleSpacing);
+
+	static void ValidateGeometryCollection(FManagedArrayCollection& InOutCollection,
+		const bool InRemoveUnreferencedGeometry,
+		const bool InRemoveClustersOfOne,
+		const bool InRemoveDanglingClusters);
+
 };
 
 #if UE_ENABLE_INCLUDE_ORDER_DEPRECATED_IN_5_2
diff --git a/Engine/Plugins/Experimental/GeometryCollectionPlugin/Source/GeometryCollectionNodes/Private/Dataflow/GeometryCollectionArrayNodes.cpp b/Engine/Plugins/Experimental/GeometryCollectionPlugin/Source/GeometryCollectionNodes/Private/Dataflow/GeometryCollectionArrayNodes.cpp
--- a/Engine/Plugins/Experimental/GeometryCollectionPlugin/Source/GeometryCollectionNodes/Private/Dataflow/GeometryCollectionArrayNodes.cpp
+++ b/Engine/Plugins/Experimental/GeometryCollectionPlugin/Source/GeometryCollectionNodes/Private/Dataflow/GeometryCollectionArrayNodes.cpp
@@ -24,7 +24,7 @@ namespace Dataflow
 		DATAFLOW_NODE_REGISTER_CREATION_FACTORY(FUnionIntArraysDataflowNode);
 		DATAFLOW_NODE_REGISTER_CREATION_FACTORY(FRemoveFloatArrayElementDataflowNode);
 		DATAFLOW_NODE_REGISTER_CREATION_FACTORY(FFloatArrayComputeStatisticsDataflowNode);
-
+		DATAFLOW_NODE_REGISTER_CREATION_FACTORY(FRandomizeFloatArrayDataflowNode);
 	}
 }
 
@@ -586,3 +586,29 @@ void FFloatArrayComputeStatisticsDataflowNode::Evaluate(Dataflow::FContext& Cont
 		SetValue(Context, OutIndices, &Indices);
 	}
 }
+
+void FRandomizeFloatArrayDataflowNode::Evaluate(Dataflow::FContext& Context, const FDataflowOutput* Out) const
+{
+	if (Out->IsA(&FloatArray))
+	{
+		const TArray<float> InFloatArray = GetValue(Context, &FloatArray);
+		const float InRandomRangeMin = GetValue(Context, &RandomRangeMin);
+		const float InRandomRangeMax = GetValue(Context, &RandomRangeMax);
+		const int32 InRandomSeed = GetValue(Context, &RandomSeed);
+		TArray<float> OutFloatArray;
+
+		FRandomStream RandStream(InRandomSeed);
+
+		const int32 NumPoints = InFloatArray.Num();
+		if (NumPoints > 0)
+		{
+			OutFloatArray.Reserve(NumPoints);
+			for (int32 Idx = 0; Idx < NumPoints; ++Idx)
+			{
+				OutFloatArray.Emplace(RandStream.FRandRange(InRandomRangeMin, InRandomRangeMax));
+			}
+		}
+
+		SetValue(Context, MoveTemp(OutFloatArray), &FloatArray);
+	}
+}
diff --git a/Engine/Plugins/Experimental/GeometryCollectionPlugin/Source/GeometryCollectionNodes/Private/Dataflow/GeometryCollectionFracturingNodes.cpp b/Engine/Plugins/Experimental/GeometryCollectionPlugin/Source/GeometryCollectionNodes/Private/Dataflow/GeometryCollectionFracturingNodes.cpp
--- a/Engine/Plugins/Experimental/GeometryCollectionPlugin/Source/GeometryCollectionNodes/Private/Dataflow/GeometryCollectionFracturingNodes.cpp
+++ b/Engine/Plugins/Experimental/GeometryCollectionPlugin/Source/GeometryCollectionNodes/Private/Dataflow/GeometryCollectionFracturingNodes.cpp
@@ -30,6 +30,8 @@
 #include "GeometryCollection/Facades/CollectionTransformSelectionFacade.h"
 #include "GeometryCollection/Facades/CollectionBoundsFacade.h"
 #include "Dataflow/DataflowSelection.h"
+#include "DynamicMesh/DynamicMesh3.h"
+#include "UDynamicMesh.h"
 
 #include UE_INLINE_GENERATED_CPP_BY_NAME(GeometryCollectionFracturingNodes)
 
@@ -52,6 +54,7 @@ namespace Dataflow
 		DATAFLOW_NODE_REGISTER_CREATION_FACTORY(FExplodedViewDataflowNode);
 		DATAFLOW_NODE_REGISTER_CREATION_FACTORY(FSliceCutterDataflowNode);
 		DATAFLOW_NODE_REGISTER_CREATION_FACTORY(FBrickCutterDataflowNode);
+		DATAFLOW_NODE_REGISTER_CREATION_FACTORY(FMeshCutterDataflowNode);
 
 		// GeometryCollection|Fracture
 		DATAFLOW_NODE_REGISTER_CREATION_FACTORY_NODE_COLORS_BY_CATEGORY("GeometryCollection|Fracture", FLinearColor(1.f, 1.f, 0.8f), CDefaultNodeBodyTintColor);
@@ -389,7 +392,7 @@ void FVoronoiFractureDataflowNode_v2::Evaluate(Dataflow::FContext& Context, cons
 				InTransformSelection,
 				GetValue<TArray<FVector>>(Context, &Points),
 				GetValue<FTransform>(Context, &Transform),
-				GetValue<int32>(Context, &RandomSeed),
+				0, // RandomSeed is not used in Voronoi fracture, it is used in the source point generation
 				GetValue<float>(Context, &ChanceToFracture),
 				SplitIslands,
 				GetValue<float>(Context, &Grout),
@@ -428,7 +431,6 @@ void FVoronoiFractureDataflowNode_v2::Evaluate(Dataflow::FContext& Context, cons
 							OriginalSelection.SetSelected(Idx);
 						}
 					}
-
 				}
 			}
 
@@ -559,7 +561,6 @@ void FPlaneCutterDataflowNode_v2::Evaluate(Dataflow::FContext& Context, const FD
 							OriginalSelection.SetSelected(Idx);
 						}
 					}
-
 				}
 			}
 
@@ -664,7 +665,6 @@ void FSliceCutterDataflowNode::Evaluate(Dataflow::FContext& Context, const FData
 							OriginalSelection.SetSelected(Idx);
 						}
 					}
-
 				}
 			}
 
@@ -770,7 +770,6 @@ void FBrickCutterDataflowNode::Evaluate(Dataflow::FContext& Context, const FData
 							OriginalSelection.SetSelected(Idx);
 						}
 					}
-
 				}
 			}
 
@@ -787,5 +786,148 @@ void FBrickCutterDataflowNode::Evaluate(Dataflow::FContext& Context, const FData
 	}
 }
 
+void FMeshCutterDataflowNode::Evaluate(Dataflow::FContext& Context, const FDataflowOutput* Out) const
+{
+	if (Out->IsA(&Collection) ||
+		Out->IsA(&TransformSelection) ||
+		Out->IsA(&NewGeometryTransformSelection))
+	{
+		FDataflowTransformSelection InTransformSelection = GetValue(Context, &TransformSelection);
+		//
+		// If not connected select everything by default
+		//
+		if (!IsConnected(&TransformSelection))
+		{
+			const FManagedArrayCollection& InCollection = GetValue(Context, &Collection);
+
+			GeometryCollection::Facades::FCollectionTransformSelectionFacade TransformSelectionFacade(InCollection);
+			const TArray<int32>& SelectionArr = TransformSelectionFacade.SelectAll();
+
+			FDataflowTransformSelection NewTransformSelection;
+			NewTransformSelection.Initialize(InCollection.NumElements(FGeometryCollection::TransformGroup), false);
+			NewTransformSelection.SetFromArray(SelectionArr);
+
+			InTransformSelection = NewTransformSelection;
+		}
+
+		FBox InBoundingBox = GetValue(Context, &BoundingBox);
+		//
+		// If not connected set bounds to collection bounds
+		//
+		if (!IsConnected(&BoundingBox))
+		{
+			const FManagedArrayCollection& InCollection = GetValue(Context, &Collection);
+
+			GeometryCollection::Facades::FBoundsFacade BoundsFacade(InCollection);
+			const FBox& BoundingBoxInCollectionSpace = BoundsFacade.GetBoundingBoxInCollectionSpace();
+
+			InBoundingBox = BoundingBoxInCollectionSpace;
+		}
+
+		if (InTransformSelection.AnySelected())
+		{
+			FManagedArrayCollection InCollection = GetValue(Context, &Collection);
+
+			if (TObjectPtr<UDynamicMesh> InCuttingMesh = GetValue(Context, &CuttingMesh))
+			{
+				const UE::Geometry::FDynamicMesh3& InDynCuttingMesh = InCuttingMesh->GetMeshRef();
+
+				if (InDynCuttingMesh.VertexCount() > 0)
+				{
+					const int32 InRandomSeed = GetValue(Context, &RandomSeed);
+					const int32 InNumberToScatter = GetValue(Context, &NumberToScatter);
+					const int32 InGridX = GetValue(Context, &GridX);
+					const int32 InGridY = GetValue(Context, &GridY);
+					const int32 InGridZ = GetValue(Context, &GridZ);
+					const float InVariability = GetValue(Context, &Variability);
+					const float InMinScaleFactor = GetValue(Context, &MinScaleFactor);
+					const float InMaxScaleFactor = GetValue(Context, &MaxScaleFactor);
+					const float InRollRange = GetValue(Context, &RollRange);
+					const float InPitchRange = GetValue(Context, &PitchRange);
+					const float InYawRange = GetValue(Context, &YawRange);
+					const FTransform InTransform = GetValue(Context, &Transform);
+					const float InChanceToFracture = GetValue(Context, &ChanceToFracture);
+					const float InCollisionSampleSpacing = GetValue(Context, &CollisionSampleSpacing);
+
+					TArray<FTransform> MeshTransforms;
+
+					if (CutDistribution == EMeshCutterCutDistribution::SingleCut)
+					{
+						MeshTransforms.Add(InTransform);
+					}
+					else
+					{
+						FFractureEngineFracturing::GenerateMeshTransforms(MeshTransforms,
+							InBoundingBox,
+							InRandomSeed,
+							CutDistribution,
+							InNumberToScatter,
+							InGridX,
+							InGridY,
+							InGridZ,
+							InVariability,
+							InMinScaleFactor,
+							InMaxScaleFactor,
+							bRandomOrientation,
+							InRollRange,
+							InPitchRange,
+							InYawRange);
+					}
+
+					int32 ResultGeometryIndex = FFractureEngineFracturing::MeshCutter(MeshTransforms,
+						InCollection,
+						InTransformSelection,
+						InDynCuttingMesh,
+						InTransform,
+						InRandomSeed,
+						InChanceToFracture,
+						SplitIslands,
+						InCollisionSampleSpacing);
+
+					FDataflowTransformSelection NewSelection;
+					FDataflowTransformSelection OriginalSelection;
+
+					if (ResultGeometryIndex != INDEX_NONE)
+					{
+						if (InCollection.HasAttribute("TransformIndex", FGeometryCollection::GeometryGroup))
+						{
+							const TManagedArray<int32>& TransformIndices = InCollection.GetAttribute<int32>("TransformIndex", FGeometryCollection::GeometryGroup);
+
+							NewSelection.Initialize(TransformIndices.Num(), false);
+							OriginalSelection.Initialize(TransformIndices.Num(), false);
+
+							// The newly fractured pieces are added to the end of the transform array (starting position is ResultGeometryIndex)
+							for (int32 Idx = ResultGeometryIndex; Idx < TransformIndices.Num(); ++Idx)
+							{
+								int32 BoneIdx = TransformIndices[Idx];
+								NewSelection.SetSelected(BoneIdx);
+							}
+
+							for (int32 Idx = 0; Idx < InTransformSelection.Num(); ++Idx)
+							{
+								if (InTransformSelection.IsSelected(Idx))
+								{
+									OriginalSelection.SetSelected(Idx);
+								}
+							}
+						}
+					}
+
+					SetValue(Context, MoveTemp(InCollection), &Collection);
+					SetValue(Context, OriginalSelection, &TransformSelection);
+					SetValue(Context, NewSelection, &NewGeometryTransformSelection);
+
+					return;
+				}
+			}
+		}
+
+		SafeForwardInput(Context, &Collection, &Collection);
+		SetValue(Context, InTransformSelection, &TransformSelection);
+		SetValue(Context, FDataflowTransformSelection(), &NewGeometryTransformSelection);
+	}
+}
+
+
 
 
diff --git a/Engine/Plugins/Experimental/GeometryCollectionPlugin/Source/GeometryCollectionNodes/Private/Dataflow/GeometryCollectionMakeNodes.cpp b/Engine/Plugins/Experimental/GeometryCollectionPlugin/Source/GeometryCollectionNodes/Private/Dataflow/GeometryCollectionMakeNodes.cpp
--- a/Engine/Plugins/Experimental/GeometryCollectionPlugin/Source/GeometryCollectionNodes/Private/Dataflow/GeometryCollectionMakeNodes.cpp
+++ b/Engine/Plugins/Experimental/GeometryCollectionPlugin/Source/GeometryCollectionNodes/Private/Dataflow/GeometryCollectionMakeNodes.cpp
@@ -56,6 +56,7 @@ namespace Dataflow
 		DATAFLOW_NODE_REGISTER_CREATION_FACTORY(FMakeLiteralVectorDataflowNode);
 		DATAFLOW_NODE_REGISTER_CREATION_FACTORY(FMakeTransformDataflowNode);
 		DATAFLOW_NODE_REGISTER_CREATION_FACTORY(FMakeQuaternionDataflowNode);
+		DATAFLOW_NODE_REGISTER_CREATION_FACTORY(FMakeFloatArrayDataflowNode);
 
 		// Generators
 		DATAFLOW_NODE_REGISTER_CREATION_FACTORY_NODE_COLORS_BY_CATEGORY("Generators", FLinearColor(.4f, 0.8f, 0.f), CDefaultNodeBodyTintColor);
@@ -166,4 +167,17 @@ void FMakeQuaternionDataflowNode::Evaluate(Dataflow::FContext& Context, const FD
 	}
 }
 
+void FMakeFloatArrayDataflowNode::Evaluate(Dataflow::FContext& Context, const FDataflowOutput* Out) const
+{
+	if (Out->IsA(&FloatArray))
+	{
+		const int32 InNumElements = GetValue(Context, &NumElements);
+		const float InValue = GetValue(Context, &Value);
+
+		TArray<float> OutFloatArray;
+		OutFloatArray.Init(InValue, InNumElements);
+
+		SetValue(Context, OutFloatArray, &FloatArray);
+	}
+}
 
diff --git a/Engine/Plugins/Experimental/GeometryCollectionPlugin/Source/GeometryCollectionNodes/Private/Dataflow/GeometryCollectionMeshNodes.cpp b/Engine/Plugins/Experimental/GeometryCollectionPlugin/Source/GeometryCollectionNodes/Private/Dataflow/GeometryCollectionMeshNodes.cpp
--- a/Engine/Plugins/Experimental/GeometryCollectionPlugin/Source/GeometryCollectionNodes/Private/Dataflow/GeometryCollectionMeshNodes.cpp
+++ b/Engine/Plugins/Experimental/GeometryCollectionPlugin/Source/GeometryCollectionNodes/Private/Dataflow/GeometryCollectionMeshNodes.cpp
@@ -216,24 +216,27 @@ void FCollectionToMeshDataflowNode::Evaluate(Dataflow::FContext& Context, const
 void FStaticMeshToMeshDataflowNode::Evaluate(Dataflow::FContext& Context, const FDataflowOutput* Out) const
 {
 #if WITH_EDITORONLY_DATA
-	if (Out->IsA<TObjectPtr<UDynamicMesh>>(&Mesh))
+	if (Out->IsA(&Mesh))
 	{
-		if (FMeshDescription* MeshDescription = bUseHiRes ? StaticMesh->GetHiResMeshDescription() : StaticMesh->GetMeshDescription(LODLevel))
+		if (StaticMesh)
 		{
-			TObjectPtr<UDynamicMesh> NewMesh = NewObject<UDynamicMesh>();
-			NewMesh->Reset();
+			if (FMeshDescription* MeshDescription = bUseHiRes ? StaticMesh->GetHiResMeshDescription() : StaticMesh->GetMeshDescription(LODLevel))
+			{
+				TObjectPtr<UDynamicMesh> NewMesh = NewObject<UDynamicMesh>();
+				NewMesh->Reset();
+
+				UE::Geometry::FDynamicMesh3& DynMesh = NewMesh->GetMeshRef();
+				{
+					FMeshDescriptionToDynamicMesh ConverterToDynamicMesh;
+					ConverterToDynamicMesh.Convert(MeshDescription, DynMesh);
+				}
 
-			UE::Geometry::FDynamicMesh3& DynMesh = NewMesh->GetMeshRef();
+				SetValue(Context, NewMesh, &Mesh);
+			}
+			else
 			{
-				FMeshDescriptionToDynamicMesh ConverterToDynamicMesh;
-				ConverterToDynamicMesh.Convert(MeshDescription, DynMesh);
+				SetValue(Context, TObjectPtr<UDynamicMesh>(NewObject<UDynamicMesh>()), &Mesh);
 			}
-
-			SetValue(Context, NewMesh, &Mesh);
-		}
-		else
-		{
-			SetValue(Context, TObjectPtr<UDynamicMesh>(NewObject<UDynamicMesh>()), &Mesh);
 		}
 	}
 #endif
diff --git a/Engine/Plugins/Experimental/GeometryCollectionPlugin/Source/GeometryCollectionNodes/Private/Dataflow/GeometryCollectionNodesPlugin.cpp b/Engine/Plugins/Experimental/GeometryCollectionPlugin/Source/GeometryCollectionNodes/Private/Dataflow/GeometryCollectionNodesPlugin.cpp
--- a/Engine/Plugins/Experimental/GeometryCollectionPlugin/Source/GeometryCollectionNodes/Private/Dataflow/GeometryCollectionNodesPlugin.cpp
+++ b/Engine/Plugins/Experimental/GeometryCollectionPlugin/Source/GeometryCollectionNodes/Private/Dataflow/GeometryCollectionNodesPlugin.cpp
@@ -22,6 +22,7 @@
 #include "Dataflow/GeometryCollectionOverrideNodes.h"
 #include "Dataflow/GeometryCollectionProcessingNodes.h"
 #include "Dataflow/GeometryCollectionSelectionNodes.h"
+#include "Dataflow/GeometryCollectionSamplingNodes.h"
 #include "Dataflow/GeometryCollectionSkeletalMeshToCollectionNode.h"
 #include "Dataflow/GeometryCollectionSkeletonToCollectionNode.h"
 #include "Dataflow/GeometryCollectionStaticMeshToCollectionNode.h"
@@ -57,6 +58,7 @@ void IGeometryCollectionNodesPlugin::StartupModule()
 	Dataflow::GeometryCollectionVerticesNodes();
 	Dataflow::GeometryCollectionArrayNodes();
 	Dataflow::GeometryCollectionDebugNodes();
+	Dataflow::GeometryCollectionSamplingNodes();
 	DATAFLOW_NODE_REGISTER_CREATION_FACTORY(FCreateColorArrayFromFloatArrayDataflowNode);
 	DATAFLOW_NODE_REGISTER_CREATION_FACTORY(FGeometryCollectionVertexScalarToVertexIndicesNode);
 	DATAFLOW_NODE_REGISTER_CREATION_FACTORY(FSetVertexColorFromFloatArrayDataflowNode);
diff --git a/Engine/Plugins/Experimental/GeometryCollectionPlugin/Source/GeometryCollectionNodes/Private/Dataflow/GeometryCollectionSamplingNodes.cpp b/Engine/Plugins/Experimental/GeometryCollectionPlugin/Source/GeometryCollectionNodes/Private/Dataflow/GeometryCollectionSamplingNodes.cpp
new file mode 100644
--- /dev/null
+++ b/Engine/Plugins/Experimental/GeometryCollectionPlugin/Source/GeometryCollectionNodes/Private/Dataflow/GeometryCollectionSamplingNodes.cpp
@@ -0,0 +1,177 @@
+// Copyright Epic Games, Inc. All Rights Reserved.
+
+#include "Dataflow/GeometryCollectionSamplingNodes.h"
+
+#include "UDynamicMesh.h"
+
+#include UE_INLINE_GENERATED_CPP_BY_NAME(GeometryCollectionSamplingNodes)
+
+namespace Dataflow
+{
+
+	void GeometryCollectionSamplingNodes()
+	{
+		static const FLinearColor CDefaultNodeBodyTintColor = FLinearColor(0.f, 0.f, 0.f, 0.5f);
+
+		DATAFLOW_NODE_REGISTER_CREATION_FACTORY(FUniformPointSamplingDataflowNode);
+		DATAFLOW_NODE_REGISTER_CREATION_FACTORY(FNonUniformPointSamplingDataflowNode);
+		DATAFLOW_NODE_REGISTER_CREATION_FACTORY(FVertexWeightedPointSamplingDataflowNode);
+
+		// GeometryCollection|Sampling
+		DATAFLOW_NODE_REGISTER_CREATION_FACTORY_NODE_COLORS_BY_CATEGORY("GeometryCollection|Sampling", FLinearColor(.1f, 1.f, 0.6f), CDefaultNodeBodyTintColor);
+	}
+}
+
+void FUniformPointSamplingDataflowNode::Evaluate(Dataflow::FContext& Context, const FDataflowOutput* Out) const
+{
+	if (Out->IsA(&SamplePoints) || 
+		Out->IsA(&SampleTriangleIDs) ||
+		Out->IsA(&NumSamplePoints))
+	{
+		if (TObjectPtr<UDynamicMesh> InTargetMesh = GetValue(Context, &TargetMesh))
+		{
+			const UE::Geometry::FDynamicMesh3& InDynTargetMesh = InTargetMesh->GetMeshRef();
+
+			if (InDynTargetMesh.VertexCount() > 0)
+			{
+				TArray<FTransform> OutSamples;
+				TArray<int32> OutTriangleIDs;
+				TArray<FVector> OutBarycentricCoords;
+
+				FFractureEngineSampling::ComputeUniformPointSampling(InDynTargetMesh,
+					GetValue(Context, &SamplingRadius),
+					GetValue(Context, &MaxNumSamples),
+					GetValue(Context, &SubSampleDensity),
+					GetValue(Context, &RandomSeed),
+					OutSamples, 
+					OutTriangleIDs, 
+					OutBarycentricCoords);
+
+				const int32 NumSamples = OutSamples.Num();
+
+				TArray<FVector> OutPoints;
+				OutPoints.AddUninitialized(NumSamples);
+
+				for (int32 Idx = 0; Idx < NumSamples; ++Idx)
+				{
+					OutPoints[Idx] = OutSamples[Idx].GetTranslation();
+				}
+
+				SetValue(Context, OutPoints, &SamplePoints);
+				SetValue(Context, OutTriangleIDs, &SampleTriangleIDs);
+				SetValue(Context, OutBarycentricCoords, &SampleBarycentricCoords);
+				SetValue(Context, OutPoints.Num(), &NumSamplePoints);
+			}
+		}
+	}
+}
+
+void FNonUniformPointSamplingDataflowNode::Evaluate(Dataflow::FContext& Context, const FDataflowOutput* Out) const
+{
+	if (Out->IsA(&SamplePoints) ||
+		Out->IsA(&SampleRadii) ||
+		Out->IsA(&SampleTriangleIDs) ||
+		Out->IsA(&NumSamplePoints))
+	{
+		if (TObjectPtr<UDynamicMesh> InTargetMesh = GetValue(Context, &TargetMesh))
+		{
+			const UE::Geometry::FDynamicMesh3& InDynTargetMesh = InTargetMesh->GetMeshRef();
+
+			if (InDynTargetMesh.VertexCount() > 0)
+			{
+				TArray<FTransform> OutSamples;
+				TArray<float> OutSampleRadii;
+				TArray<int32> OutTriangleIDs;
+				TArray<FVector> OutBarycentricCoords;
+
+				FFractureEngineSampling::ComputeNonUniformPointSampling(InDynTargetMesh,
+					GetValue(Context, &SamplingRadius),
+					GetValue(Context, &MaxNumSamples),
+					GetValue(Context, &SubSampleDensity),
+					GetValue(Context, &RandomSeed),
+					GetValue(Context, &MaxSamplingRadius),
+					SizeDistribution,
+					GetValue(Context, &SizeDistributionPower),
+					OutSamples,
+					OutSampleRadii,
+					OutTriangleIDs, 
+					OutBarycentricCoords);
+
+				const int32 NumSamples = OutSamples.Num();
+
+				TArray<FVector> OutPoints;
+				OutPoints.AddUninitialized(NumSamples);
+
+				for (int32 Idx = 0; Idx < NumSamples; ++Idx)
+				{
+					OutPoints[Idx] = OutSamples[Idx].GetTranslation();
+				}
+
+				SetValue(Context, OutPoints, &SamplePoints);
+				SetValue(Context, OutSampleRadii, &SampleRadii);
+				SetValue(Context, OutTriangleIDs, &SampleTriangleIDs);
+				SetValue(Context, OutBarycentricCoords, &SampleBarycentricCoords);
+				SetValue(Context, OutPoints.Num(), &NumSamplePoints);
+			}
+		}
+	}
+}
+
+void FVertexWeightedPointSamplingDataflowNode::Evaluate(Dataflow::FContext& Context, const FDataflowOutput* Out) const
+{
+	if (Out->IsA(&SamplePoints) ||
+		Out->IsA(&SampleRadii) ||
+		Out->IsA(&SampleTriangleIDs) ||
+		Out->IsA(&NumSamplePoints))
+	{
+		if (TObjectPtr<UDynamicMesh> InTargetMesh = GetValue(Context, &TargetMesh))
+		{
+			const UE::Geometry::FDynamicMesh3& InDynTargetMesh = InTargetMesh->GetMeshRef();
+
+			if (InDynTargetMesh.VertexCount() > 0)
+			{
+				if (IsConnected(&VertexWeights))
+				{
+					TArray<FTransform> OutSamples;
+					TArray<float> OutSampleRadii;
+					TArray<int32> OutTriangleIDs;
+					TArray<FVector> OutBarycentricCoords;
+
+					FFractureEngineSampling::ComputeVertexWeightedPointSampling(InDynTargetMesh,
+						GetValue(Context, &VertexWeights),
+						GetValue(Context, &SamplingRadius),
+						GetValue(Context, &MaxNumSamples),
+						GetValue(Context, &SubSampleDensity),
+						GetValue(Context, &RandomSeed),
+						GetValue(Context, &MaxSamplingRadius),
+						SizeDistribution,
+						GetValue(Context, &SizeDistributionPower),
+						WeightMode,
+						bInvertWeights,
+						OutSamples,
+						OutSampleRadii,
+						OutTriangleIDs,
+						OutBarycentricCoords);
+
+					const int32 NumSamples = OutSamples.Num();
+
+					TArray<FVector> OutPoints;
+					OutPoints.AddUninitialized(NumSamples);
+
+					for (int32 Idx = 0; Idx < NumSamples; ++Idx)
+					{
+						OutPoints[Idx] = OutSamples[Idx].GetTranslation();
+					}
+
+					SetValue(Context, OutPoints, &SamplePoints);
+					SetValue(Context, OutSampleRadii, &SampleRadii);
+					SetValue(Context, OutTriangleIDs, &SampleTriangleIDs);
+					SetValue(Context, OutBarycentricCoords, &SampleBarycentricCoords);
+					SetValue(Context, OutPoints.Num(), &NumSamplePoints);
+				}
+			}
+		}
+	}
+}
+
+
diff --git a/Engine/Plugins/Experimental/GeometryCollectionPlugin/Source/GeometryCollectionNodes/Private/Dataflow/GeometryCollectionSelectionNodes.cpp b/Engine/Plugins/Experimental/GeometryCollectionPlugin/Source/GeometryCollectionNodes/Private/Dataflow/GeometryCollectionSelectionNodes.cpp
--- a/Engine/Plugins/Experimental/GeometryCollectionPlugin/Source/GeometryCollectionNodes/Private/Dataflow/GeometryCollectionSelectionNodes.cpp
+++ b/Engine/Plugins/Experimental/GeometryCollectionPlugin/Source/GeometryCollectionNodes/Private/Dataflow/GeometryCollectionSelectionNodes.cpp
@@ -69,6 +69,7 @@ namespace Dataflow
 		DATAFLOW_NODE_REGISTER_CREATION_FACTORY(FCollectionFaceSelectionInvertDataflowNode);
 		DATAFLOW_NODE_REGISTER_CREATION_FACTORY(FCollectionVertexSelectionByPercentageDataflowNode);
 		DATAFLOW_NODE_REGISTER_CREATION_FACTORY(FCollectionVertexSelectionSetOperationDataflowNode);
+		DATAFLOW_NODE_REGISTER_CREATION_FACTORY(FCollectionSelectionByAttrDataflowNode);
 
 		// GeometryCollection|Selection
 		DATAFLOW_NODE_REGISTER_CREATION_FACTORY_NODE_COLORS_BY_CATEGORY("GeometryCollection|Selection", FLinearColor(1.f, 1.f, 0.05f), CDefaultNodeBodyTintColor);
@@ -906,7 +907,7 @@ void FCollectionFaceSelectionCustomDataflowNode::Evaluate(Dataflow::FContext& Co
 	{
 		const FManagedArrayCollection& InCollection = GetValue<FManagedArrayCollection>(Context, &Collection);
 
-		if (InCollection.HasGroup(FGeometryCollection::GeometryGroup))
+		if (InCollection.HasGroup(FGeometryCollection::FacesGroup))
 		{
 			const int32 NumFaces = InCollection.NumElements(FGeometryCollection::FacesGroup);
 
@@ -952,12 +953,12 @@ void FCollectionFaceSelectionCustomDataflowNode::Evaluate(Dataflow::FContext& Co
 
 void FCollectionSelectionConvertDataflowNode::Evaluate(Dataflow::FContext& Context, const FDataflowOutput* Out) const
 {
-	if (Out->IsA<FDataflowTransformSelection>(&TransformSelection))
+	if (Out->IsA(&TransformSelection))
 	{
-		if (IsConnected<FDataflowVertexSelection>(&VertexSelection))
+		if (IsConnected(&VertexSelection))
 		{
-			const FManagedArrayCollection& InCollection = GetValue<FManagedArrayCollection>(Context, &Collection);
-			const FDataflowVertexSelection& InVertexSelection = GetValue<FDataflowVertexSelection>(Context, &VertexSelection);
+			const FManagedArrayCollection& InCollection = GetValue(Context, &Collection);
+			const FDataflowVertexSelection& InVertexSelection = GetValue(Context, &VertexSelection);
 
 			GeometryCollection::Facades::FCollectionTransformSelectionFacade TransformSelectionFacade(InCollection);
 			const TArray<int32>& SelectionArr = TransformSelectionFacade.ConvertVertexSelectionToTransformSelection(InVertexSelection.AsArray(), bAllElementsMustBeSelected);
@@ -968,10 +969,10 @@ void FCollectionSelectionConvertDataflowNode::Evaluate(Dataflow::FContext& Conte
 
 			SetValue(Context, MoveTemp(NewTransformSelection), &TransformSelection);
 		}
-		else if (IsConnected<FDataflowFaceSelection>(&FaceSelection))
+		else if (IsConnected(&FaceSelection))
 		{
-			const FManagedArrayCollection& InCollection = GetValue<FManagedArrayCollection>(Context, &Collection);
-			const FDataflowFaceSelection& InFaceSelection = GetValue<FDataflowFaceSelection>(Context, &FaceSelection);
+			const FManagedArrayCollection& InCollection = GetValue(Context, &Collection);
+			const FDataflowFaceSelection& InFaceSelection = GetValue(Context, &FaceSelection);
 
 			GeometryCollection::Facades::FCollectionTransformSelectionFacade TransformSelectionFacade(InCollection);
 			const TArray<int32>& SelectionArr = TransformSelectionFacade.ConvertFaceSelectionToTransformSelection(InFaceSelection.AsArray(), bAllElementsMustBeSelected);
@@ -982,19 +983,18 @@ void FCollectionSelectionConvertDataflowNode::Evaluate(Dataflow::FContext& Conte
 
 			SetValue(Context, MoveTemp(NewTransformSelection), &TransformSelection);
 		}
-		else if (IsConnected<FDataflowTransformSelection>(&TransformSelection))
+		else
 		{
 			// Passthrough
-			const FDataflowTransformSelection& InTransformSelection = GetValue<FDataflowTransformSelection>(Context, &TransformSelection);
-			SetValue(Context, InTransformSelection, &TransformSelection);
+			SafeForwardInput(Context, &TransformSelection, &TransformSelection);
 		}
 	}
-	else if (Out->IsA<FDataflowFaceSelection>(&FaceSelection))
+	else if (Out->IsA(&FaceSelection))
 	{
-		if (IsConnected<FDataflowVertexSelection>(&VertexSelection))
+		if (IsConnected(&VertexSelection))
 		{
-			const FManagedArrayCollection& InCollection = GetValue<FManagedArrayCollection>(Context, &Collection);
-			const FDataflowVertexSelection& InVertexSelection = GetValue<FDataflowVertexSelection>(Context, &VertexSelection);
+			const FManagedArrayCollection& InCollection = GetValue(Context, &Collection);
+			const FDataflowVertexSelection& InVertexSelection = GetValue(Context, &VertexSelection);
 
 			GeometryCollection::Facades::FCollectionTransformSelectionFacade TransformSelectionFacade(InCollection);
 			const TArray<int32>& SelectionArr = TransformSelectionFacade.ConvertVertexSelectionToFaceSelection(InVertexSelection.AsArray(), bAllElementsMustBeSelected);
@@ -1005,16 +1005,10 @@ void FCollectionSelectionConvertDataflowNode::Evaluate(Dataflow::FContext& Conte
 
 			SetValue(Context, MoveTemp(NewFaceSelection), &FaceSelection);
 		}
-		else if (IsConnected<FDataflowFaceSelection>(&FaceSelection))
-		{
-			// Passthrough
-			const FDataflowFaceSelection& InFaceSelection = GetValue<FDataflowFaceSelection>(Context, &FaceSelection);
-			SetValue(Context, InFaceSelection, &FaceSelection);
-		}
-		else if (IsConnected<FDataflowTransformSelection>(&TransformSelection))
+		else if (IsConnected(&TransformSelection))
 		{
-			const FManagedArrayCollection& InCollection = GetValue<FManagedArrayCollection>(Context, &Collection);
-			const FDataflowTransformSelection& InTransformSelection = GetValue<FDataflowTransformSelection>(Context, &TransformSelection);
+			const FManagedArrayCollection& InCollection = GetValue(Context, &Collection);
+			const FDataflowTransformSelection& InTransformSelection = GetValue(Context, &TransformSelection);
 
 			GeometryCollection::Facades::FCollectionTransformSelectionFacade TransformSelectionFacade(InCollection);
 			const TArray<int32>& SelectionArr = TransformSelectionFacade.ConvertTransformSelectionToFaceSelection(InTransformSelection.AsArray());
@@ -1025,19 +1019,18 @@ void FCollectionSelectionConvertDataflowNode::Evaluate(Dataflow::FContext& Conte
 
 			SetValue(Context, MoveTemp(NewFaceSelection), &FaceSelection);
 		}
-	}
-	else if (Out->IsA<FDataflowVertexSelection>(&VertexSelection))
-	{
-		if (IsConnected<FDataflowVertexSelection>(&VertexSelection))
+		else
 		{
 			// Passthrough
-			const FDataflowVertexSelection& InVertexSelection = GetValue<FDataflowVertexSelection>(Context, &VertexSelection);
-			SetValue(Context, InVertexSelection, &VertexSelection);
+			SafeForwardInput(Context, &FaceSelection, &FaceSelection);
 		}
-		else if (IsConnected<FDataflowFaceSelection>(&FaceSelection))
+	}
+	else if (Out->IsA(&VertexSelection))
+	{
+		if (IsConnected(&FaceSelection))
 		{
-			const FManagedArrayCollection& InCollection = GetValue<FManagedArrayCollection>(Context, &Collection);
-			const FDataflowFaceSelection& InFaceSelection = GetValue<FDataflowFaceSelection>(Context, &FaceSelection);
+			const FManagedArrayCollection& InCollection = GetValue(Context, &Collection);
+			const FDataflowFaceSelection& InFaceSelection = GetValue(Context, &FaceSelection);
 
 			GeometryCollection::Facades::FCollectionTransformSelectionFacade TransformSelectionFacade(InCollection);
 			const TArray<int32>& SelectionArr = TransformSelectionFacade.ConvertFaceSelectionToVertexSelection(InFaceSelection.AsArray());
@@ -1048,10 +1041,10 @@ void FCollectionSelectionConvertDataflowNode::Evaluate(Dataflow::FContext& Conte
 
 			SetValue(Context, MoveTemp(NewVertexSelection), &VertexSelection);
 		}
-		else if (IsConnected<FDataflowTransformSelection>(&TransformSelection))
+		else if (IsConnected(&TransformSelection))
 		{
-			const FManagedArrayCollection& InCollection = GetValue<FManagedArrayCollection>(Context, &Collection);
-			const FDataflowTransformSelection& InTransformSelection = GetValue<FDataflowTransformSelection>(Context, &TransformSelection);
+			const FManagedArrayCollection& InCollection = GetValue(Context, &Collection);
+			const FDataflowTransformSelection& InTransformSelection = GetValue(Context, &TransformSelection);
 
 			GeometryCollection::Facades::FCollectionTransformSelectionFacade TransformSelectionFacade(InCollection);
 			const TArray<int32>& SelectionArr = TransformSelectionFacade.ConvertTransformSelectionToVertexSelection(InTransformSelection.AsArray());
@@ -1062,8 +1055,13 @@ void FCollectionSelectionConvertDataflowNode::Evaluate(Dataflow::FContext& Conte
 
 			SetValue(Context, MoveTemp(NewVertexSelection), &VertexSelection);
 		}
+		else
+		{
+			// Passthrough
+			SafeForwardInput(Context, &VertexSelection, &VertexSelection);
+		}
 	}
-	else if (Out->IsA<FManagedArrayCollection>(&Collection))
+	else if (Out->IsA(&Collection))
 	{
 		SafeForwardInput(Context, &Collection, &Collection);
 	}
@@ -1141,3 +1139,160 @@ void FCollectionVertexSelectionSetOperationDataflowNode::Evaluate(Dataflow::FCon
 	}
 }
 
+static void CreateSelectionFromAttr(const FManagedArrayCollection& InCollection,
+	const FName InGroup,
+	const FName InAttribute,
+	const FString InValue,
+	const ESelectionByAttrOperation InOperation,
+	FDataflowSelection& OutSelection)
+{
+	const FManagedArrayCollection::EArrayType ArrayType = InCollection.GetAttributeType(InAttribute, InGroup);
+	const int32 NumElements = InCollection.NumElements(InGroup);
+
+	if (ArrayType == FManagedArrayCollection::EArrayType::FFloatType)
+	{
+		const TManagedArray<float>* const Array = InCollection.FindAttributeTyped<float>(InAttribute, InGroup);
+		if (InValue.IsNumeric())
+		{
+			float FloatValue = FCString::Atof(*InValue);
+
+			for (int32 Idx = 0; Idx < NumElements; ++Idx)
+			{
+				if ((InOperation == ESelectionByAttrOperation::Equal && (*Array)[Idx] == FloatValue) ||
+					(InOperation == ESelectionByAttrOperation::NotEqual && (*Array)[Idx] != FloatValue) ||
+					(InOperation == ESelectionByAttrOperation::Greater && (*Array)[Idx] > FloatValue) ||
+					(InOperation == ESelectionByAttrOperation::GreaterOrEqual && (*Array)[Idx] >= FloatValue) ||
+					(InOperation == ESelectionByAttrOperation::Smaller && (*Array)[Idx] < FloatValue) ||
+					(InOperation == ESelectionByAttrOperation::SmallerOrEqual && (*Array)[Idx] <= FloatValue))
+				{
+					OutSelection.SetSelected(Idx);
+				}
+			}
+		}
+		else
+		{
+			// Error: Invalid Value specified
+			return;
+		}
+	}
+	else if (ArrayType == FManagedArrayCollection::EArrayType::FInt32Type)
+	{
+		const TManagedArray<int32>* const Array = InCollection.FindAttributeTyped<int32>(InAttribute, InGroup);
+		if (InValue.IsNumeric())
+		{
+			float IntValue = FCString::Atoi(*InValue);
+
+			for (int32 Idx = 0; Idx < NumElements; ++Idx)
+			{
+				if ((InOperation == ESelectionByAttrOperation::Equal && (*Array)[Idx] == IntValue) ||
+					(InOperation == ESelectionByAttrOperation::NotEqual && (*Array)[Idx] != IntValue) ||
+					(InOperation == ESelectionByAttrOperation::Greater && (*Array)[Idx] > IntValue) ||
+					(InOperation == ESelectionByAttrOperation::GreaterOrEqual && (*Array)[Idx] >= IntValue) ||
+					(InOperation == ESelectionByAttrOperation::Smaller && (*Array)[Idx] < IntValue) ||
+					(InOperation == ESelectionByAttrOperation::SmallerOrEqual && (*Array)[Idx] <= IntValue))
+				{
+					OutSelection.SetSelected(Idx);
+				}
+			}
+		}
+		else
+		{
+			// Error: Invalid Value specified
+			return;
+		}
+	}
+	else if (ArrayType == FManagedArrayCollection::EArrayType::FStringType)
+	{
+		const TManagedArray<FString>* const Array = InCollection.FindAttributeTyped<FString>(InAttribute, InGroup);
+
+		for (int32 Idx = 0; Idx < NumElements; ++Idx)
+		{
+			if ((InOperation == ESelectionByAttrOperation::Equal && (*Array)[Idx] == InValue) ||
+				(InOperation == ESelectionByAttrOperation::NotEqual && !((*Array)[Idx] == InValue)))
+			{
+				OutSelection.SetSelected(Idx);
+			}
+		}
+	}
+	else if (ArrayType == FManagedArrayCollection::EArrayType::FBoolType)
+	{
+		const TManagedArray<bool>* const Array = InCollection.FindAttributeTyped<bool>(InAttribute, InGroup);
+		bool BoolValue = false;
+		if (InValue.IsNumeric())
+		{
+			float FloatValue = FCString::Atof(*InValue);
+
+			if (FloatValue > 0.f)
+			{
+				BoolValue = true;
+			}
+		}
+		else
+		{
+			if (InValue == FString("true") || InValue == FString("True"))
+			{
+				BoolValue = true;
+			}
+		}
+
+		for (int32 Idx = 0; Idx < NumElements; ++Idx)
+		{
+			if ((InOperation == ESelectionByAttrOperation::Equal && (*Array)[Idx] == BoolValue) ||
+				(InOperation == ESelectionByAttrOperation::NotEqual && !((*Array)[Idx] == BoolValue)))
+			{
+				OutSelection.SetSelected(Idx);
+			}
+		}
+	}
+}
+
+
+void FCollectionSelectionByAttrDataflowNode::Evaluate(Dataflow::FContext& Context, const FDataflowOutput* Out) const
+{
+	if (Out->IsA<FDataflowVertexSelection>(&VertexSelection) ||
+		Out->IsA<FDataflowFaceSelection>(&FaceSelection) ||
+		Out->IsA<FDataflowTransformSelection>(&TransformSelection) ||
+		Out->IsA<FDataflowGeometrySelection>(&GeometrySelection) ||
+		Out->IsA<FDataflowMaterialSelection>(&MaterialSelection))
+	{
+		const FManagedArrayCollection& InCollection = GetValue<FManagedArrayCollection>(Context, &Collection);
+		const FName GroupName = Dataflow::Private::GetAttributeFromEnumAsName(Group);
+
+		if (InCollection.HasGroup(GroupName))
+		{
+			if (InCollection.HasAttribute(FName(*Attribute), GroupName))
+			{
+				const int32 NumFaces = InCollection.NumElements(GroupName);
+
+				FDataflowSelection NewSelection;
+				NewSelection.Initialize(NumFaces, false);
+
+				CreateSelectionFromAttr(InCollection,
+					GroupName,
+					FName(*Attribute),
+					Value,
+					Operation,
+					NewSelection);
+
+				SetValue(Context, Group == ESelectionByAttrGroup::Vertices ? MoveTemp(NewSelection) : FDataflowSelection(), &VertexSelection);
+				SetValue(Context, Group == ESelectionByAttrGroup::Faces ? MoveTemp(NewSelection) : FDataflowSelection(), &FaceSelection);
+				SetValue(Context, Group == ESelectionByAttrGroup::Transform ? MoveTemp(NewSelection) : FDataflowSelection(), &TransformSelection);
+				SetValue(Context, Group == ESelectionByAttrGroup::Geometry ? MoveTemp(NewSelection) : FDataflowSelection(), &GeometrySelection);
+				SetValue(Context, Group == ESelectionByAttrGroup::Material ? MoveTemp(NewSelection) : FDataflowSelection(), &MaterialSelection);
+
+				return;
+			}
+		}
+
+		SetValue(Context, FDataflowSelection(), &VertexSelection);
+		SetValue(Context, FDataflowSelection(), &FaceSelection);
+		SetValue(Context, FDataflowSelection(), &TransformSelection);
+		SetValue(Context, FDataflowSelection(), &GeometrySelection);
+		SetValue(Context, FDataflowSelection(), &MaterialSelection);
+	}
+	else if (Out->IsA<FManagedArrayCollection>(&Collection))
+	{
+		SafeForwardInput(Context, &Collection, &Collection);
+	}
+}
+
diff --git a/Engine/Plugins/Experimental/GeometryCollectionPlugin/Source/GeometryCollectionNodes/Private/Dataflow/GeometryCollectionUtilityNodes.cpp b/Engine/Plugins/Experimental/GeometryCollectionPlugin/Source/GeometryCollectionNodes/Private/Dataflow/GeometryCollectionUtilityNodes.cpp
--- a/Engine/Plugins/Experimental/GeometryCollectionPlugin/Source/GeometryCollectionNodes/Private/Dataflow/GeometryCollectionUtilityNodes.cpp
+++ b/Engine/Plugins/Experimental/GeometryCollectionPlugin/Source/GeometryCollectionNodes/Private/Dataflow/GeometryCollectionUtilityNodes.cpp
@@ -34,6 +34,10 @@ namespace Dataflow
 		DATAFLOW_NODE_REGISTER_CREATION_FACTORY(FMergeConvexHullsDataflowNode);
 		DATAFLOW_NODE_REGISTER_CREATION_FACTORY(FUpdateVolumeAttributesDataflowNode);
 		DATAFLOW_NODE_REGISTER_CREATION_FACTORY(FGetConvexHullVolumeDataflowNode);
+		DATAFLOW_NODE_REGISTER_CREATION_FACTORY(FFixTinyGeoDataflowNode);
+		DATAFLOW_NODE_REGISTER_CREATION_FACTORY(FRecomputeNormalsInGeometryCollectionDataflowNode);
+		DATAFLOW_NODE_REGISTER_CREATION_FACTORY(FResampleGeometryCollectionDataflowNode);
+		DATAFLOW_NODE_REGISTER_CREATION_FACTORY(FValidateGeometryCollectionDataflowNode);
 	}
 }
 
@@ -603,3 +607,147 @@ void FGetConvexHullVolumeDataflowNode::Evaluate(Dataflow::FContext& Context, con
 		SetValue(Context, VolumeSum, &Volume);
 	}
 }
+
+void FFixTinyGeoDataflowNode::Evaluate(Dataflow::FContext& Context, const FDataflowOutput* Out) const
+{
+	if (Out->IsA(&Collection))
+	{
+		FDataflowTransformSelection InTransformSelection = GetValue(Context, &TransformSelection);
+		//
+		// If not connected select everything by default
+		//
+		if (!IsConnected(&TransformSelection))
+		{
+			const FManagedArrayCollection& InCollection = GetValue(Context, &Collection);
+
+			GeometryCollection::Facades::FCollectionTransformSelectionFacade TransformSelectionFacade(InCollection);
+			const TArray<int32>& SelectionArr = TransformSelectionFacade.SelectAll();
+
+			FDataflowTransformSelection NewTransformSelection;
+			NewTransformSelection.Initialize(InCollection.NumElements(FGeometryCollection::TransformGroup), false);
+			NewTransformSelection.SetFromArray(SelectionArr);
+
+			InTransformSelection = NewTransformSelection;
+		}
+
+		if (InTransformSelection.AnySelected())
+		{
+			FManagedArrayCollection InCollection = GetValue(Context, &Collection);
+
+			FFractureEngineUtility::FixTinyGeo(InCollection,
+				InTransformSelection,
+				MergeType,
+				bOnFractureLevel,
+				SelectionMethod,
+				MinVolumeCubeRoot,
+				RelativeVolume,
+				UseBoneSelection,
+				bOnlyClusters,
+				NeighborSelection,
+				bOnlyToConnected,
+				bOnlySameParent);
+
+			SetValue(Context, MoveTemp(InCollection), &Collection);
+
+			return;
+		}
+
+		SafeForwardInput(Context, &Collection, &Collection);
+	}
+}
+
+void FRecomputeNormalsInGeometryCollectionDataflowNode::Evaluate(Dataflow::FContext& Context, const FDataflowOutput* Out) const
+{
+	if (Out->IsA(&Collection))
+	{
+		FDataflowTransformSelection InTransformSelection = GetValue(Context, &TransformSelection);
+		//
+		// If not connected select everything by default
+		//
+		if (!IsConnected(&TransformSelection))
+		{
+			const FManagedArrayCollection& InCollection = GetValue(Context, &Collection);
+
+			GeometryCollection::Facades::FCollectionTransformSelectionFacade TransformSelectionFacade(InCollection);
+			const TArray<int32>& SelectionArr = TransformSelectionFacade.SelectAll();
+
+			FDataflowTransformSelection NewTransformSelection;
+			NewTransformSelection.Initialize(InCollection.NumElements(FGeometryCollection::TransformGroup), false);
+			NewTransformSelection.SetFromArray(SelectionArr);
+
+			InTransformSelection = NewTransformSelection;
+		}
+
+		if (InTransformSelection.AnySelected())
+		{
+			FManagedArrayCollection InCollection = GetValue(Context, &Collection);
+
+			FFractureEngineUtility::RecomputeNormalsInGeometryCollection(InCollection,
+				InTransformSelection,
+				bOnlyTangents,
+				bRecomputeSharpEdges,
+				SharpEdgeAngleThreshold,
+				bOnlyInternalSurfaces);
+
+			SetValue(Context, MoveTemp(InCollection), &Collection);
+
+			return;
+		}
+
+		SafeForwardInput(Context, &Collection, &Collection);
+	}
+}
+
+void FResampleGeometryCollectionDataflowNode::Evaluate(Dataflow::FContext& Context, const FDataflowOutput* Out) const
+{
+	if (Out->IsA(&Collection))
+	{
+		FDataflowTransformSelection InTransformSelection = GetValue(Context, &TransformSelection);
+		//
+		// If not connected select everything by default
+		//
+		if (!IsConnected(&TransformSelection))
+		{
+			const FManagedArrayCollection& InCollection = GetValue(Context, &Collection);
+
+			GeometryCollection::Facades::FCollectionTransformSelectionFacade TransformSelectionFacade(InCollection);
+			const TArray<int32>& SelectionArr = TransformSelectionFacade.SelectAll();
+
+			FDataflowTransformSelection NewTransformSelection;
+			NewTransformSelection.Initialize(InCollection.NumElements(FGeometryCollection::TransformGroup), false);
+			NewTransformSelection.SetFromArray(SelectionArr);
+
+			InTransformSelection = NewTransformSelection;
+		}
+
+		if (InTransformSelection.AnySelected())
+		{
+			FManagedArrayCollection InCollection = GetValue(Context, &Collection);
+
+			FFractureEngineUtility::ResampleGeometryCollection(InCollection,
+				InTransformSelection,
+				GetValue(Context, &CollisionSampleSpacing));
+
+			SetValue(Context, MoveTemp(InCollection), &Collection);
+
+			return;
+		}
+
+		SafeForwardInput(Context, &Collection, &Collection);
+	}
+}
+
+void FValidateGeometryCollectionDataflowNode::Evaluate(Dataflow::FContext& Context, const FDataflowOutput* Out) const
+{
+	if (Out->IsA(&Collection))
+	{
+		FManagedArrayCollection InCollection = GetValue(Context, &Collection);
+
+		FFractureEngineUtility::ValidateGeometryCollection(InCollection,
+			bRemoveUnreferencedGeometry,
+			bRemoveClustersOfOne,
+			bRemoveDanglingClusters);
+
+		SetValue(Context, MoveTemp(InCollection), &Collection);
+	}
+}
diff --git a/Engine/Plugins/Experimental/GeometryCollectionPlugin/Source/GeometryCollectionNodes/Public/Dataflow/GeometryCollectionArrayNodes.h b/Engine/Plugins/Experimental/GeometryCollectionPlugin/Source/GeometryCollectionNodes/Public/Dataflow/GeometryCollectionArrayNodes.h
--- a/Engine/Plugins/Experimental/GeometryCollectionPlugin/Source/GeometryCollectionNodes/Public/Dataflow/GeometryCollectionArrayNodes.h
+++ b/Engine/Plugins/Experimental/GeometryCollectionPlugin/Source/GeometryCollectionNodes/Public/Dataflow/GeometryCollectionArrayNodes.h
@@ -487,7 +487,46 @@ public:
 
 };
 
+/**
+ *
+ * Randomize elements in a float array (Random value will be in (RandomRangeMin, RandomRangeMax)
+ *
+ */
+USTRUCT()
+struct FRandomizeFloatArrayDataflowNode : public FDataflowNode
+{
+	GENERATED_USTRUCT_BODY()
+	DATAFLOW_NODE_DEFINE_INTERNAL(FRandomizeFloatArrayDataflowNode, "RandomizeFloatArray", "Utilities|Array", "")
+
+public:
+	/** Array to randomize */
+	UPROPERTY(meta = (DataflowInput, DataflowOutput, DataflowPassthrough = "FloatArray", DataflowIntrinsic))
+	TArray<float> FloatArray;
 
+	/** Random range min */
+	UPROPERTY(EditAnywhere, Category = "Random", meta = (DataflowInput, DisplayName = "Random Range Min"))
+	float RandomRangeMin = 0.f;
+
+	/** Random range max */
+	UPROPERTY(EditAnywhere, Category = "Random", meta = (DataflowInput, DisplayName = "Random Range Max"))
+	float RandomRangeMax = 1.f;
+
+	/** Seed for random */
+	UPROPERTY(EditAnywhere, Category = "Random", meta = (DataflowInput, UIMin = "0"));
+	int32 RandomSeed = 0;
+
+	FRandomizeFloatArrayDataflowNode(const Dataflow::FNodeParameters& InParam, FGuid InGuid = FGuid::NewGuid())
+		: FDataflowNode(InParam, InGuid)
+	{
+		RegisterInputConnection(&FloatArray);
+		RegisterInputConnection(&RandomRangeMin);
+		RegisterInputConnection(&RandomRangeMax);
+		RegisterInputConnection(&RandomSeed);
+		RegisterOutputConnection(&FloatArray, &FloatArray);
+	}
+
+	virtual void Evaluate(Dataflow::FContext& Context, const FDataflowOutput* Out) const override;
+};
 
 namespace Dataflow
 {
diff --git a/Engine/Plugins/Experimental/GeometryCollectionPlugin/Source/GeometryCollectionNodes/Public/Dataflow/GeometryCollectionFracturingNodes.h b/Engine/Plugins/Experimental/GeometryCollectionPlugin/Source/GeometryCollectionNodes/Public/Dataflow/GeometryCollectionFracturingNodes.h
--- a/Engine/Plugins/Experimental/GeometryCollectionPlugin/Source/GeometryCollectionNodes/Public/Dataflow/GeometryCollectionFracturingNodes.h
+++ b/Engine/Plugins/Experimental/GeometryCollectionPlugin/Source/GeometryCollectionNodes/Public/Dataflow/GeometryCollectionFracturingNodes.h
@@ -11,6 +11,7 @@
 #include "GeometryCollectionFracturingNodes.generated.h"
 
 class FGeometryCollection;
+class UDynamicMesh;
 
 USTRUCT(meta = (DataflowGeometryCollection, Deprecated = "5.5"))
 struct FUniformScatterPointsDataflowNode : public FDataflowNode
@@ -298,7 +299,8 @@ public:
 
 /**
  *
- * Generates a Voronoi fracture
+ * Editor Fracture Mode / Fracture / Uniform tool
+ * Fracture using a Voronoi diagram with a uniform random pattern, creating fracture pieces of similar volume across the shape.
  *
  */
 USTRUCT(meta = (DataflowGeometryCollection, Deprecated = "5.5"))
@@ -377,6 +379,12 @@ public:
 
 };
 
+/**
+ *
+ * Editor Fracture Mode / Fracture / Uniform tool
+ * Fracture using a Voronoi diagram with a uniform random pattern, creating fracture pieces of similar volume across the shape.
+ *
+ */
 USTRUCT(meta = (DataflowGeometryCollection))
 struct FVoronoiFractureDataflowNode_v2 : public FDataflowNode
 {
@@ -400,15 +408,16 @@ public:
 	UPROPERTY(meta = (DataflowInput, DisplayName = "Transform"))
 	FTransform Transform = FTransform::Identity;
 
+	// Made it hidden to hide this parameter, the random seed for Voronoi fracturing is on the point generation nodes
 	/** Seed for random */
-	UPROPERTY(EditAnywhere, Category = "Common Fracture", meta = (DataflowInput, UIMin = "0"));
+	UPROPERTY();
 	int32 RandomSeed = 0;
 
 	/** Chance to fracture each selected bone. If 0, no bones will fracture; if 1, all bones will fracture. */
 	UPROPERTY(EditAnywhere, Category = "Common Fracture", meta = (DataflowInput, UIMin = "0.0", UIMax = "1.0", ClampMin = "0.0", ClampMax = "1.0"));
 	float ChanceToFracture = 1.f;
 
-	// Whether to split the fractured mesh pieces based on geometric connectivity after fracturing
+	/** Whether to split the fractured mesh pieces based on geometric connectivity after fracturing */
 	UPROPERTY(EditAnywhere, Category = "Common Fracture");
 	bool SplitIslands = true;
 
@@ -470,7 +479,6 @@ public:
 		RegisterInputConnection(&Points);
 		RegisterInputConnection(&TransformSelection);
 		RegisterInputConnection(&Transform);
-		RegisterInputConnection(&RandomSeed);
 		RegisterInputConnection(&ChanceToFracture);
 		RegisterInputConnection(&Grout);
 		RegisterInputConnection(&Amplitude);
@@ -490,7 +498,8 @@ public:
 
 /**
  *
- * Cuts geometry using a set of noised up planes
+ * Editor Fracture Mode / Fracture / Planar tool
+ * Fracture using a set of noised up planes.
  *
  */
 USTRUCT(meta = (DataflowGeometryCollection, Deprecated = "5.5"))
@@ -566,6 +575,12 @@ public:
 	virtual void Evaluate(Dataflow::FContext& Context, const FDataflowOutput* Out) const override;
 };
 
+/**
+ *
+ * Editor Fracture Mode / Fracture / Planar tool
+ * Fracture using a set of noised up planes.
+ *
+ */
 USTRUCT(meta = (DataflowGeometryCollection))
 struct FPlaneCutterDataflowNode_v2 : public FDataflowNode
 {
@@ -601,7 +616,7 @@ public:
 	UPROPERTY(EditAnywhere, Category = "Common Fracture", meta = (DataflowInput, UIMin = 0.f, DisplayName = "Chance To Fracture Per Bone", UIMax = 1.f));
 	float ChanceToFracture = 1.f;
 
-	// Whether to split the fractured mesh pieces based on geometric connectivity after fracturing
+	/** Whether to split the fractured mesh pieces based on geometric connectivity after fracturing */
 	UPROPERTY(EditAnywhere, Category = "Common Fracture");
 	bool SplitIslands = true;
 
@@ -725,8 +740,9 @@ private:
 
 /**
  *
- * Slice tool
- *
+ * Editor Fracture Mode / Fracture / Slice tool
+ * Fracture with a grid of X, Y, and Z slices, with optional random variation in angle and offset.
+ * 
  */
 USTRUCT(meta = (DataflowGeometryCollection))
 struct FSliceCutterDataflowNode : public FDataflowNode
@@ -774,7 +790,7 @@ public:
 	UPROPERTY(EditAnywhere, Category = "Common Fracture", meta = (DataflowInput, UIMin = 0.f, UIMax = 1.f));
 	float ChanceToFracture = 1.f;
 
-	// Whether to split the fractured mesh pieces based on geometric connectivity after fracturing
+	/** Whether to split the fractured mesh pieces based on geometric connectivity after fracturing */
 	UPROPERTY(EditAnywhere, Category = "Common Fracture");
 	bool SplitIslands = true;
 
@@ -859,8 +875,10 @@ public:
 };
 
 /**
- *
- * Brick tool
+ * 
+ * Editor Fracture Mode / Fracture / Brick tool
+ * Fracture with a customizable brick pattern. 
+ * Note: Currently only supports fracturing with at least some (non-zero) Grout.
  *
  */
 USTRUCT(meta = (DataflowGeometryCollection))
@@ -910,7 +928,7 @@ public:
 	UPROPERTY(EditAnywhere, Category = "Common Fracture", meta = (DataflowInput, UIMin = 0.f, UIMax = 1.f));
 	float ChanceToFracture = 1.f;
 
-	// Whether to split the fractured mesh pieces based on geometric connectivity after fracturing
+	/** Whether to split the fractured mesh pieces based on geometric connectivity after fracturing */
 	UPROPERTY(EditAnywhere, Category = "Common Fracture");
 	bool SplitIslands = true;
 
@@ -993,38 +1011,147 @@ public:
 	virtual void Evaluate(Dataflow::FContext& Context, const FDataflowOutput* Out) const override;
 };
 
+/**
+ * 
+ * Editor Fracture Mode / Fracture / Mesh tool
+ * Fracture using the shape of a chosen static mesh
+ * 
+ */
+USTRUCT(meta = (DataflowGeometryCollection))
+struct FMeshCutterDataflowNode : public FDataflowNode
+{
+	GENERATED_USTRUCT_BODY()
+	DATAFLOW_NODE_DEFINE_INTERNAL(FMeshCutterDataflowNode, "MeshCutter", "GeometryCollection|Fracture", "")
 
+public:
+	/** Collection to cut */
+	UPROPERTY(meta = (DataflowInput, DataflowOutput, DataflowPassthrough = "Collection", DataflowIntrinsic))
+	FManagedArrayCollection Collection;
 
+	/** Boundingbox to create the cutting planes in */
+	UPROPERTY(meta = (DataflowInput))
+	FBox BoundingBox = FBox(ForceInit);
 
+	/** The selected pieces to cut */
+	UPROPERTY(meta = (DataflowInput, DataflowOutput, DataflowPassthrough = "TransformSelection", DisplayName = "TransformSelection"))
+	FDataflowTransformSelection TransformSelection;
 
+	/** Transform to apply to cut planes */
+	UPROPERTY(meta = (DataflowInput, DisplayName = "Transform"))
+	FTransform Transform = FTransform::Identity;
 
+	/** Mesh to cut with */
+	UPROPERTY(meta = (DataflowInput, DataflowIntrinsic))
+	TObjectPtr<UDynamicMesh> CuttingMesh;
 
+	/** How to arrange the mesh cuts in space */
+	UPROPERTY(EditAnywhere, Category = Distribution)
+	EMeshCutterCutDistribution CutDistribution = EMeshCutterCutDistribution::SingleCut;
 
+	/** Number of meshes to random scatter */
+	UPROPERTY(EditAnywhere, Category = Distribution, meta = (DataflowInput, ClampMin = "1", UIMax = "5000", EditCondition = "CutDistribution == EMeshCutterCutDistribution::UniformRandom", EditConditionHides))
+	int32 NumberToScatter = 10;
 
+	/** Number of meshes to add to grid in X */
+	UPROPERTY(EditAnywhere, Category = Distribution, meta = (DataflowInput, DisplayName = "Grid Width", ClampMin = "1", UIMax = "100", ClampMax = "5000", EditCondition = "CutDistribution == EMeshCutterCutDistribution::Grid", EditConditionHides))
+	int32 GridX = 2;
 
+	/** Number of meshes to add to grid in Y */
+	UPROPERTY(EditAnywhere, Category = Distribution, meta = (DataflowInput, DisplayName = "Grid Depth", ClampMin = "1", UIMax = "100", ClampMax = "5000", EditCondition = "CutDistribution == EMeshCutterCutDistribution::Grid", EditConditionHides))
+	int32 GridY = 2;
 
+	/** Number of meshes to add to grid in Z */
+	UPROPERTY(EditAnywhere, Category = Distribution, meta = (DataflowInput, DisplayName = "Grid Height", ClampMin = "1", UIMax = "100", ClampMax = "5000", EditCondition = "CutDistribution == EMeshCutterCutDistribution::Grid", EditConditionHides))
+	int32 GridZ = 2;
 
+	/** Magnitude of random displacement to cutting meshes */
+	UPROPERTY(EditAnywhere, Category = Distribution, meta = (DataflowInput, DisplayName = "Variability", EditCondition = "CutDistribution == EMeshCutterCutDistribution::Grid", EditConditionHides, UIMin = "0.0", ClampMin = "0.0"))
+	float Variability = 0.f;
 
+	/** Minimum scale factor to apply to cutting meshes. A random scale will be chosen between Min and Max */
+	UPROPERTY(EditAnywhere, Category = Distribution, meta = (DataflowInput, ClampMin = "0.001", EditCondition = "CutDistribution != EMeshCutterCutDistribution::SingleCut", EditConditionHides))
+	float MinScaleFactor = .5;
 
+	/** Maximum scale factor to apply to cutting meshes. A random scale will be chosen between Min and Max */
+	UPROPERTY(EditAnywhere, Category = Distribution, meta = (DataflowInput, ClampMin = "0.001", EditCondition = "CutDistribution != EMeshCutterCutDistribution::SingleCut", EditConditionHides))
+	float MaxScaleFactor = 1.5;
 
+	/** Whether to randomly vary the orientation of the cutting meshes */
+	UPROPERTY(EditAnywhere, Category = Distribution, meta = (EditCondition = "CutDistribution != EMeshCutterCutDistribution::SingleCut", EditConditionHides))
+	bool bRandomOrientation = true;
 
+	/** Roll will be chosen between -Range and +Range */
+	UPROPERTY(EditAnywhere, Category = Distribution, meta = (DataflowInput, DisplayName = "+/- Roll Range", EditCondition = "CutDistribution != EMeshCutterCutDistribution::SingleCut && bRandomOrientation", EditConditionHides, ClampMin = "0", ClampMax = "180"))
+	float RollRange = 180;
 
+	/** Pitch will be chosen between -Range and +Range */
+	UPROPERTY(EditAnywhere, Category = Distribution, meta = (DataflowInput, DisplayName = "+/- Pitch Range", EditCondition = "CutDistribution != EMeshCutterCutDistribution::SingleCut && bRandomOrientation", EditConditionHides, ClampMin = "0", ClampMax = "180"))
+	float PitchRange = 180;
 
+	/** Yaw will be chosen between -Range and +Range */
+	UPROPERTY(EditAnywhere, Category = Distribution, meta = (DataflowInput, DisplayName = "+/- Yaw Range", EditCondition = "CutDistribution != EMeshCutterCutDistribution::SingleCut && bRandomOrientation", EditConditionHides, ClampMin = "0", ClampMax = "180"))
+	float YawRange = 180;
 
+	/** Seed for random */
+	UPROPERTY(EditAnywhere, Category = "Common Fracture", meta = (DataflowInput, UIMin = "0"))
+	int32 RandomSeed = 0;
 
+	/** Chance to fracture each selected bone. If 0, no bones will fracture; if 1, all bones will fracture. */
+	UPROPERTY(EditAnywhere, Category = "Common Fracture", meta = (DataflowInput, UIMin = 0.f, UIMax = 1.f));
+	float ChanceToFracture = 1.f;
 
+	/** Whether to split the fractured mesh pieces based on geometric connectivity after fracturing */
+	UPROPERTY(EditAnywhere, Category = "Common Fracture");
+	bool SplitIslands = true;
+	
+	/**
+	 * If enabled, add extra vertices (without triangles) to the geometry in regions where vertices are spaced too far apart (e.g. across large triangles)
+	 * These extra vertices will be used as collision samples in particle-implicit collisions, and can help the physics system detect collisions more accurately
+	 *
+	 * Note this is *only* useful for simulations that use particle-implicit collisions
+	 */
+	UPROPERTY(EditAnywhere, Category = "Collision");
+	bool AddSamplesForCollision = false;
 
+	/**
+	 * The number of centimeters to allow between vertices on the mesh surface: If there are gaps larger than this, add additional vertices (without triangles) to help support particle-implicit collisions
+	 * Only used if Add Samples For Collision is enabled
+	 */
+	UPROPERTY(EditAnywhere, Category = "Collision", meta = (DataflowInput, DisplayName = "Point Spacing", UIMin = 0.f, EditCondition = "AddSamplesForCollision"));
+	float CollisionSampleSpacing = 50.f;
 
+	/** Fractured Pieces */
+	UPROPERTY(meta = (DataflowOutput, DisplayName = "FracturedTransformSelection "))
+	FDataflowTransformSelection NewGeometryTransformSelection;
 
+	FMeshCutterDataflowNode(const Dataflow::FNodeParameters& InParam, FGuid InGuid = FGuid::NewGuid())
+		: FDataflowNode(InParam, InGuid)
+	{
+		RegisterInputConnection(&Collection);
+		RegisterInputConnection(&BoundingBox);
+		RegisterInputConnection(&TransformSelection);
+		RegisterInputConnection(&Transform);
+		RegisterInputConnection(&CuttingMesh);
+		RegisterInputConnection(&NumberToScatter);
+		RegisterInputConnection(&GridX);
+		RegisterInputConnection(&GridY);
+		RegisterInputConnection(&GridZ);
+		RegisterInputConnection(&Variability);
+		RegisterInputConnection(&MinScaleFactor);
+		RegisterInputConnection(&MaxScaleFactor);
+		RegisterInputConnection(&RollRange);
+		RegisterInputConnection(&PitchRange);
+		RegisterInputConnection(&YawRange);
+		RegisterInputConnection(&RandomSeed);
+		RegisterInputConnection(&ChanceToFracture);
+		RegisterInputConnection(&CollisionSampleSpacing);
+		RegisterOutputConnection(&Collection, &Collection);
+		RegisterOutputConnection(&TransformSelection, &TransformSelection);
+		RegisterOutputConnection(&NewGeometryTransformSelection);
+	}
 
-
-
-
-
-
-
-
-
+	virtual void Evaluate(Dataflow::FContext& Context, const FDataflowOutput* Out) const override;
+};
 
 namespace Dataflow
 {
diff --git a/Engine/Plugins/Experimental/GeometryCollectionPlugin/Source/GeometryCollectionNodes/Public/Dataflow/GeometryCollectionMakeNodes.h b/Engine/Plugins/Experimental/GeometryCollectionPlugin/Source/GeometryCollectionNodes/Public/Dataflow/GeometryCollectionMakeNodes.h
--- a/Engine/Plugins/Experimental/GeometryCollectionPlugin/Source/GeometryCollectionNodes/Public/Dataflow/GeometryCollectionMakeNodes.h
+++ b/Engine/Plugins/Experimental/GeometryCollectionPlugin/Source/GeometryCollectionNodes/Public/Dataflow/GeometryCollectionMakeNodes.h
@@ -362,6 +362,41 @@ public:
 
 };
 
+/**
+ *
+ * Description for this node
+ *
+ */
+USTRUCT()
+struct FMakeFloatArrayDataflowNode : public FDataflowNode
+{
+	GENERATED_USTRUCT_BODY()
+	DATAFLOW_NODE_DEFINE_INTERNAL(FMakeFloatArrayDataflowNode, "MakeFloatArray", "Math|Float", "")
+
+public:
+	/** Number of elements of the array */
+	UPROPERTY(EditAnywhere, Category = "Float", meta = (DataflowInput, DisplayName = "Number of Elements", UIMin = "0"));
+	int32 NumElements = 1;
+
+	/** Value to initialize the array with */
+	UPROPERTY(EditAnywhere, Category = "Float", meta = (DataflowInput));
+	float Value = 0.f;
+
+	/** Output float array */
+	UPROPERTY(meta = (DataflowOutput))
+	TArray<float> FloatArray;
+
+	FMakeFloatArrayDataflowNode(const Dataflow::FNodeParameters& InParam, FGuid InGuid = FGuid::NewGuid())
+		: FDataflowNode(InParam, InGuid)
+	{
+		RegisterInputConnection(&NumElements);
+		RegisterInputConnection(&Value);
+		RegisterOutputConnection(&FloatArray);
+	}
+
+	virtual void Evaluate(Dataflow::FContext& Context, const FDataflowOutput* Out) const override;
+
+};
 
 namespace Dataflow
 {
diff --git a/Engine/Plugins/Experimental/GeometryCollectionPlugin/Source/GeometryCollectionNodes/Public/Dataflow/GeometryCollectionSamplingNodes.h b/Engine/Plugins/Experimental/GeometryCollectionPlugin/Source/GeometryCollectionNodes/Public/Dataflow/GeometryCollectionSamplingNodes.h
new file mode 100644
--- /dev/null
+++ b/Engine/Plugins/Experimental/GeometryCollectionPlugin/Source/GeometryCollectionNodes/Public/Dataflow/GeometryCollectionSamplingNodes.h
@@ -0,0 +1,262 @@
+// Copyright Epic Games, Inc. All Rights Reserved.
+
+#pragma once 
+
+#include "Dataflow/DataflowEngine.h"
+#include "FractureEngineSampling.h"
+
+#include "GeometryCollectionSamplingNodes.generated.h"
+
+class FGeometryCollection;
+class UDynamicMesh;
+
+/**
+ *
+ * Uniform Sampling on a DynamicMesh
+ *
+ */
+USTRUCT(meta = (DataflowGeometryCollection))
+struct FUniformPointSamplingDataflowNode : public FDataflowNode
+{
+	GENERATED_USTRUCT_BODY()
+	DATAFLOW_NODE_DEFINE_INTERNAL(FUniformPointSamplingDataflowNode, "UniformPointSampling", "PointSampling", "")
+
+public:
+	/** Mesh to sample points on */
+	UPROPERTY(meta = (DataflowInput, DataflowIntrinsic))
+	TObjectPtr<UDynamicMesh> TargetMesh;
+
+	/** Desired "radius" of sample points. Spacing between samples is at least 2x this value. */
+	UPROPERTY(EditAnywhere, Category = Options, meta = (DataflowInput, UIMin = 0.f))
+	float SamplingRadius = 10.f;
+
+	/** Maximum number of samples requested. If 0 or default value, mesh will be maximally sampled */
+	UPROPERTY(EditAnywhere, Category = Options, meta = (DataflowInput, UIMin = 0))
+	int32 MaxNumSamples = 0;
+
+	/** Density of subsampling used in Poisson strategy. Larger numbers mean "more accurate" (but slower) results. */
+	UPROPERTY(EditAnywhere, Category = Options, meta = (DataflowInput, UIMin = 0.f))
+	float SubSampleDensity = 10.f;
+
+	/** Random Seed used to initialize sampling strategies */
+	UPROPERTY(EditAnywhere, Category = Options, meta = (DataflowInput))
+	int32 RandomSeed = 0;
+
+	/** Sampled positions on the mesh */
+	UPROPERTY(meta = (DataflowOutput))
+	TArray<FVector> SamplePoints;
+
+	/** Sampled triangleID */
+	UPROPERTY(meta = (DataflowOutput))
+	TArray<int32> SampleTriangleIDs;
+
+	/** Barycentric Coordinates of each Sample Point in it's respective Triangle. */
+	UPROPERTY(meta = (DataflowOutput))
+	TArray<FVector> SampleBarycentricCoords;
+
+	/** Number of Sampled positions on the mesh */
+	UPROPERTY(meta = (DataflowOutput))
+	int32 NumSamplePoints = 0;
+
+	FUniformPointSamplingDataflowNode(const Dataflow::FNodeParameters& InParam, FGuid InGuid = FGuid::NewGuid())
+		: FDataflowNode(InParam, InGuid)
+	{
+		RegisterInputConnection(&TargetMesh);
+		RegisterInputConnection(&SamplingRadius);
+		RegisterInputConnection(&MaxNumSamples);
+		RegisterInputConnection(&SubSampleDensity);
+		RegisterInputConnection(&RandomSeed);
+		RegisterOutputConnection(&SamplePoints);
+		RegisterOutputConnection(&SampleTriangleIDs);
+		RegisterOutputConnection(&SampleBarycentricCoords);
+		RegisterOutputConnection(&NumSamplePoints);
+	}
+
+	virtual void Evaluate(Dataflow::FContext& Context, const FDataflowOutput* Out) const override;
+};
+
+/**
+ *
+ * NonUniform Sampling on a DynamicMesh
+ *
+ */
+USTRUCT(meta = (DataflowGeometryCollection))
+struct FNonUniformPointSamplingDataflowNode : public FDataflowNode
+{
+	GENERATED_USTRUCT_BODY()
+	DATAFLOW_NODE_DEFINE_INTERNAL(FNonUniformPointSamplingDataflowNode, "NonUniformPointSampling", "PointSampling", "")
+
+public:
+	/** Mesh to sample points on */
+	UPROPERTY(meta = (DataflowInput, DataflowIntrinsic))
+	TObjectPtr<UDynamicMesh> TargetMesh;
+
+	/** Desired "radius" of sample points. Spacing between samples is at least 2x this value. */
+	UPROPERTY(EditAnywhere, Category = Options, meta = (DataflowInput, UIMin = 0.f))
+	float SamplingRadius = 10.f;
+
+	/** Maximum number of samples requested. If 0 or default value, mesh will be maximally sampled */
+	UPROPERTY(EditAnywhere, Category = Options, meta = (DataflowInput, UIMin = 0))
+	int32 MaxNumSamples = 0;
+
+	/** Density of subsampling used in Poisson strategy. Larger numbers mean "more accurate" (but slower) results. */
+	UPROPERTY(EditAnywhere, Category = Options, meta = (DataflowInput, UIMin = 0.f))
+	float SubSampleDensity = 10.f;
+
+	/** Random Seed used to initialize sampling strategies */
+	UPROPERTY(EditAnywhere, Category = Options, meta = (DataflowInput))
+	int32 RandomSeed = 0;
+
+	/** If MaxSampleRadius > SampleRadius, then output sample radius will be in range [SampleRadius, MaxSampleRadius] */
+	UPROPERTY(EditAnywhere, Category = Options, meta = (DataflowInput, UIMin = 0.f))
+	float MaxSamplingRadius = 10.f;
+
+	/** SizeDistribution setting controls the distribution of sample radii */
+	UPROPERTY(EditAnywhere, Category = Distribution)
+	ENonUniformSamplingDistributionMode SizeDistribution = ENonUniformSamplingDistributionMode::ENonUniformSamplingDistributionMode_Uniform;
+
+	/** SizeDistributionPower is used to control how extreme the Size Distribution shift is. Valid range is [1,10] */
+	UPROPERTY(EditAnywhere, Category = Options, meta = (DataflowInput, UIMin = 0.f))
+	float SizeDistributionPower = 2.f;
+
+	/** Sampled positions on the mesh */
+	UPROPERTY(meta = (DataflowOutput))
+	TArray<FVector> SamplePoints;
+
+	/** Sampled radii */
+	UPROPERTY(meta = (DataflowOutput))
+	TArray<float> SampleRadii;
+
+	/** Sampled triangleID */
+	UPROPERTY(meta = (DataflowOutput))
+	TArray<int32> SampleTriangleIDs;
+
+	/** Barycentric Coordinates of each Sample Point in it's respective Triangle. */
+	UPROPERTY(meta = (DataflowOutput))
+	TArray<FVector> SampleBarycentricCoords;
+
+	/** Number of Sampled positions on the mesh */
+	UPROPERTY(meta = (DataflowOutput))
+	int32 NumSamplePoints = 0;
+
+	FNonUniformPointSamplingDataflowNode(const Dataflow::FNodeParameters& InParam, FGuid InGuid = FGuid::NewGuid())
+		: FDataflowNode(InParam, InGuid)
+	{
+		RegisterInputConnection(&TargetMesh);
+		RegisterInputConnection(&SamplingRadius);
+		RegisterInputConnection(&MaxNumSamples);
+		RegisterInputConnection(&SubSampleDensity);
+		RegisterInputConnection(&RandomSeed);
+		RegisterInputConnection(&MaxSamplingRadius);
+		RegisterInputConnection(&SizeDistributionPower);
+		RegisterOutputConnection(&SamplePoints);
+		RegisterOutputConnection(&SampleRadii);
+		RegisterOutputConnection(&SampleTriangleIDs);
+		RegisterOutputConnection(&SampleBarycentricCoords);
+		RegisterOutputConnection(&NumSamplePoints);
+	}
+
+	virtual void Evaluate(Dataflow::FContext& Context, const FDataflowOutput* Out) const override;
+};
+
+/**
+ *
+ * VertexWeighted Sampling on a DynamicMesh
+ *
+ */
+USTRUCT(meta = (DataflowGeometryCollection))
+struct FVertexWeightedPointSamplingDataflowNode : public FDataflowNode
+{
+	GENERATED_USTRUCT_BODY()
+	DATAFLOW_NODE_DEFINE_INTERNAL(FVertexWeightedPointSamplingDataflowNode, "VertexWeightedPointSampling", "PointSampling", "")
+
+public:
+	/** Mesh to sample points on */
+	UPROPERTY(meta = (DataflowInput, DataflowIntrinsic))
+	TObjectPtr<UDynamicMesh> TargetMesh;
+
+	/** Weight array */
+	UPROPERTY(meta = (DataflowInput, DataflowIntrinsic))
+	TArray<float> VertexWeights;
+
+	/** Desired "radius" of sample points. Spacing between samples is at least 2x this value. */
+	UPROPERTY(EditAnywhere, Category = Options, meta = (DataflowInput, UIMin = 0.f))
+	float SamplingRadius = 10.f;
+
+	/** Maximum number of samples requested. If 0 or default value, mesh will be maximally sampled */
+	UPROPERTY(EditAnywhere, Category = Options, meta = (DataflowInput, UIMin = 0))
+	int32 MaxNumSamples = 0;
+
+	/** Density of subsampling used in Poisson strategy. Larger numbers mean "more accurate" (but slower) results. */
+	UPROPERTY(EditAnywhere, Category = Options, meta = (DataflowInput, UIMin = 0.f))
+	float SubSampleDensity = 10.f;
+
+	/** Random Seed used to initialize sampling strategies */
+	UPROPERTY(EditAnywhere, Category = Options, meta = (DataflowInput))
+	int32 RandomSeed = 0;
+
+	/** If MaxSampleRadius > SampleRadius, then output sample radius will be in range [SampleRadius, MaxSampleRadius] */
+	UPROPERTY(EditAnywhere, Category = Options, meta = (DataflowInput, UIMin = 0.f))
+	float MaxSamplingRadius = 10.f;
+
+	/** SizeDistribution setting controls the distribution of sample radii */
+	UPROPERTY(EditAnywhere, Category = Distribution)
+	ENonUniformSamplingDistributionMode SizeDistribution = ENonUniformSamplingDistributionMode::ENonUniformSamplingDistributionMode_Uniform;
+
+	/** SizeDistributionPower is used to control how extreme the Size Distribution shift is. Valid range is [1,10] */
+	UPROPERTY(EditAnywhere, Category = Options, meta = (DataflowInput, UIMin = 0.f))
+	float SizeDistributionPower = 2.f;
+
+	UPROPERTY(EditAnywhere, Category = Distribution)
+	ENonUniformSamplingWeightMode WeightMode = ENonUniformSamplingWeightMode::ENonUniformSamplingWeightMode_WeightedRandom;
+
+	UPROPERTY(EditAnywhere, Category = Options)
+	bool bInvertWeights = false;
+
+	/** Sampled positions on the mesh */
+	UPROPERTY(meta = (DataflowOutput))
+	TArray<FVector> SamplePoints;
+
+	/** Sampled radii */
+	UPROPERTY(meta = (DataflowOutput))
+	TArray<float> SampleRadii;
+
+	/** Sampled triangleID */
+	UPROPERTY(meta = (DataflowOutput))
+	TArray<int32> SampleTriangleIDs;
+
+	/** Barycentric Coordinates of each Sample Point in it's respective Triangle. */
+	UPROPERTY(meta = (DataflowOutput))
+	TArray<FVector> SampleBarycentricCoords;
+
+	/** Number of Sampled positions on the mesh */
+	UPROPERTY(meta = (DataflowOutput))
+	int32 NumSamplePoints = 0;
+
+	FVertexWeightedPointSamplingDataflowNode(const Dataflow::FNodeParameters& InParam, FGuid InGuid = FGuid::NewGuid())
+		: FDataflowNode(InParam, InGuid)
+	{
+		RegisterInputConnection(&TargetMesh);
+		RegisterInputConnection(&VertexWeights);
+		RegisterInputConnection(&SamplingRadius);
+		RegisterInputConnection(&MaxNumSamples);
+		RegisterInputConnection(&SubSampleDensity);
+		RegisterInputConnection(&RandomSeed);
+		RegisterInputConnection(&MaxSamplingRadius);
+		RegisterInputConnection(&SizeDistributionPower);
+		RegisterOutputConnection(&SamplePoints);
+		RegisterOutputConnection(&SampleRadii);
+		RegisterOutputConnection(&SampleTriangleIDs);
+		RegisterOutputConnection(&SampleBarycentricCoords);
+		RegisterOutputConnection(&NumSamplePoints);
+	}
+
+	virtual void Evaluate(Dataflow::FContext& Context, const FDataflowOutput* Out) const override;
+};
+
+
+namespace Dataflow
+{
+	void GeometryCollectionSamplingNodes();
+}
+
diff --git a/Engine/Plugins/Experimental/GeometryCollectionPlugin/Source/GeometryCollectionNodes/Public/Dataflow/GeometryCollectionSelectionNodes.h b/Engine/Plugins/Experimental/GeometryCollectionPlugin/Source/GeometryCollectionNodes/Public/Dataflow/GeometryCollectionSelectionNodes.h
--- a/Engine/Plugins/Experimental/GeometryCollectionPlugin/Source/GeometryCollectionNodes/Public/Dataflow/GeometryCollectionSelectionNodes.h
+++ b/Engine/Plugins/Experimental/GeometryCollectionPlugin/Source/GeometryCollectionNodes/Public/Dataflow/GeometryCollectionSelectionNodes.h
@@ -1328,6 +1328,110 @@ public:
 
 };
 
+UENUM(BlueprintType)
+enum class ESelectionByAttrGroup : uint8
+{
+	Vertices UMETA(DisplayName = "Vertices"),
+	Faces UMETA(DisplayName = "Faces"),
+	Transform UMETA(DisplayName = "Transform"),
+	Geometry UMETA(DisplayName = "Geometry"),
+	Material UMETA(DisplayName = "Material")
+};
+
+namespace Dataflow::Private
+{
+	inline FName GetAttributeFromEnumAsName(const ESelectionByAttrGroup Value)
+	{
+		static const UEnum* SelectionByAttrGroupEnum = StaticEnum<ESelectionByAttrGroup>();
+		return *SelectionByAttrGroupEnum->GetNameStringByValue((int64)Value);
+	}
+}
+
+UENUM(BlueprintType)
+enum class ESelectionByAttrOperation : uint8
+{
+	/** Select faces which attribute value equal with specified value */
+	Equal UMETA(DisplayName = "=="),
+	/** Select faces which attribute value not equal with specified value */
+	NotEqual UMETA(DisplayName = "!="),
+	/** Select faces which attribute value greater than specified value */
+	Greater UMETA(DisplayName = ">"),
+	/** Select faces which attribute value greater or equal than specified value */
+	GreaterOrEqual UMETA(DisplayName = ">="),
+	/** Select faces which attribute value smaller than specified value */
+	Smaller UMETA(DisplayName = "<"),
+	/** Select faces which attribute value greater than specified value */
+	SmallerOrEqual UMETA(DisplayName = "<=")
+};
+
+/**
+ *
+ * Selects specified Vertices/Faces/Transforms in the GeometryCollection by using an attribute value
+ * Currently supported attribute types: float, int32, String, bool
+ *
+ */
+USTRUCT(meta = (DataflowGeometryCollection))
+struct FCollectionSelectionByAttrDataflowNode : public FDataflowNode
+{
+	GENERATED_USTRUCT_BODY()
+	DATAFLOW_NODE_DEFINE_INTERNAL(FCollectionSelectionByAttrDataflowNode, "CollectionSelectByAttr", "GeometryCollection|Selection|All", "")
+
+public:
+	/** GeometryCollection for the selection */
+	UPROPERTY(meta = (DataflowInput, DataflowOutput, DataflowPassthrough = "Collection", DataflowIntrinsic))
+	FManagedArrayCollection Collection;
+
+	/** Group */
+	UPROPERTY(EditAnywhere, Category = "Selection")
+	ESelectionByAttrGroup Group = ESelectionByAttrGroup::Faces;
+
+	/** Attribute for the selection */
+	UPROPERTY(EditAnywhere, Category = "Selection")
+	FString Attribute = FString("Internal");
+
+	/** Operation */
+	UPROPERTY(EditAnywhere, Category = "Selection")
+	ESelectionByAttrOperation Operation = ESelectionByAttrOperation::Equal;
+
+	/** Attribute value for the operation */
+	UPROPERTY(EditAnywhere, Category = "Selection")
+	FString Value = FString("true");
+
+	/** Vertex selection output */
+	UPROPERTY(meta = (DataflowOutput, DisplayName = "VertexSelection"))
+	FDataflowVertexSelection VertexSelection;
+
+	/** Face selection output */
+	UPROPERTY(meta = (DataflowOutput, DisplayName = "FaceSelection"))
+	FDataflowFaceSelection FaceSelection;
+
+	/** Transform selection output */
+	UPROPERTY(meta = (DataflowOutput, DisplayName = "TransformSelection"))
+	FDataflowTransformSelection TransformSelection;
+
+	/** Geometry selection output */
+	UPROPERTY(meta = (DataflowOutput, DisplayName = "GeometrySelection"))
+	FDataflowGeometrySelection GeometrySelection;
+
+	/** Material selection output */
+	UPROPERTY(meta = (DataflowOutput, DisplayName = "MaterialSelection"))
+	FDataflowMaterialSelection MaterialSelection;
+
+	FCollectionSelectionByAttrDataflowNode(const Dataflow::FNodeParameters& InParam, FGuid InGuid = FGuid::NewGuid())
+		: FDataflowNode(InParam, InGuid)
+	{
+		RegisterInputConnection(&Collection);
+		RegisterOutputConnection(&Collection, &Collection);
+		RegisterOutputConnection(&VertexSelection);
+		RegisterOutputConnection(&FaceSelection);
+		RegisterOutputConnection(&TransformSelection);
+		RegisterOutputConnection(&GeometrySelection);
+		RegisterOutputConnection(&MaterialSelection);
+	}
+
+	virtual void Evaluate(Dataflow::FContext& Context, const FDataflowOutput* Out) const override;
+};
+
 namespace Dataflow
 {
 	void GeometryCollectionSelectionNodes();
diff --git a/Engine/Plugins/Experimental/GeometryCollectionPlugin/Source/GeometryCollectionNodes/Public/Dataflow/GeometryCollectionUtilityNodes.h b/Engine/Plugins/Experimental/GeometryCollectionPlugin/Source/GeometryCollectionNodes/Public/Dataflow/GeometryCollectionUtilityNodes.h
--- a/Engine/Plugins/Experimental/GeometryCollectionPlugin/Source/GeometryCollectionNodes/Public/Dataflow/GeometryCollectionUtilityNodes.h
+++ b/Engine/Plugins/Experimental/GeometryCollectionPlugin/Source/GeometryCollectionNodes/Public/Dataflow/GeometryCollectionUtilityNodes.h
@@ -10,6 +10,7 @@
 #include "CompGeom/ConvexDecomposition3.h"
 
 #include "FractureEngineConvex.h"
+#include "FractureEngineUtility.h"
 
 #include "GeometryCollectionUtilityNodes.generated.h"
 
@@ -562,6 +563,233 @@ public:
 	virtual void Evaluate(Dataflow::FContext& Context, const FDataflowOutput* Out) const override;
 };
 
+/**
+ * 
+ * Editor Fracture Mode / Utilities / TinyGeo tool
+ * Merge pieces of geometry onto their neighbors -- use it to, for example, clean up too small pieces of geometry.
+ * 
+ */
+USTRUCT(meta = (DataflowGeometryCollection))
+struct FFixTinyGeoDataflowNode : public FDataflowNode
+{
+	GENERATED_USTRUCT_BODY()
+	DATAFLOW_NODE_DEFINE_INTERNAL(FFixTinyGeoDataflowNode, "FixTinyGeo", "GeometryCollection|Fracture|Utilities", "")
+
+public:
+	/** Collection to use */
+	UPROPERTY(meta = (DataflowInput, DataflowOutput, DataflowPassthrough = "Collection", DataflowIntrinsic))
+	FManagedArrayCollection Collection;
+
+	/** The selected pieces to use */
+	UPROPERTY(meta = (DataflowInput, DisplayName = "TransformSelection"))
+	FDataflowTransformSelection TransformSelection;
+
+	/** Whether to merge small geometry, or small clusters */
+	UPROPERTY(EditAnywhere, Category = Distribution)
+	EFixTinyGeoMergeType MergeType = EFixTinyGeoMergeType::MergeGeometry;
+
+	/** Only consider bones at the current Fracture Level */
+	UPROPERTY(EditAnywhere, Category = MergeSettings, meta = (EditCondition = "!bFractureLevelIsAll && MergeType == EFixTinyGeoMergeType::MergeClusters", EditConditionHides))
+	bool bOnFractureLevel = true;
+
+	/** Only auto - consider clusters for merging.Note that leaf nodes can still be consider if manually selected. */
+	UPROPERTY(EditAnywhere, Category = MergeSettings, meta = (EditCondition = "!bFractureLevelIsAll && MergeType == EFixTinyGeoMergeType::MergeClusters && bOnFractureLevel", EditConditionHides))
+	bool bOnlyClusters = false;
+
+	/** Only merge to neighbors with the same parent in the hierarchy */
+	UPROPERTY(EditAnywhere, Category = MergeSettings, meta = (EditCondition = "MergeType == EFixTinyGeoMergeType::MergeClusters", EditConditionHides))
+	bool bOnlySameParent = true;
+
+	/** Helper variable to let the EditConditions above check whether the Fracture Level is set to 'All' */
+	UPROPERTY()
+	bool bFractureLevelIsAll = false;
+
+	UPROPERTY(EditAnywhere, Category = MergeSettings, meta = (DisplayName = "Merge To"))
+	EFixTinyGeoNeighborSelectionMethod NeighborSelection = EFixTinyGeoNeighborSelectionMethod::LargestNeighbor;
+
+	/** Only merge pieces that are connected in the proximity graph.If unchecked, connected pieces will still be favored, but if none are available the closest disconnected piece can be merged. */
+	UPROPERTY(EditAnywhere, Category = MergeSettings, meta = (EditCondition = "MergeType == EFixTinyGeoMergeType::MergeClusters", EditConditionHides))
+	bool bOnlyToConnected = true;
+
+	/** Options for using the current bone selection */
+	UPROPERTY(EditAnywhere, Category = MergeSettings, meta = (DisplayName = "Bone Selection"))
+	EFixTinyGeoUseBoneSelection UseBoneSelection = EFixTinyGeoUseBoneSelection::NoEffect;
+
+	UPROPERTY(EditAnywhere, Category = FilterSettings)
+	EFixTinyGeoGeometrySelectionMethod SelectionMethod = EFixTinyGeoGeometrySelectionMethod::RelativeVolume;
+
+	/** If size (cube root of volume) is less than this value, geometry should be merged into neighbors -- i.e. a value of 2 merges geometry smaller than a 2x2x2 cube */
+	UPROPERTY(EditAnywhere, Category = FilterSettings, meta = (DisplayName = "MinSize", ClampMin = ".00001", UIMin = ".1", UIMax = "10", EditCondition = "SelectionMethod == EFixTinyGeoGeometrySelectionMethod::VolumeCubeRoot", EditConditionHides))
+	float MinVolumeCubeRoot = 1.f;
+
+	/** If cube root of volume relative to the overall shape's cube root of volume is less than this, the geometry should be merged into its neighbors.
+	(Note: This is a bit different from the histogram viewer's "Relative Size," which instead shows values relative to the largest rigid bone.) */
+	UPROPERTY(EditAnywhere, Category = FilterSettings, meta = (ClampMin = "0", UIMax = ".1", ClampMax = "1.0", EditCondition = "SelectionMethod == EFixTinyGeoGeometrySelectionMethod::RelativeVolume", EditConditionHides))
+	float RelativeVolume = .01f;
+
+	/**
+	 * If enabled, add extra vertices (without triangles) to the geometry in regions where vertices are spaced too far apart (e.g. across large triangles)
+	 * These extra vertices will be used as collision samples in particle-implicit collisions, and can help the physics system detect collisions more accurately
+	 *
+	 * Note this is *only* useful for simulations that use particle-implicit collisions
+	 */
+	UPROPERTY(EditAnywhere, Category = "Collision");
+	bool AddSamplesForCollision = false;
+
+	/**
+	 * The number of centimeters to allow between vertices on the mesh surface: If there are gaps larger than this, add additional vertices (without triangles) to help support particle-implicit collisions
+	 * Only used if Add Samples For Collision is enabled
+	 */
+	UPROPERTY(EditAnywhere, Category = "Collision", meta = (DataflowInput, DisplayName = "Point Spacing", UIMin = 0.f, EditCondition = "AddSamplesForCollision"));
+	float CollisionSampleSpacing = 50.f;
+
+	FFixTinyGeoDataflowNode(const Dataflow::FNodeParameters& InParam, FGuid InGuid = FGuid::NewGuid())
+		: FDataflowNode(InParam, InGuid)
+	{
+		RegisterInputConnection(&Collection);
+		RegisterInputConnection(&TransformSelection);
+		RegisterInputConnection(&CollisionSampleSpacing);
+		RegisterOutputConnection(&Collection, &Collection);
+	}
+
+	virtual void Evaluate(Dataflow::FContext& Context, const FDataflowOutput* Out) const override;
+};
+
+/**
+ *
+ * Editor Fracture Mode / Utilities / Normals tool
+ * Recompute normals and tangents.
+ */
+USTRUCT(meta = (DataflowGeometryCollection))
+struct FRecomputeNormalsInGeometryCollectionDataflowNode : public FDataflowNode
+{
+	GENERATED_USTRUCT_BODY()
+	DATAFLOW_NODE_DEFINE_INTERNAL(FRecomputeNormalsInGeometryCollectionDataflowNode, "RecomputeNormalsInGeometryCollection", "GeometryCollection|Fracture|Utilities", "")
+
+public:
+	/** Collection to use */
+	UPROPERTY(meta = (DataflowInput, DataflowOutput, DataflowPassthrough = "Collection", DataflowIntrinsic))
+	FManagedArrayCollection Collection;
+
+	/** The selected pieces to use */
+	UPROPERTY(meta = (DataflowInput, DisplayName = "TransformSelection"))
+	FDataflowTransformSelection TransformSelection;
+
+	/** Whether to only recompute tangents, and leave normals as they were */
+	UPROPERTY(EditAnywhere, Category = RecomputeSettings)
+	bool bOnlyTangents = false;
+
+	/** If true, update where edges are 'sharp' by comparing adjacent triangle face normals vs the Sharp Edge Angle Threshold. */
+	UPROPERTY(EditAnywhere, Category = RecomputeSettings, meta = (EditCondition = "!bOnlyTangents"))
+	bool bRecomputeSharpEdges = false;
+
+	/** Threshold on angle of change in face normals across an edge, above which we create a sharp edge if bRecomputeSharpEdges is true */
+	UPROPERTY(EditAnywhere, Category = RecomputeSettings, meta = (UIMin = "0.0", UIMax = "180.0", ClampMin = "0.0", ClampMax = "180.0", EditCondition = "bRecomputeSharpEdges && !bOnlyTangents"))
+	float SharpEdgeAngleThreshold = 60.0f;
+
+	/** Whether to only change internal surface normals / tangents */
+	UPROPERTY(EditAnywhere, Category = RecomputeSettings, AdvancedDisplay)
+	bool bOnlyInternalSurfaces = true;
+
+	FRecomputeNormalsInGeometryCollectionDataflowNode(const Dataflow::FNodeParameters& InParam, FGuid InGuid = FGuid::NewGuid())
+		: FDataflowNode(InParam, InGuid)
+	{
+		RegisterInputConnection(&Collection);
+		RegisterInputConnection(&TransformSelection);
+		RegisterOutputConnection(&Collection, &Collection);
+	}
+
+	virtual void Evaluate(Dataflow::FContext& Context, const FDataflowOutput* Out) const override;
+};
+
+/**
+ *
+ * Editor Fracture Mode / Utilities / Resample tool
+ * Resample to add collision particles in large flat regions that otherwise might have poor collision response.
+ * Only useful to help improve Particle - Implicit collisions.
+ * 
+ */
+USTRUCT(meta = (DataflowGeometryCollection))
+struct FResampleGeometryCollectionDataflowNode : public FDataflowNode
+{
+	GENERATED_USTRUCT_BODY()
+	DATAFLOW_NODE_DEFINE_INTERNAL(FResampleGeometryCollectionDataflowNode, "ResampleGeometryCollection", "GeometryCollection|Fracture|Utilities", "")
+
+public:
+	/** Collection to use */
+	UPROPERTY(meta = (DataflowInput, DataflowOutput, DataflowPassthrough = "Collection", DataflowIntrinsic))
+	FManagedArrayCollection Collection;
+
+	/** The selected pieces to use */
+	UPROPERTY(meta = (DataflowInput, DisplayName = "TransformSelection"))
+	FDataflowTransformSelection TransformSelection;
+
+	/**
+	 * If enabled, add extra vertices (without triangles) to the geometry in regions where vertices are spaced too far apart (e.g. across large triangles)
+	 * These extra vertices will be used as collision samples in particle-implicit collisions, and can help the physics system detect collisions more accurately
+	 *
+	 * Note this is *only* useful for simulations that use particle-implicit collisions
+	 */
+	UPROPERTY(EditAnywhere, Category = "Collision");
+	bool AddSamplesForCollision = false;
+
+	/**
+	 * The number of centimeters to allow between vertices on the mesh surface: If there are gaps larger than this, add additional vertices (without triangles) to help support particle-implicit collisions
+	 * Only used if Add Samples For Collision is enabled
+	 */
+	UPROPERTY(EditAnywhere, Category = "Collision", meta = (DataflowInput, DisplayName = "Point Spacing", UIMin = 0.f, EditCondition = "AddSamplesForCollision"));
+	float CollisionSampleSpacing = 50.f;
+
+	FResampleGeometryCollectionDataflowNode(const Dataflow::FNodeParameters& InParam, FGuid InGuid = FGuid::NewGuid())
+		: FDataflowNode(InParam, InGuid)
+	{
+		RegisterInputConnection(&Collection);
+		RegisterInputConnection(&TransformSelection);
+		RegisterInputConnection(&CollisionSampleSpacing);
+		RegisterOutputConnection(&Collection, &Collection);
+	}
+
+	virtual void Evaluate(Dataflow::FContext& Context, const FDataflowOutput* Out) const override;
+};
+
+/**
+ *
+ * Editor Fracture Mode / Utilities / Validate tool
+ * Ensures that geometrycollection is valid and clean.
+ * 
+ */
+USTRUCT(meta = (DataflowGeometryCollection))
+struct FValidateGeometryCollectionDataflowNode : public FDataflowNode
+{
+	GENERATED_USTRUCT_BODY()
+	DATAFLOW_NODE_DEFINE_INTERNAL(FValidateGeometryCollectionDataflowNode, "ValidateGeometryCollection", "GeometryCollection|Fracture|Utilities", "")
+
+public:
+	/** Collection to use */
+	UPROPERTY(meta = (DataflowInput, DataflowOutput, DataflowPassthrough = "Collection", DataflowIntrinsic))
+	FManagedArrayCollection Collection;
+
+	/** Find and remove any unused geometry data */
+	UPROPERTY(EditAnywhere, Category = CleanUnused)
+	bool bRemoveUnreferencedGeometry = true;
+
+	/** Whether to collapse any clusters with only a single child */
+	UPROPERTY(EditAnywhere, Category = Clustering)
+	bool bRemoveClustersOfOne = false;
+
+	/** Remove dangling clusters -- Note this can invalidate caches */
+	UPROPERTY(EditAnywhere, Category = Clustering)
+	bool bRemoveDanglingClusters = false;
+
+	FValidateGeometryCollectionDataflowNode(const Dataflow::FNodeParameters& InParam, FGuid InGuid = FGuid::NewGuid())
+		: FDataflowNode(InParam, InGuid)
+	{
+		RegisterInputConnection(&Collection);
+		RegisterOutputConnection(&Collection, &Collection);
+	}
+
+	virtual void Evaluate(Dataflow::FContext& Context, const FDataflowOutput* Out) const override;
+};
 
 namespace Dataflow
 {
diff --git a/Engine/Source/Runtime/Experimental/Dataflow/Core/Public/Dataflow/DataflowSelection.h b/Engine/Source/Runtime/Experimental/Dataflow/Core/Public/Dataflow/DataflowSelection.h
--- a/Engine/Source/Runtime/Experimental/Dataflow/Core/Public/Dataflow/DataflowSelection.h
+++ b/Engine/Source/Runtime/Experimental/Dataflow/Core/Public/Dataflow/DataflowSelection.h
@@ -65,3 +65,14 @@ struct FDataflowFaceSelection : public FDataflowSelection
 	GENERATED_USTRUCT_BODY()
 };
 
+USTRUCT()
+struct FDataflowGeometrySelection : public FDataflowSelection
+{
+	GENERATED_USTRUCT_BODY()
+};
+
+USTRUCT()
+struct FDataflowMaterialSelection : public FDataflowSelection
+{
+	GENERATED_USTRUCT_BODY()
+};
