ChaosFlesh: GenerateSurfaceBindings binds global mesh instead of individual flesh components; vertex attribute transfer; supports materials for DeformableTetrahedralComponent

diff --git a/Engine/Plugins/Experimental/ChaosFlesh/Source/ChaosFleshEngine/Private/ChaosFlesh/ChaosDeformableTetrahedralComponent.cpp b/Engine/Plugins/Experimental/ChaosFlesh/Source/ChaosFleshEngine/Private/ChaosFlesh/ChaosDeformableTetrahedralComponent.cpp
--- a/Engine/Plugins/Experimental/ChaosFlesh/Source/ChaosFleshEngine/Private/ChaosFlesh/ChaosDeformableTetrahedralComponent.cpp
+++ b/Engine/Plugins/Experimental/ChaosFlesh/Source/ChaosFleshEngine/Private/ChaosFlesh/ChaosDeformableTetrahedralComponent.cpp
@@ -496,7 +496,7 @@ void UDeformableTetrahedralComponent::RenderProceduralMesh()
 								RenderMesh = new FFleshRenderMesh;
 								TArray<FVector3f> Vertex;
 								Flesh.ComponentSpaceVertices(Vertex);
-
+								const TManagedArray<FLinearColor>* Color = FleshCollection->FindAttributeTyped<FLinearColor>(FGeometryCollection::ColorAttribute, FGeometryCollection::VerticesGroup);
 								for (int i = 0; i < NumFaces; ++i)
 								{
 									const auto& P1 = Vertex[Flesh.Indices[i][0]];
@@ -506,10 +506,18 @@ void UDeformableTetrahedralComponent::RenderProceduralMesh()
 									RenderMesh->Vertices.Add(FVector(P1));
 									RenderMesh->Vertices.Add(FVector(P2));
 									RenderMesh->Vertices.Add(FVector(P3));
-
-									RenderMesh->Colors.Add(FLinearColor::White);
-									RenderMesh->Colors.Add(FLinearColor::White);
-									RenderMesh->Colors.Add(FLinearColor::White);
+									if (Color)
+									{
+										RenderMesh->Colors.Add((*Color)[Flesh.Indices[i][0]]);
+										RenderMesh->Colors.Add((*Color)[Flesh.Indices[i][1]]);
+										RenderMesh->Colors.Add((*Color)[Flesh.Indices[i][2]]);
+									}
+									else
+									{
+										RenderMesh->Colors.Add(FLinearColor::White);
+										RenderMesh->Colors.Add(FLinearColor::White);
+										RenderMesh->Colors.Add(FLinearColor::White);
+									}
 
 									RenderMesh->UVs.Add(FVector2D(0, 0));
 									RenderMesh->UVs.Add(FVector2D(0, 0));
@@ -531,7 +539,10 @@ void UDeformableTetrahedralComponent::RenderProceduralMesh()
 									Tangent = (P1 - P3).GetSafeNormal();
 									RenderMesh->Tangents.Add(FProcMeshTangent(Tangent[0], Tangent[1], Tangent[2]));
 								}
-
+								if (Material)
+								{
+									Mesh->SetMaterial(0, Material);
+								}
 								Mesh->SetRelativeTransform(GetComponentTransform());
 								Mesh->CreateMeshSection_LinearColor(0, RenderMesh->Vertices, RenderMesh->Triangles, RenderMesh->Normals, RenderMesh->UVs, RenderMesh->Colors, RenderMesh->Tangents, false);
 							}
@@ -575,7 +586,7 @@ void UDeformableTetrahedralComponent::RenderProceduralMesh()
 								if (!Mesh->GetComponentTransform().Equals(GetComponentTransform())) {
 									Mesh->SetRelativeTransform(GetComponentTransform());
 								}
-								Mesh->UpdateMeshSection_LinearColor(0, RenderMesh->Vertices, RenderMesh->Normals, RenderMesh->UVs, RenderMesh->Colors, RenderMesh->Tangents);
+								Mesh->UpdateMeshSection_LinearColor(0, RenderMesh->Vertices, RenderMesh->Normals, RenderMesh->UVs, RenderMesh->Colors, RenderMesh->Tangents, false);
 							}
 
 							bCanRender = true;
diff --git a/Engine/Plugins/Experimental/ChaosFlesh/Source/ChaosFleshEngine/Public/ChaosFlesh/ChaosDeformableTetrahedralComponent.h b/Engine/Plugins/Experimental/ChaosFlesh/Source/ChaosFleshEngine/Public/ChaosFlesh/ChaosDeformableTetrahedralComponent.h
--- a/Engine/Plugins/Experimental/ChaosFlesh/Source/ChaosFleshEngine/Public/ChaosFlesh/ChaosDeformableTetrahedralComponent.h
+++ b/Engine/Plugins/Experimental/ChaosFlesh/Source/ChaosFleshEngine/Public/ChaosFlesh/ChaosDeformableTetrahedralComponent.h
@@ -143,6 +143,10 @@ public:
 	virtual TOptional<TArray<int32>> GetMeshImportVertexMap(const USkinnedAsset& SkinnedMeshAsset) const override;
 	//~ End IDataflowGeometryCachable Interface
 
+	virtual void SetMaterial(int32 Index, UMaterialInterface* InMaterial) override
+	{
+		Material = InMaterial;
+	}
 	/**
 	* Get the current positions of the transformation hierarchy from \c TargetDeformationSkeleton,
 	* deformed by the tetrahedral mesh.  Results can be in world space postions/deltas, component space
@@ -165,6 +169,9 @@ public:
 	UPROPERTY(EditAnywhere, Category = "Rendering")
 	TObjectPtr<UProceduralMeshComponent> Mesh;
 
+	UPROPERTY(BlueprintReadWrite, Category = "Rendering")
+	TObjectPtr<UMaterialInterface> Material = nullptr;
+
 	UPROPERTY(EditAnywhere, Category = "Rendering")
 	TArray<int32> HideTetrahedra;
 
diff --git a/Engine/Plugins/Experimental/ChaosFlesh/Source/ChaosFleshNodes/Private/Dataflow/ChaosFleshGenerateSurfaceBindingsNode.cpp b/Engine/Plugins/Experimental/ChaosFlesh/Source/ChaosFleshNodes/Private/Dataflow/ChaosFleshGenerateSurfaceBindingsNode.cpp
--- a/Engine/Plugins/Experimental/ChaosFlesh/Source/ChaosFleshNodes/Private/Dataflow/ChaosFleshGenerateSurfaceBindingsNode.cpp
+++ b/Engine/Plugins/Experimental/ChaosFlesh/Source/ChaosFleshNodes/Private/Dataflow/ChaosFleshGenerateSurfaceBindingsNode.cpp
@@ -272,6 +272,20 @@ FGenerateSurfaceBindings::Evaluate(Dataflow::FContext& Context, const FDataflowO
 			}
 			TArray<FString> GeometryGroupGuidsLocal = GetValue<TArray<FString>>(Context, &GeometryGroupGuidsIn);
 			TManagedArray<FString>* Guids = OutCollection.FindAttribute<FString>("Guid", FGeometryCollection::GeometryGroup);
+
+			// Build Tetrahedra
+			TArray<Chaos::TTetrahedron<Chaos::FReal>> Tets;			// Index 0 == TetMeshStart
+			TArray<Chaos::TTetrahedron<Chaos::FReal>*> BVHTetPtrs;
+
+			//
+			// Init boundary mesh for projections.
+			//
+			Chaos::FTriangleMesh TetBoundaryMesh;
+			Chaos::FTriangleMesh::TBVHType<Chaos::FRealDouble> TetBoundaryBVH;
+			TArray<Chaos::FVec3> VertexD;
+			TConstArrayView<Chaos::TVec3<Chaos::FRealDouble>> VertexDView(VertexD);
+			TArray<Chaos::FVec3> PointNormals;
+
 			for (int32 TetMeshIdx = 0; TetMeshIdx < TetrahedronStart->Num(); TetMeshIdx++)
 			{
 				if (GeometryGroupGuidsLocal.Num() && Guids)
@@ -285,59 +299,22 @@ FGenerateSurfaceBindings::Evaluate(Dataflow::FContext& Context, const FDataflowO
 				const int32 TetMeshCount = (*TetrahedronCount)[TetMeshIdx];
 				if (TetMeshCount > 0)
 				{
-					// Build Tetrahedra
-					TArray<Chaos::TTetrahedron<Chaos::FReal>> Tets;			// Index 0 == TetMeshStart
-					TArray<Chaos::TTetrahedron<Chaos::FReal>*> BVHTetPtrs;
-					Tets.SetNumUninitialized(TetMeshCount);
-					BVHTetPtrs.SetNumUninitialized(TetMeshCount);
 					for (int32 i = 0; i < TetMeshCount; i++)
 					{
 						const int32 Idx = TetMeshStart + i;
 						const FIntVector4& Tet = (*Tetrahedron)[Idx];
-						Tets[i] = Chaos::TTetrahedron<Chaos::FReal>(
+						Tets.Add(Chaos::TTetrahedron<Chaos::FReal>(
 							(*Vertex)[Tet[0]],
 							(*Vertex)[Tet[1]],
 							(*Vertex)[Tet[2]],
-							(*Vertex)[Tet[3]]);
-						BVHTetPtrs[i] = &Tets[i];
+							(*Vertex)[Tet[3]]));
 					}
 
-					// Init BVH for tetrahedra.
-					Chaos::TBoundingVolumeHierarchy<
-						TArray<Chaos::TTetrahedron<Chaos::FReal>*>,
-						TArray<int32>,
-						Chaos::FReal,
-						3> TetBVH(BVHTetPtrs);
-
-					//
-					// Init boundary mesh for projections.
-					//
-
-					Chaos::FTriangleMesh TetBoundaryMesh;
-					Chaos::FTriangleMesh::TBVHType<Chaos::FRealDouble> TetBoundaryBVH;
-					TArray<Chaos::FVec3> VertexD;
-					TConstArrayView<Chaos::TVec3<Chaos::FRealDouble>> VertexDView(VertexD);
-					TArray<Chaos::FVec3> PointNormals;
-					const int32 TriMeshStart = (*FacesStart)[TetMeshIdx];
-					const int32 TriMeshCount = (*FacesCount)[TetMeshIdx];
 					if (bDoSurfaceProjection)
 					{
-						if (TriMeshCount == Triangle->GetConstArray().Num())
-						{
-							TetBoundaryMesh.Init(
+						TetBoundaryMesh.Init(
 								reinterpret_cast<const TArray<Chaos::TVec3<int32>>&>(Triangle->GetConstArray()), 0, -1, false);
-						}
-						else
-						{
-							TArray<Chaos::TVec3<int32>> Faces; Faces.SetNumUninitialized(TriMeshCount);
-							for (int32 i = 0; i < TriMeshCount; i++)
-							{
-								const int32 Idx = TriMeshStart + i;
-								Faces[i] = reinterpret_cast<const Chaos::TVec3<int32>&>(Triangle->GetConstArray()[Idx]);
-							}
-							TetBoundaryMesh.Init(Faces, 0, -1, false);
-						}
-
+						
 						// Promote vertices to double because that's what FTriangleMesh wants.
 						VertexD.SetNumUninitialized(Vertex->Num());
 						for (int32 i = 0; i < VertexD.Num(); i++)
@@ -349,269 +326,278 @@ FGenerateSurfaceBindings::Evaluate(Dataflow::FContext& Context, const FDataflowO
 						PointNormals = TetBoundaryMesh.GetPointNormals(VertexDView, false, true);
 						TetBoundaryMesh.BuildBVH(VertexDView, TetBoundaryBVH);
 					}
+				}
+			}
+			// Init BVH for tetrahedra.
+			BVHTetPtrs.SetNumUninitialized(Tets.Num());
+			for (int32 TetIdx = 0; TetIdx < Tets.Num(); ++TetIdx)
+			{
+				BVHTetPtrs[TetIdx] = &Tets[TetIdx];
+			}
+			Chaos::TBoundingVolumeHierarchy<
+				TArray<Chaos::TTetrahedron<Chaos::FReal>*>,
+				TArray<int32>,
+				Chaos::FReal,
+				3> TetBVH(BVHTetPtrs);
+
+			//
+			// Do intersection tests against tets, then the surface.
+			//
+
+			TArray<TArray<FIntVector4>> Parents; Parents.SetNum(MeshVertices.Num());
+			TArray<TArray<FVector4f>> Weights;	 Weights.SetNum(MeshVertices.Num());
+			TArray<TArray<FVector3f>> Offsets;	 Offsets.SetNum(MeshVertices.Num());
+			TArray<TArray<float>> Masks;		 Masks.SetNum(MeshVertices.Num());
+			TArray<int32> Orphans;
+			int32 TetHits = 0;
+			int32 TriHits = 0;
+			int32 Adoptions = 0;
+			int32 NumOrphans = 0;
+			for (int32 LOD = 0; LOD < MeshVertices.Num(); LOD++)
+			{
+				Parents[LOD].SetNumUninitialized(MeshVertices[LOD].Num());
+				Weights[LOD].SetNumUninitialized(MeshVertices[LOD].Num());
+				Offsets[LOD].SetNumUninitialized(MeshVertices[LOD].Num());
+				Masks[LOD].SetNumUninitialized(MeshVertices[LOD].Num());
 
-					//
-					// Do intersection tests against tets, then the surface.
-					//
-
-					TArray<TArray<FIntVector4>> Parents; Parents.SetNum(MeshVertices.Num());
-					TArray<TArray<FVector4f>> Weights;	 Weights.SetNum(MeshVertices.Num());
-					TArray<TArray<FVector3f>> Offsets;	 Offsets.SetNum(MeshVertices.Num());
-					TArray<TArray<float>> Masks;		 Masks.SetNum(MeshVertices.Num());
-					TArray<int32> Orphans;
-					int32 TetHits = 0;
-					int32 TriHits = 0;
-					int32 Adoptions = 0;
-					int32 NumOrphans = 0;
-					for (int32 LOD = 0; LOD < MeshVertices.Num(); LOD++)
+				TArray<int32> TetIntersections; TetIntersections.Reserve(64);
+				for (int32 i = 0; i < MeshVertices[LOD].Num(); i++)
+				{
+					Parents[LOD][i] = FIntVector4(INDEX_NONE);
+					Weights[LOD][i] = FVector4f(0);
+					Offsets[LOD][i] = FVector3f(0);
+					Masks[LOD][i] = 0.0; // Shader does skinning for this vertex
+
+					const FVector3f& Pos = MeshVertices[LOD][i];
+					Chaos::TVec3<Chaos::FReal> PosD(Pos[0], Pos[1], Pos[2]);
+					TetIntersections = TetBVH.FindAllIntersections(PosD);
+					int32 j = 0;
+					for (; j < TetIntersections.Num(); j++)
 					{
-						Parents[LOD].SetNumUninitialized(MeshVertices[LOD].Num());
-						Weights[LOD].SetNumUninitialized(MeshVertices[LOD].Num());
-						Offsets[LOD].SetNumUninitialized(MeshVertices[LOD].Num());
-						Masks[LOD].SetNumUninitialized(MeshVertices[LOD].Num());
-
-						TArray<int32> TetIntersections; TetIntersections.Reserve(64);
-						for (int32 i = 0; i < MeshVertices[LOD].Num(); i++)
+						const int32 TetIdx = TetIntersections[j];
+						if (!Tets[TetIdx].Outside(Pos, 0.f)) // includes boundary
 						{
-							Parents[LOD][i] = FIntVector4(INDEX_NONE);
-							Weights[LOD][i] = FVector4f(0);
+							TetHits++;
+							Parents[LOD][i] = (*Tetrahedron)[TetIdx];
+							Chaos::TVector<Chaos::FReal, 4> WeightsD = Tets[TetIdx].GetBarycentricCoordinates(Pos);
+							Weights[LOD][i] = FVector4f(WeightsD[0], WeightsD[1], WeightsD[2], WeightsD[3]);
 							Offsets[LOD][i] = FVector3f(0);
-							Masks[LOD][i] = 0.0; // Shader does skinning for this vertex
+							Masks[LOD][i] = 1.0; // Shader does sim for this vertex
+
+							FVector3f EmbeddedPos =
+								(*Vertex)[Parents[LOD][i][0]] * Weights[LOD][i][0] +
+								(*Vertex)[Parents[LOD][i][1]] * Weights[LOD][i][1] +
+								(*Vertex)[Parents[LOD][i][2]] * Weights[LOD][i][2] +
+								(*Vertex)[Parents[LOD][i][3]] * Weights[LOD][i][3];
+							check((Pos - EmbeddedPos).SquaredLength() < UE_SMALL_NUMBER);
+							break;
+						}
+					}
+					if (j == TetIntersections.Num())
+					{
+						// This vertex didn't land inside any tetrahedra. Project it to the tet boundary surface.
+						int32 TriIdx = INDEX_NONE;
+						Chaos::FVec3 TriWeights;
+						if (bDoSurfaceProjection &&
+							TetBoundaryMesh.SmoothProject(
+								TetBoundaryBVH,
+								VertexDView,
+								PointNormals, Pos,
+								TriIdx, TriWeights, SurfaceProjectionIterations))
+						{
+							TriHits++;
+							const FIntVector& Tri = (*Triangle)[TriIdx];
+							Parents[LOD][i][0] = Tri[0];
+							Parents[LOD][i][1] = Tri[1];
+							Parents[LOD][i][2] = Tri[2];
+							Parents[LOD][i][3] = INDEX_NONE;
+
+							Weights[LOD][i][0] = TriWeights[0];
+							Weights[LOD][i][1] = TriWeights[1];
+							Weights[LOD][i][2] = TriWeights[2];
+							Weights[LOD][i][3] = 0.0;
+
+							const FVector3f EmbeddedPos =
+								TriWeights[0] * Vertex->GetConstArray()[Tri[0]] +
+								TriWeights[1] * Vertex->GetConstArray()[Tri[1]] +
+								TriWeights[2] * Vertex->GetConstArray()[Tri[2]];
+							Offsets[LOD][i] = EmbeddedPos - Pos;
+
+							Masks[LOD][i] = 1.0; // Shader does sim for this vertex
+						}
+						else
+						{
+							// Despair...
+							Orphans.Add(i);
 
-							const FVector3f& Pos = MeshVertices[LOD][i];
-							Chaos::TVec3<Chaos::FReal> PosD(Pos[0], Pos[1], Pos[2]);
-							TetIntersections = TetBVH.FindAllIntersections(PosD);
-							int32 j = 0;
-							for (; j < TetIntersections.Num(); j++)
-							{
-								const int32 TetIdx = TetIntersections[j];
-								if (!Tets[TetIdx].Outside(Pos, 1.0e-2)) // includes boundary
-								{
-									TetHits++;
-									Parents[LOD][i] = (*Tetrahedron)[TetIdx + TetMeshStart];
-									Chaos::TVector<Chaos::FReal, 4> WeightsD = Tets[TetIdx].GetBarycentricCoordinates(Pos);
-									Weights[LOD][i] = FVector4f(WeightsD[0], WeightsD[1], WeightsD[2], WeightsD[3]);
-									Offsets[LOD][i] = FVector3f(0);
-									Masks[LOD][i] = 1.0; // Shader does sim for this vertex
-
-									FVector3f EmbeddedPos =
-										(*Vertex)[Parents[LOD][i][0]] * Weights[LOD][i][0] +
-										(*Vertex)[Parents[LOD][i][1]] * Weights[LOD][i][1] +
-										(*Vertex)[Parents[LOD][i][2]] * Weights[LOD][i][2] +
-										(*Vertex)[Parents[LOD][i][3]] * Weights[LOD][i][3];
-									check((Pos - EmbeddedPos).SquaredLength() < 1.0);
-
-									break;
-								}
-							}
-							if (j == TetIntersections.Num())
-							{
-								// This vertex didn't land inside any tetrahedra. Project it to the tet boundary surface.
-								int32 TriIdx = INDEX_NONE;
-								Chaos::FVec3 TriWeights;
-								if (bDoSurfaceProjection &&
-									TetBoundaryMesh.SmoothProject(
-										TetBoundaryBVH,
-										VertexDView,
-										PointNormals, Pos,
-										TriIdx, TriWeights, SurfaceProjectionIterations))
-								{
-									TriHits++;
-									const FIntVector& Tri = (*Triangle)[TriIdx + TriMeshStart];
-									Parents[LOD][i][0] = Tri[0];
-									Parents[LOD][i][1] = Tri[1];
-									Parents[LOD][i][2] = Tri[2];
-									Parents[LOD][i][3] = INDEX_NONE;
-
-									Weights[LOD][i][0] = TriWeights[0];
-									Weights[LOD][i][1] = TriWeights[1];
-									Weights[LOD][i][2] = TriWeights[2];
-									Weights[LOD][i][3] = 0.0;
-
-									const FVector3f EmbeddedPos =
-										TriWeights[0] * Vertex->GetConstArray()[Tri[0]] +
-										TriWeights[1] * Vertex->GetConstArray()[Tri[1]] +
-										TriWeights[2] * Vertex->GetConstArray()[Tri[2]];
-									Offsets[LOD][i] = EmbeddedPos - Pos;
-
-									Masks[LOD][i] = 1.0; // Shader does sim for this vertex
-								}
-								else
-								{
-									// Despair...
-									Orphans.Add(i);
-
-									Parents[LOD][i][0] = INDEX_NONE;
-									Parents[LOD][i][1] = INDEX_NONE;
-									Parents[LOD][i][2] = INDEX_NONE;
-									Parents[LOD][i][3] = INDEX_NONE;
-
-									Weights[LOD][i][0] = 0.0;
-									Weights[LOD][i][1] = 0.0;
-									Weights[LOD][i][2] = 0.0;
-									Weights[LOD][i][3] = 0.0;
-
-									Offsets[LOD][i][0] = 0.0;
-									Offsets[LOD][i][1] = 0.0;
-									Offsets[LOD][i][2] = 0.0;
-
-									Masks[LOD][i] = 0.0; // Shader does skinning for this vertex
-								} // if !SmoothProject()
-							} // if !TetIntersections
-						} // end for all vertices
-
-						// 
-						// Advancing front orphan reparenting
-						//
-						if (!MeshNeighborNodes.IsValidIndex(LOD))
+							Parents[LOD][i][0] = INDEX_NONE;
+							Parents[LOD][i][1] = INDEX_NONE;
+							Parents[LOD][i][2] = INDEX_NONE;
+							Parents[LOD][i][3] = INDEX_NONE;
+
+							Weights[LOD][i][0] = 0.0;
+							Weights[LOD][i][1] = 0.0;
+							Weights[LOD][i][2] = 0.0;
+							Weights[LOD][i][3] = 0.0;
+
+							Offsets[LOD][i][0] = 0.0;
+							Offsets[LOD][i][1] = 0.0;
+							Offsets[LOD][i][2] = 0.0;
+
+							Masks[LOD][i] = 0.0; // Shader does skinning for this vertex
+						} // if !SmoothProject()
+					} // if !TetIntersections
+				} // end for all vertices
+
+				// 
+				// Advancing front orphan reparenting
+				//
+				if (!MeshNeighborNodes.IsValidIndex(LOD))
+				{
+					continue;
+				}
+				const TArray<TArray<uint32>>& NeighborNodes = MeshNeighborNodes[LOD];
+				TSet<int32> OrphanSet(Orphans);
+				while (bDoOrphanReparenting && Orphans.Num())
+				{
+					// Find the orphan with the fewest number of orphan neighbors, and the 
+					// most non-orphans in their 1 ring.
+					int32 Orphan = INDEX_NONE;
+					int32 NumOrphanNeighbors = TNumericLimits<int32>::Max();
+					int32 NumNonOrphanNeighbors = 0;
+					for (int32 i = 0; i < Orphans.Num(); i++)
+					{
+						int32 CurrOrphan = Orphans[i];
+						if (!NeighborNodes.IsValidIndex(CurrOrphan))
 						{
 							continue;
 						}
-						const TArray<TArray<uint32>>& NeighborNodes = MeshNeighborNodes[LOD];
-						TSet<int32> OrphanSet(Orphans);
-						while (bDoOrphanReparenting && Orphans.Num())
+						const TArray<uint32>& Neighbors = NeighborNodes[CurrOrphan];
+						int32 OrphanCount = 0;
+						int32 NonOrphanCount = 0;
+						for (int32 j = 0; j < Neighbors.Num(); j++)
 						{
-							// Find the orphan with the fewest number of orphan neighbors, and the 
-							// most non-orphans in their 1 ring.
-							int32 Orphan = INDEX_NONE;
-							int32 NumOrphanNeighbors = TNumericLimits<int32>::Max();
-							int32 NumNonOrphanNeighbors = 0;
-							for (int32 i = 0; i < Orphans.Num(); i++)
+							if (OrphanSet.Contains(Neighbors[j]))
 							{
-								int32 CurrOrphan = Orphans[i];
-								if (!NeighborNodes.IsValidIndex(CurrOrphan))
-								{
-									continue;
-								}
-								const TArray<uint32>& Neighbors = NeighborNodes[CurrOrphan];
-								int32 OrphanCount = 0;
-								int32 NonOrphanCount = 0;
-								for (int32 j = 0; j < Neighbors.Num(); j++)
-								{
-									if (OrphanSet.Contains(Neighbors[j]))
-									{
-										OrphanCount++;
-									}
-									else
-									{
-										NonOrphanCount++;
-									}
-								}
-								if (OrphanCount <= NumOrphanNeighbors && NonOrphanCount > NumNonOrphanNeighbors)
-								{
-									Orphan = CurrOrphan;
-									NumOrphanNeighbors = OrphanCount;
-									NumNonOrphanNeighbors = NonOrphanCount;
-								}
+								OrphanCount++;
 							}
-							if (Orphan == INDEX_NONE)
+							else
 							{
-								// We only have orphans with no neighbors left.
-								break;
+								NonOrphanCount++;
 							}
-							const FVector3f& Pos = MeshVertices[LOD][Orphan];
-							Chaos::TVec3<Chaos::FReal> PosD(Pos[0], Pos[1], Pos[2]);
-
-							// Use the parent simplices of non-orphan neighbors as test candidates.
-							Chaos::FReal CurrDist = TNumericLimits<Chaos::FReal>::Max();
-							const TArray<uint32>& Neighbors = NeighborNodes[Orphan];
-							bool FoundBinding = false;
-							for (int32 i = 0; i < Neighbors.Num(); i++)
-							{
-								const uint32 Neighbor = Neighbors[i];
-								if (OrphanSet.Contains(Neighbor))
-								{
-									continue;
-								}
+						}
+						if (OrphanCount <= NumOrphanNeighbors && NonOrphanCount > NumNonOrphanNeighbors)
+						{
+							Orphan = CurrOrphan;
+							NumOrphanNeighbors = OrphanCount;
+							NumNonOrphanNeighbors = NonOrphanCount;
+						}
+					}
+					if (Orphan == INDEX_NONE)
+					{
+						// We only have orphans with no neighbors left.
+						break;
+					}
+					const FVector3f& Pos = MeshVertices[LOD][Orphan];
+					Chaos::TVec3<Chaos::FReal> PosD(Pos[0], Pos[1], Pos[2]);
+
+					// Use the parent simplices of non-orphan neighbors as test candidates.
+					Chaos::FReal CurrDist = TNumericLimits<Chaos::FReal>::Max();
+					const TArray<uint32>& Neighbors = NeighborNodes[Orphan];
+					bool FoundBinding = false;
+					for (int32 i = 0; i < Neighbors.Num(); i++)
+					{
+						const uint32 Neighbor = Neighbors[i];
+						if (OrphanSet.Contains(Neighbor))
+						{
+							continue;
+						}
 
-								const FIntVector4& P = Parents[LOD][Neighbor];
-								int32 NumValid = 0;
-								for (int32 j = 0; j < 4; j++)
-								{
-									NumValid += P[j] != INDEX_NONE ? 1 : 0;
-								}
+						const FIntVector4& P = Parents[LOD][Neighbor];
+						int32 NumValid = 0;
+						for (int32 j = 0; j < 4; j++)
+						{
+							NumValid += P[j] != INDEX_NONE ? 1 : 0;
+						}
 
-								if (NumValid == 0)
-								{
-									continue;
-								}
-								else
+						if (NumValid == 0)
+						{
+							continue;
+						}
+						else
+						{
+							// Find tets that share parent indices
+							for (int32 j = 0; j < 4; j++)
+							{
+								const int32 ParentIdx = P[j];
+								if (IncidentElements->GetConstArray().IsValidIndex(ParentIdx))
 								{
-									// Find tets that share parent indices
-									for (int32 j = 0; j < 4; j++)
+									const TArray<int32>& NeighborTets = (*IncidentElements)[ParentIdx];
+									for (int32 k = 0; k < NeighborTets.Num(); k++)
 									{
-										const int32 ParentIdx = P[j];
-										if (IncidentElements->GetConstArray().IsValidIndex(ParentIdx))
+										const int32 TetIdx = NeighborTets[k];
+										if (ensure(Tets.IsValidIndex(TetIdx)))
 										{
-											const TArray<int32>& NeighborTets = (*IncidentElements)[ParentIdx];
-											for (int32 k = 0; k < NeighborTets.Num(); k++)
+											const Chaos::FTetrahedron& Tet = Tets[TetIdx];
+
+											Chaos::TVec4<Chaos::FReal> W;
+											Chaos::TVec3<Chaos::FReal> EmbeddedPos = Tet.FindClosestPointAndBary(PosD, W, 1.0e-4);
+											Chaos::TVec3<Chaos::FReal> O = EmbeddedPos - PosD;
+											Chaos::FReal Dist = O.SquaredLength();
+											if (Dist < CurrDist)
 											{
-												const int32 TetIdx = NeighborTets[k] - TetMeshStart;
-												if (ensure(Tets.IsValidIndex(TetIdx)))
-												{
-													const Chaos::FTetrahedron& Tet = Tets[TetIdx];
-
-													Chaos::TVec4<Chaos::FReal> W;
-													Chaos::TVec3<Chaos::FReal> EmbeddedPos = Tet.FindClosestPointAndBary(PosD, W, 1.0e-4);
-													Chaos::TVec3<Chaos::FReal> O = EmbeddedPos - PosD;
-													Chaos::FReal Dist = O.SquaredLength();
-													if (Dist < CurrDist)
-													{
-														CurrDist = Dist;
-														Parents[LOD][Orphan] = (*Tetrahedron)[TetIdx + TetMeshStart];
-														Weights[LOD][Orphan] = FVector4f(W[0], W[1], W[2], W[3]);
-														Offsets[LOD][Orphan] = FVector3f(O[0], O[1], O[2]);
-														Masks[LOD][i] = 1.0; // Shader does sim for this vertex
-														FoundBinding = true;
-													}
-												}
+												CurrDist = Dist;
+												Parents[LOD][Orphan] = (*Tetrahedron)[TetIdx];
+												Weights[LOD][Orphan] = FVector4f(W[0], W[1], W[2], W[3]);
+												Offsets[LOD][Orphan] = FVector3f(O[0], O[1], O[2]);
+												Masks[LOD][i] = 1.0; // Shader does sim for this vertex
+												FoundBinding = true;
 											}
 										}
 									}
 								}
-							} // end for all neighbors
-
-							// Whether or not we successfully reparented, remove the orphan from the list.
-							OrphanSet.Remove(Orphan);
-							Orphans.Remove(Orphan);
-							if (FoundBinding)
-							{
-								Adoptions++;
-							}
-							else
-							{
-								NumOrphans++;
 							}
-						} // end while(Orphans)
-						NumOrphans += Orphans.Num();
-
-						//ELogVerbosity::Type Verbosity = Orphans.Num() > 0 ? ELogVerbosity::Error : ELogVerbosity::Display;
-						UE_LOG(LogMeshBindings, Display,
-							TEXT("'%s' - Generated mesh bindings between tet mesh index %d and %s mesh of '%s' LOD %d - stats:\n"
-								"    Render vertices num: %d\n"
-								"    Vertices in tetrahedra: %d\n"
-								"    Vertices bound to tet surface: %d\n"
-								"    Orphaned vertices reparented: %d\n"
-								"    Vertices orphaned: %d"),
-							*GetName().ToString(),
-							TetMeshIdx,
-							bUseSkeletalMeshImportModel ? TEXT("import") : TEXT("render"),
-							*MeshId, LOD,
-							MeshVertices[LOD].Num(), TetHits, TriHits, Adoptions, NumOrphans);
-
-					} // end for all LOD
-
-					// Stash bindings in the geometry collection
-					GeometryCollection::Facades::FTetrahedralBindings TetBindings(OutCollection);
-					TetBindings.DefineSchema();
-					FName MeshName(*MeshId, MeshId.Len());
-					for (int32 LOD = 0; LOD < MeshVertices.Num(); LOD++)
+						}
+					} // end for all neighbors
+
+					// Whether or not we successfully reparented, remove the orphan from the list.
+					OrphanSet.Remove(Orphan);
+					Orphans.Remove(Orphan);
+					if (FoundBinding)
 					{
-						TetBindings.AddBindingsGroup(TetMeshIdx, MeshName, LOD);
-						TetBindings.SetBindingsData(Parents[LOD], Weights[LOD], Offsets[LOD], Masks[LOD]);
+						Adoptions++;
 					}
-				} // TetMeshCount > 0
-			} // end for TetMeshIdx
+					else
+					{
+						NumOrphans++;
+					}
+				} // end while(Orphans)
+				NumOrphans += Orphans.Num();
+
+				//ELogVerbosity::Type Verbosity = Orphans.Num() > 0 ? ELogVerbosity::Error : ELogVerbosity::Display;
+				UE_LOG(LogMeshBindings, Display,
+					TEXT("'%s' - Generated mesh bindings between tet mesh and %s mesh of '%s' LOD %d - stats:\n"
+						"    Render vertices num: %d\n"
+						"    Vertices in tetrahedra: %d\n"
+						"    Vertices bound to tet surface: %d\n"
+						"    Orphaned vertices reparented: %d\n"
+						"    Vertices orphaned: %d"),
+					*GetName().ToString(),
+					bUseSkeletalMeshImportModel ? TEXT("import") : TEXT("render"),
+					*MeshId, LOD,
+					MeshVertices[LOD].Num(), TetHits, TriHits, Adoptions, NumOrphans);
+
+			} // end for all LOD
+
+			// Stash bindings in the geometry collection
+			GeometryCollection::Facades::FTetrahedralBindings TetBindings(OutCollection);
+			TetBindings.DefineSchema();
+			FName MeshName(*MeshId, MeshId.Len());
+			for (int32 LOD = 0; LOD < MeshVertices.Num(); LOD++)
+			{
+				TetBindings.AddBindingsGroup(/*TetMeshIdx = */ 0, MeshName, LOD);
+				TetBindings.SetBindingsData(Parents[LOD], Weights[LOD], Offsets[LOD], Masks[LOD]);
+			}
 		}
 		SetValue(Context, MoveTemp(OutCollection), &Collection);
 	}
diff --git a/Engine/Plugins/Experimental/GeometryCollectionPlugin/Source/GeometryCollectionDepNodes/GeometryCollectionDepNodes.Build.cs b/Engine/Plugins/Experimental/GeometryCollectionPlugin/Source/GeometryCollectionDepNodes/GeometryCollectionDepNodes.Build.cs
--- a/Engine/Plugins/Experimental/GeometryCollectionPlugin/Source/GeometryCollectionDepNodes/GeometryCollectionDepNodes.Build.cs
+++ b/Engine/Plugins/Experimental/GeometryCollectionPlugin/Source/GeometryCollectionDepNodes/GeometryCollectionDepNodes.Build.cs
@@ -15,7 +15,9 @@ namespace UnrealBuildTool.Rules
 					"Chaos",
 					"DataflowCore",
 					"DataflowEngine",
+					"DataflowEnginePlugin",
 					"GeometryCollectionEngine",
+					"GeometryCollectionNodes",
 					"Engine"
 				}
 			);
diff --git a/Engine/Plugins/Experimental/GeometryCollectionPlugin/Source/GeometryCollectionDepNodes/Private/Dataflow/GeometryCollectionDepNodesPlugin.cpp b/Engine/Plugins/Experimental/GeometryCollectionPlugin/Source/GeometryCollectionDepNodes/Private/Dataflow/GeometryCollectionDepNodesPlugin.cpp
--- a/Engine/Plugins/Experimental/GeometryCollectionPlugin/Source/GeometryCollectionDepNodes/Private/Dataflow/GeometryCollectionDepNodesPlugin.cpp
+++ b/Engine/Plugins/Experimental/GeometryCollectionPlugin/Source/GeometryCollectionDepNodes/Private/Dataflow/GeometryCollectionDepNodesPlugin.cpp
@@ -5,16 +5,16 @@
 #include "CoreMinimal.h"
 
 #include "Dataflow/DataflowNodeFactory.h"
+#include "Dataflow/GeometryCollectionTransferVertexScalarAttributeDepNode.h"
 #include "Dataflow/SetVertexColorFromFloatArrayDepNode.h"
 #include "Dataflow/SetVertexColorFromVertexSelectionDepNode.h"
 
-
-
 #define LOCTEXT_NAMESPACE "DataflowNodes"
 
 
 void IGeometryCollectionDepNodesPlugin::StartupModule()
 {
+	DATAFLOW_NODE_REGISTER_CREATION_FACTORY(FGeometryCollectionTransferVertexScalarAttributeNode);
 	DATAFLOW_NODE_REGISTER_CREATION_FACTORY(FSetVertexColorInCollectionFromVertexSelectionDataflowNode);
 	DATAFLOW_NODE_REGISTER_CREATION_FACTORY(FSetVertexColorInCollectionFromFloatArrayDataflowNode);
 }
diff --git a/Engine/Plugins/Experimental/GeometryCollectionPlugin/Source/GeometryCollectionDepNodes/Private/Dataflow/GeometryCollectionTransferVertexScalarAttributeDepNode.cpp b/Engine/Plugins/Experimental/GeometryCollectionPlugin/Source/GeometryCollectionDepNodes/Private/Dataflow/GeometryCollectionTransferVertexScalarAttributeDepNode.cpp
new file mode 100644
--- /dev/null
+++ b/Engine/Plugins/Experimental/GeometryCollectionPlugin/Source/GeometryCollectionDepNodes/Private/Dataflow/GeometryCollectionTransferVertexScalarAttributeDepNode.cpp
@@ -0,0 +1,496 @@
+// Copyright Epic Games, Inc. All Rights Reserved.
+
+#include "Dataflow/GeometryCollectionTransferVertexScalarAttributeDepNode.h"
+
+#include "Chaos/Triangle.h"
+#include "Chaos/TriangleMesh.h"
+#include "Chaos/HierarchicalSpatialHash.h"
+#include "Chaos/TriangleCollisionPoint.h"
+#include "Dataflow/DataflowInputOutput.h"
+#include "GeometryCollection/GeometryCollectionAlgo.h"
+#include "GeometryCollection/ManagedArrayAccessor.h"
+
+#include UE_INLINE_GENERATED_CPP_BY_NAME(GeometryCollectionTransferVertexScalarAttributeDepNode)
+#define LOCTEXT_NAMESPACE "FGeometryCollectionTransferVertexScalarAttributeDepNode"
+
+namespace UE::Private {
+
+	class FTransferFacade
+	{
+		const FManagedArrayCollection& ConstCollection;
+		FManagedArrayCollection* Collection = nullptr;
+	public:
+		FTransferFacade(FManagedArrayCollection& InCollection)
+			: ConstCollection(InCollection)
+			, Collection(&InCollection)
+			, BoneMap(InCollection, FName("BoneMap"), FName("Vertices"))
+			, Vertex(InCollection, FName("Vertex"), FName("Vertices"))
+			, Indices(InCollection, FName("Indices"), FName("Faces"))
+			, Transform(InCollection, FTransformCollection::TransformAttribute, FTransformCollection::TransformGroup)
+			, Parent(InCollection, FTransformCollection::ParentAttribute, FTransformCollection::TransformGroup)
+			, VertexStart(InCollection, "VertexStart", FGeometryCollection::GeometryGroup)
+			, VertexCount(InCollection, "VertexCount", FGeometryCollection::GeometryGroup)
+			, FaceStart(InCollection, "FaceStart", FGeometryCollection::GeometryGroup)
+			, FaceCount(InCollection, "FaceCount", FGeometryCollection::GeometryGroup)
+		{}
+
+		FTransferFacade(const FManagedArrayCollection& InCollection)
+			: ConstCollection(InCollection)
+			, Collection(nullptr)
+			, BoneMap(InCollection, FName("BoneMap"), FName("Vertices"))
+			, Vertex(InCollection, FName("Vertex"), FName("Vertices"))
+			, Indices(InCollection, FName("Indices"), FName("Faces"))
+			, Transform(InCollection, FTransformCollection::TransformAttribute, FTransformCollection::TransformGroup)
+			, Parent(InCollection, FTransformCollection::ParentAttribute, FTransformCollection::TransformGroup)
+			, VertexStart(InCollection, "VertexStart", FGeometryCollection::GeometryGroup)
+			, VertexCount(InCollection, "VertexCount", FGeometryCollection::GeometryGroup)
+			, FaceStart(InCollection, "FaceStart", FGeometryCollection::GeometryGroup)
+			, FaceCount(InCollection, "FaceCount", FGeometryCollection::GeometryGroup)
+		{}
+
+
+		bool IsValid() const {
+			return BoneMap.IsValid() && Vertex.IsValid() && Indices.IsValid() && Transform.IsValid() && Parent.IsValid() && VertexStart.IsValid() &&
+				VertexCount.IsValid() && FaceStart.IsValid() && FaceCount.IsValid();
+		}
+
+		const TManagedArray<float>* GetFloatArray(FString AttributeName, FString Group) const
+		{
+			return ConstCollection.FindAttribute<float>(FName(AttributeName), FName(Group));
+		}
+
+		TManagedArray<float>* GetFloatArray(FString AttributeName, FString Group)
+		{
+			TManagedArray<float>* TargetFloatArray = nullptr;
+			if (!Collection->HasAttribute(FName(AttributeName), FName(Group)))
+			{
+				Collection->AddAttribute<float>(FName(AttributeName), FName(Group));
+			}
+			return Collection->FindAttribute<float>(FName(AttributeName), FName(Group));
+		}
+
+
+		TManagedArrayAccessor<int32> BoneMap;
+		TManagedArrayAccessor<FVector3f> Vertex;
+		TManagedArrayAccessor<FIntVector3> Indices;
+		TManagedArrayAccessor<FTransform3f> Transform;
+		TManagedArrayAccessor<int32> Parent;
+		TManagedArrayAccessor<int32> VertexStart;
+		TManagedArrayAccessor<int32> VertexCount;
+		TManagedArrayAccessor<int32> FaceStart;
+		TManagedArrayAccessor<int32> FaceCount;
+	};
+
+}
+
+void FGeometryCollectionTransferVertexScalarAttributeNode::Evaluate(Dataflow::FContext& Context, const FDataflowOutput* Out) const
+{
+	typedef TPair<FName, FName> NamePair;
+	FCollectionAttributeKey Key = GetValue(Context, &AttributeKey, AttributeKey);
+
+	if (Out->IsA(&Collection))
+	{
+		FManagedArrayCollection TargetCollection = GetValue(Context, &Collection);
+		const FManagedArrayCollection& SampleCollection = GetValue(Context, &FromCollection);
+
+		UE::Private::FTransferFacade Target(TargetCollection);
+		const UE::Private::FTransferFacade Sample(SampleCollection);
+
+		if (Target.IsValid() && Sample.IsValid())
+		{
+			if (TManagedArray<float>* TargetFloatArray = Target.GetFloatArray(Key.Attribute, Key.Group))
+			{
+				TargetFloatArray->Fill(0.f);
+
+				TArray<FIntVector2> AlignedGeometry = FindSourceToTargetGeometryMap(SampleCollection, TargetCollection);
+				if (AlignedGeometry.Num() == TargetCollection.NumElements(FGeometryCollection::GeometryGroup))
+				{
+					PairedGeometryTransfer(Key, AlignedGeometry, Sample, Target, TargetFloatArray);
+				}
+				else
+				{
+					NearestVertexTransfer(Key, Sample, Target, TargetFloatArray);
+				}
+			}
+		}
+
+		SetValue(Context, MoveTemp(TargetCollection), &Collection);
+	}
+	else if (Out->IsA(&AttributeKey))
+	{
+		SetValue(Context, MoveTemp(Key), &AttributeKey);
+	}
+}
+
+
+TArray<FIntVector2> FGeometryCollectionTransferVertexScalarAttributeNode::FindSourceToTargetGeometryMap(const FManagedArrayCollection& AttributeCollectionVal, const FManagedArrayCollection& CollectionVal) const
+{
+	TArray<FIntVector2> Mapping;
+	const TManagedArray<FString>* SourceName = AttributeCollectionVal.FindAttribute<FString>(FName("BoneName"), FTransformCollection::TransformGroup);
+	const TManagedArray<int32>* SourceGeometryGroup = AttributeCollectionVal.FindAttribute<int32>(FName("TransformToGeometryIndex"), FTransformCollection::TransformGroup);
+	const TManagedArray<FString>* TargetName = CollectionVal.FindAttribute<FString>(FName("BoneName"), FTransformCollection::TransformGroup);
+	const TManagedArray<int32>* TargetGeometryGroup = CollectionVal.FindAttribute<int32>(FName("TransformToGeometryIndex"), FTransformCollection::TransformGroup);
+	if (SourceName && SourceGeometryGroup && TargetName && TargetGeometryGroup)
+	{
+		for (int i = 0; i < SourceName->Num(); i++)
+		{
+			for (int j = 0; j < TargetName->Num(); j++)
+			{
+				FString TestName = FString::Printf(TEXT("%s_Tet"), *(*SourceName)[i]);
+				if ((*TargetName)[j].StartsWith(TestName))
+				{
+					Mapping.Add(FIntVector2((*SourceGeometryGroup)[i], (*TargetGeometryGroup)[j]));
+					break;
+				}
+			}
+		}
+	}
+
+	return Mapping;
+}
+
+void FGeometryCollectionTransferVertexScalarAttributeNode::PairedGeometryTransfer(FCollectionAttributeKey Key, const TArray<FIntVector2>& PairedGeometry,
+	const UE::Private::FTransferFacade& Sample, UE::Private::FTransferFacade&  Target, TManagedArray<float>* TargetFloatArray) const
+{
+	if (const TManagedArray<float>* FloatArray = Sample.GetFloatArray(Key.Attribute, Key.Group))
+	{
+		Chaos::FReal SphereFullRadius;
+
+		if (SampleScale == EDataflowTransferNodeSampleScale::Dataflow_Transfer_Asset_Edge || SampleScale == EDataflowTransferNodeSampleScale::Dataflow_Transfer_Asset_Bound)
+		{
+			// Build component space vertices for TargetCollection
+			TArray<FVector> ComponentSpaceFullTargetVertices; // size of num vertices of the geometry entry. 
+			BuildComponentSpaceVertices(&Target.Transform.Get(), &Target.Parent.Get(), &Target.BoneMap.Get(), &Target.Vertex.Get(),
+				0, Target.Vertex.Num(), ComponentSpaceFullTargetVertices);
+
+			// Build component space vertices for TargetCollection
+			TArray<FVector> ComponentSpaceFullVertices; // size of num vertices of the geometry entry
+			BuildComponentSpaceVertices(&Sample.Transform.Get(), &Sample.Parent.Get(), &Sample.BoneMap.Get(), &Sample.Vertex.Get(),
+				0, Sample.Vertex.Num(), ComponentSpaceFullVertices);
+			if (SampleScale == EDataflowTransferNodeSampleScale::Dataflow_Transfer_Asset_Edge)
+			{
+				SphereFullRadius = Chaos::FReal(EdgeMultiplier * FMath::Max(
+					MaxEdgeLength(ComponentSpaceFullTargetVertices, Target.Indices.Get(), 0, 0, Target.Indices.Num()),
+					MaxEdgeLength(ComponentSpaceFullVertices, Sample.Indices.Get(), 0, 0, Sample.Indices.Num())));
+			}
+			else if (SampleScale == EDataflowTransferNodeSampleScale::Dataflow_Transfer_Asset_Bound)
+			{
+				Chaos::TVec3<float> CoordMaxs(-FLT_MAX);
+				Chaos::TVec3<float> CoordMins(FLT_MAX);
+				for (int32 i = 0; i < ComponentSpaceFullVertices.Num(); i++)
+				{
+					for (int32 j = 0; j < 3; j++)
+					{
+						if (ComponentSpaceFullVertices[i][j] > CoordMaxs[j])
+						{
+							CoordMaxs[j] = ComponentSpaceFullVertices[i][j];
+						}
+						if (ComponentSpaceFullVertices[i][j] < CoordMins[j])
+						{
+							CoordMins[j] = ComponentSpaceFullVertices[i][j];
+						}
+					}
+				}
+				Chaos::TVec3<float> CoordDiff = (CoordMaxs - CoordMins) * BoundMultiplier;
+				SphereFullRadius = Chaos::FReal(FGenericPlatformMath::Min(CoordDiff[0], FGenericPlatformMath::Min(CoordDiff[1], CoordDiff[2])));
+			}
+		}
+		ParallelFor(PairedGeometry.Num(), [&](int32 Pdx)
+		{
+			int32 AttributeGeometryIndex = PairedGeometry[Pdx][0];
+			int32 TargetGeometryIndex = PairedGeometry[Pdx][1];
+			if (ensure(0 <= AttributeGeometryIndex && AttributeGeometryIndex < Sample.VertexStart.Num()))
+			{
+				if (ensure(0 <= TargetGeometryIndex && TargetGeometryIndex < Target.VertexStart.Num()))
+				{
+					// Build component space vertices for TargetCollection
+					TArray<FVector> ComponentSpaceTargetVertices; // size of num vertices of the geometry entry. 
+					BuildComponentSpaceVertices(&Target.Transform.Get(), &Target.Parent.Get(), &Target.BoneMap.Get(), &Target.Vertex.Get(),
+						Target.VertexStart[TargetGeometryIndex], Target.VertexCount[TargetGeometryIndex], ComponentSpaceTargetVertices);
+
+					// Build component space vertices for SampleCollection
+					TArray<FVector> ComponentSpaceVertices; // size of num vertices of the geometry entry
+					BuildComponentSpaceVertices(&Sample.Transform.Get(), &Sample.Parent.Get(), &Sample.BoneMap.Get(), &Sample.Vertex.Get(),
+						Sample.VertexStart[AttributeGeometryIndex], Sample.VertexCount[AttributeGeometryIndex], ComponentSpaceVertices);
+
+					// build Sphere based BVH
+					Chaos::FReal SphereRadius = SphereFullRadius;
+					if (SampleScale == EDataflowTransferNodeSampleScale::Dataflow_Transfer_Component_Edge)
+					{
+						SphereRadius = Chaos::FReal(EdgeMultiplier * FMath::Max(
+							MaxEdgeLength(ComponentSpaceTargetVertices, Target.Indices.Get(), Target.VertexStart[TargetGeometryIndex], Target.FaceStart[TargetGeometryIndex], Target.FaceCount[TargetGeometryIndex]),
+							MaxEdgeLength(ComponentSpaceVertices, Sample.Indices.Get(), Sample.VertexStart[AttributeGeometryIndex], Sample.FaceStart[AttributeGeometryIndex], Sample.FaceCount[AttributeGeometryIndex])));
+					}
+
+					int32 TargetVertexStartVal = Target.VertexStart[TargetGeometryIndex];
+					int32 TargetVertexCountVal = Target.VertexCount[TargetGeometryIndex];
+					int32 VertexStartVal = Sample.VertexStart[AttributeGeometryIndex];
+					int32 FaceStartVal = Sample.FaceStart[AttributeGeometryIndex];
+					int32 FaceCountVal = Sample.FaceCount[AttributeGeometryIndex];
+
+					if (BoundingVolumeType == EDataflowTransferNodeBoundingVolume::Dataflow_Transfer_Triangle)
+					{
+						TArray<Chaos::TVec3<Chaos::FReal>> ComponentSpaceVerticesTVec3;
+						ComponentSpaceVerticesTVec3.SetNum(ComponentSpaceVertices.Num());
+						for (int32 SourceIndex = 0; SourceIndex < ComponentSpaceVerticesTVec3.Num(); SourceIndex++)
+						{
+							ComponentSpaceVerticesTVec3[SourceIndex] = Chaos::TVec3<Chaos::FReal>(ComponentSpaceVertices[SourceIndex]);
+						}
+						TConstArrayView<Chaos::TVec3<Chaos::FReal>> ConstComponentSpaceVertices(ComponentSpaceVerticesTVec3);
+						Chaos::FTriangleMesh TriangleMesh;
+						TArray<Chaos::TVec3<int32>> SourceElements;
+						SourceElements.SetNum(FaceCountVal);
+						for (int32 ElementIndex = 0; ElementIndex < FaceCountVal; ++ElementIndex)
+						{
+							FIntVector3 Element = Sample.Indices[FaceStartVal + ElementIndex];
+							SourceElements[ElementIndex] = Chaos::TVec3<int32>(Element[0]-VertexStartVal, Element[1]-VertexStartVal, Element[2]-VertexStartVal);
+						}
+						TriangleMesh.Init(SourceElements);
+						Chaos::FTriangleMesh::TSpatialHashType<Chaos::FReal> SpatialHash;
+						TriangleMesh.BuildSpatialHash(ConstComponentSpaceVertices, SpatialHash, SphereRadius);
+						for (int32 TargetIndex = 0; TargetIndex < TargetVertexCountVal; TargetIndex++)
+						{
+							TArray<Chaos::TTriangleCollisionPoint<Chaos::FReal>> Result;
+							if (TriangleMesh.PointClosestTriangleQuery(SpatialHash, ConstComponentSpaceVertices,
+								TargetIndex, Chaos::TVec3<Chaos::FReal>(ComponentSpaceTargetVertices[TargetIndex]), SphereRadius / 2.f, SphereRadius / 2.f,
+								[](const int32 PointIndex, const int32 TriangleIndex)->bool {return true; }, Result))
+							{
+								for (const Chaos::TTriangleCollisionPoint<Chaos::FReal>& CollisionPoint : Result)
+								{
+									Chaos::FRealSingle CurrentDistance = abs(CollisionPoint.Phi);
+									float TriRadius = FalloffThreshold * MaxEdgeLength(ComponentSpaceVertices, Sample.Indices.Get(), VertexStartVal, FaceStartVal+CollisionPoint.Indices[1], 1);
+									float FalloffScale = CalculateFalloffScale(Falloff, TriRadius, CurrentDistance);
+									if (!FMath::IsNearlyZero(FalloffScale))
+									{
+										int32 TargetCandidateIndex = CollisionPoint.Indices[0]+ TargetVertexStartVal;
+										if (ensure(0 <= TargetCandidateIndex && TargetCandidateIndex < TargetFloatArray->Num()))
+										{
+											float Value = 0.f;
+											for (int32 k = 0; k < 3; k++)
+											{
+												Value += FalloffScale * (CollisionPoint.Bary[k] * (*FloatArray)[Sample.Indices[FaceStartVal+CollisionPoint.Indices[1]][k]]);
+											}
+											(*TargetFloatArray)[TargetCandidateIndex] = FMath::Max((*TargetFloatArray)[TargetCandidateIndex], Value);
+										}
+									}
+								}
+							}
+						}
+					}
+					else if (BoundingVolumeType == EDataflowTransferNodeBoundingVolume::Dataflow_Transfer_Vertex)
+					{
+						TUniquePtr<UE::Private::BVH> VertexBVH(BuildParticleSphereBVH(ComponentSpaceTargetVertices, SphereRadius));
+						for (int32 i = 0; i < FaceCountVal; i++)
+						{
+							FIntVector3 ComponentTriangle = Sample.Indices[FaceStartVal + i] - FIntVector3(VertexStartVal);
+							if (TriangleHasWeightsToTransfer(Sample.Indices[FaceStartVal + i], *FloatArray))
+							{
+								TArray<int32> TargetVertexIntersection({});
+								TriangleToVertexIntersections(*VertexBVH, ComponentSpaceVertices, ComponentTriangle, TargetVertexIntersection);
+
+								for (int32 j = 0; j < TargetVertexIntersection.Num(); j++)
+								{
+									Chaos::TVector<float, 3> Bary,
+										TriPos0(ComponentSpaceVertices[ComponentTriangle[0]]),
+										TriPos1(ComponentSpaceVertices[ComponentTriangle[1]]),
+										TriPos2(ComponentSpaceVertices[ComponentTriangle[2]]),
+										ParticlePos(ComponentSpaceTargetVertices[TargetVertexIntersection[j]]);
+
+									Chaos::TVector<Chaos::FRealSingle, 3> ClosestPoint = Chaos::FindClosestPointAndBaryOnTriangle(TriPos0, TriPos1, TriPos2, ParticlePos, Bary);
+									Chaos::FRealSingle CurrentDistance = (ParticlePos - ClosestPoint).Size();
+									float TriRadius = FalloffThreshold * MaxEdgeLength(ComponentSpaceVertices, Sample.Indices.Get(), VertexStartVal, FaceStartVal+i, 1);
+									float FalloffScale = CalculateFalloffScale(Falloff, TriRadius, CurrentDistance);
+									if (!FMath::IsNearlyZero(FalloffScale))
+									{
+										int32 TargetIndex = TargetVertexIntersection[j] + TargetVertexStartVal;
+										if (ensure(0 <= TargetIndex && TargetIndex < TargetFloatArray->Num()))
+										{
+											float Value = 0.f;
+											for (int32 k = 0; k < 3; k++)
+											{
+												Value += FalloffScale * (Bary[k] * (*FloatArray)[ComponentTriangle[k] + VertexStartVal]);
+											}
+											(*TargetFloatArray)[TargetIndex] = FMath::Max((*TargetFloatArray)[TargetIndex], Value);
+										}
+									}
+								}
+							}
+						}
+					}
+				}
+			}
+		});
+	}
+}
+
+void FGeometryCollectionTransferVertexScalarAttributeNode::NearestVertexTransfer(FCollectionAttributeKey Key, const UE::Private::FTransferFacade& Sample, UE::Private::FTransferFacade& Target, TManagedArray<float>* TargetFloatArray) const
+{
+	if (const TManagedArray<float>* FloatArray = Sample.GetFloatArray(Key.Attribute, Key.Group))
+	{
+		// Build component space vertices for TargetCollection
+		TArray<FVector> ComponentSpaceTargetVertices;
+		BuildComponentSpaceVertices(&Target.Transform.Get(), &Target.Parent.Get(), &Target.BoneMap.Get(), &Target.Vertex.Get(), 0, Target.Vertex.Num(), ComponentSpaceTargetVertices);
+
+		// Build component space vertices for SourceCollection
+		TArray<FVector> ComponentSpaceVertices;
+		BuildComponentSpaceVertices(&Sample.Transform.Get(), &Sample.Parent.Get(), &Sample.BoneMap.Get(), &Sample.Vertex.Get(), 0, Sample.Vertex.Num(), ComponentSpaceVertices);
+
+		// build Sphere based BVH
+		Chaos::FReal SphereRadius = Chaos::FReal(EdgeMultiplier * FMath::Max(
+			MaxEdgeLength(ComponentSpaceTargetVertices, Target.Indices.Get(), 0, 0, Target.Indices.Num()),
+			MaxEdgeLength(ComponentSpaceVertices, Sample.Indices.Get(), 0, 0, Sample.Indices.Num())));
+		TUniquePtr<UE::Private::BVH> VertexBVH(BuildParticleSphereBVH(ComponentSpaceTargetVertices, SphereRadius));
+
+		for (int32 i = 0; i < Sample.Indices.Num(); i++)
+		{
+			FIntVector3 Triangle = Sample.Indices[i];
+			if (TriangleHasWeightsToTransfer(Triangle, *FloatArray))
+			{
+				TArray<int32> TargetVertexIntersection({});
+				TriangleToVertexIntersections(*VertexBVH, ComponentSpaceVertices, Triangle, TargetVertexIntersection);
+
+				for (int32 j = 0; j < TargetVertexIntersection.Num(); j++)
+				{
+					Chaos::TVector<float, 3> Bary,
+						TriPos0(ComponentSpaceVertices[Sample.Indices[i][0]]),
+						TriPos1(ComponentSpaceVertices[Sample.Indices[i][1]]),
+						TriPos2(ComponentSpaceVertices[Sample.Indices[i][2]]),
+						ParticlePos(ComponentSpaceTargetVertices[TargetVertexIntersection[j]]);
+
+					Chaos::TVector<Chaos::FRealSingle, 3> ClosestPoint = Chaos::FindClosestPointAndBaryOnTriangle(TriPos0, TriPos1, TriPos2, ParticlePos, Bary);
+					Chaos::FRealSingle CurrentDistance = (ParticlePos - ClosestPoint).Size();
+					float TriRadius = FalloffThreshold * MaxEdgeLength(ComponentSpaceVertices, Sample.Indices.Get(), 0, i, 1);
+					float FalloffScale = CalculateFalloffScale(Falloff, TriRadius, CurrentDistance);
+					if (!FMath::IsNearlyZero(FalloffScale))
+					{
+						int32 TargetIndex = TargetVertexIntersection[j];
+						if (ensure(0 <= TargetIndex && TargetIndex < TargetFloatArray->Num()))
+						{
+							float Value = 0.f;
+							for (int32 k = 0; k < 3; k++)
+							{
+								Value += FalloffScale * (Bary[k] * (*FloatArray)[Sample.Indices[i][k]]);
+							}
+							(*TargetFloatArray)[TargetIndex] = FMath::Max((*TargetFloatArray)[TargetIndex], Value);
+						}
+					}
+				}
+			}
+		}
+	}
+}
+
+
+float FGeometryCollectionTransferVertexScalarAttributeNode::MaxEdgeLength(TArray<FVector>& Vert, const TManagedArray<FIntVector3>& Tri, int VertexOffset, int TriStart, int TriCount)
+{
+	auto TriInRange = [](const FIntVector3& T, int Max) {
+		for (int k = 0; k < 3; k++)
+		{
+			if (ensure(0 <= T[k] && T[k] < Max))
+			{
+				return true;
+			}
+		}
+		return false;
+	};
+
+	float Max = 0;
+	int TriStop = TriStart + TriCount;
+	for (int i = TriStart; i < TriStop; i++)
+	{
+		if (TriInRange(Tri[i] - FIntVector3(VertexOffset), Vert.Num()))
+		{
+			Max = FMath::Max(Max, FVector3f(Vert[Tri[i][0] - VertexOffset] - Vert[Tri[i][1] - VertexOffset]).SquaredLength());
+			Max = FMath::Max(Max, FVector3f(Vert[Tri[i][0] - VertexOffset] - Vert[Tri[i][2] - VertexOffset]).SquaredLength());
+			Max = FMath::Max(Max, FVector3f(Vert[Tri[i][1] - VertexOffset] - Vert[Tri[i][2] - VertexOffset]).SquaredLength());
+		}
+	}
+	return FMath::Sqrt(Max);
+}
+
+void FGeometryCollectionTransferVertexScalarAttributeNode::BuildComponentSpaceVertices(const TManagedArray<FTransform3f>* LocalSpaceTransform, const TManagedArray<int32>* Parent, const TManagedArray<int32>* BoneMapArray,
+	const TManagedArray<FVector3f>* VertexArray, int32 Start, int32 Count, TArray<FVector>& ComponentSpaceVertices)
+{
+	TArray<FTransform> ComponentTransform;
+	GeometryCollectionAlgo::GlobalMatrices(*LocalSpaceTransform, *Parent, ComponentTransform);
+
+	ComponentSpaceVertices.SetNumUninitialized(Count);
+	for (int i = 0; i < Count; i++)
+	{
+		int j = i + Start;
+		if (0 < (*BoneMapArray)[i] && (*BoneMapArray)[i] < ComponentTransform.Num())
+		{
+			ComponentSpaceVertices[i] = ComponentTransform[(*BoneMapArray)[j]].TransformPosition(FVector((*VertexArray)[j]));
+		}
+		else
+		{
+			ComponentSpaceVertices[i] = FVector((*VertexArray)[j]);
+		}
+	}
+}
+
+UE::Private::BVH* FGeometryCollectionTransferVertexScalarAttributeNode::BuildParticleSphereBVH(const TArray<FVector>& Vertices, float Radius)
+{
+	TArray<Chaos::TSphere<Chaos::FReal, 3>*> VertexSpherePtrs;
+	TArray<Chaos::TSphere<Chaos::FReal, 3>> VertexSpheres;
+
+	VertexSpheres.Init(Chaos::TSphere<Chaos::FReal, 3>(Chaos::TVec3<Chaos::FReal>(0), Radius), Vertices.Num());
+	VertexSpherePtrs.SetNum(Vertices.Num());
+
+	for (int32 i = 0; i < Vertices.Num(); i++)
+	{
+		Chaos::TVec3<Chaos::FReal> SphereCenter(Vertices[i]);
+		Chaos::TSphere<Chaos::FReal, 3> VertexSphere(SphereCenter, Radius);
+		VertexSpheres[i] = Chaos::TSphere<Chaos::FReal, 3>(SphereCenter, Radius);
+		VertexSpherePtrs[i] = &VertexSpheres[i];
+	}
+	return new UE::Private::BVH(VertexSpherePtrs);
+}
+
+bool FGeometryCollectionTransferVertexScalarAttributeNode::TriangleHasWeightsToTransfer(const FIntVector3& T, const TManagedArray<float>& F)
+{
+	return !FMath::IsNearlyZero((F[T[0]] + F[T[1]] + F[T[2]]));
+}
+
+void FGeometryCollectionTransferVertexScalarAttributeNode::TriangleToVertexIntersections(
+	UE::Private::BVH& VertexBVH, const TArray<FVector>& ComponentSpaceVertices, const FIntVector3& Triangle, TArray<int32>& OutTargetVertexIntersection)
+{
+	OutTargetVertexIntersection.Empty();
+
+	TArray<int32> TargetVertexIntersection0 = VertexBVH.FindAllIntersections(ComponentSpaceVertices[Triangle[0]]);
+	TArray<int32> TargetVertexIntersection1 = VertexBVH.FindAllIntersections(ComponentSpaceVertices[Triangle[1]]);
+	TArray<int32> TargetVertexIntersection2 = VertexBVH.FindAllIntersections(ComponentSpaceVertices[Triangle[2]]);
+	TargetVertexIntersection0.Sort();
+	TargetVertexIntersection1.Sort();
+	TargetVertexIntersection2.Sort();
+
+	for (int32 k = 0; k < TargetVertexIntersection0.Num(); k++)
+	{
+		if (TargetVertexIntersection1.Contains(TargetVertexIntersection0[k])
+			&& TargetVertexIntersection2.Contains(TargetVertexIntersection0[k]))
+		{
+			OutTargetVertexIntersection.Emplace(TargetVertexIntersection0[k]);
+		}
+	}
+}
+
+
+float FGeometryCollectionTransferVertexScalarAttributeNode::CalculateFalloffScale(EDataflowTransferNodeFalloff FalloffSetting, float Threshold, float Distance)
+{
+	float Denominator = 1.0;
+	if (Distance > Threshold && !FMath::IsNearlyZero(Threshold))
+	{
+		Denominator = Distance / Threshold;
+	}
+	switch (FalloffSetting)
+	{
+	case EDataflowTransferNodeFalloff::Dataflow_Transfer_Linear:
+		return 1. / Denominator;
+	case EDataflowTransferNodeFalloff::Dataflow_Transfer_Squared:
+		return 1. / FMath::Square(Denominator);
+	}
+	return 1.0;
+}
+
+
+#undef LOCTEXT_NAMESPACE
diff --git a/Engine/Plugins/Experimental/GeometryCollectionPlugin/Source/GeometryCollectionDepNodes/Private/Dataflow/GeometryCollectionTransferVertexScalarAttributeDepNode.h b/Engine/Plugins/Experimental/GeometryCollectionPlugin/Source/GeometryCollectionDepNodes/Private/Dataflow/GeometryCollectionTransferVertexScalarAttributeDepNode.h
new file mode 100644
--- /dev/null
+++ b/Engine/Plugins/Experimental/GeometryCollectionPlugin/Source/GeometryCollectionDepNodes/Private/Dataflow/GeometryCollectionTransferVertexScalarAttributeDepNode.h
@@ -0,0 +1,148 @@
+// Copyright Epic Games, Inc. All Rights Reserved.
+
+#pragma once 
+
+#include "CoreMinimal.h"
+#include "Chaos/BoundingVolumeHierarchy.h"
+#include "Dataflow/DataflowNode.h"
+#include "Dataflow/DataflowConnectionTypes.h"
+#include "GeometryCollection/GeometryCollection.h"
+
+#include "GeometryCollectionTransferVertexScalarAttributeDepNode.generated.h"
+
+
+namespace UE::Private 
+{
+	class FTransferFacade;
+	typedef Chaos::TBoundingVolumeHierarchy<TArray<Chaos::TSphere<Chaos::FReal, 3>*>, TArray<int32>, Chaos::FReal, 3> BVH;
+};
+
+UENUM(BlueprintType)
+enum class EDataflowTransferNodeFalloff : uint8
+{
+	/** Squared falloff based on distance from triangle*/
+	Dataflow_Transfer_Squared  UMETA(DisplayName = "Squared"),
+
+	/** Linear falloff based on distance from triangle*/
+	Dataflow_Transfer_Linear UMETA(DisplayName = "Linear"),
+
+	/** No distance falloff */
+	Dataflow_Transfer_None UMETA(DisplayName = "None"),
+	//~~~
+	//256th entry
+	Dataflow_Transfer_Dataflow_Max UMETA(Hidden)
+};
+
+UENUM(BlueprintType)
+enum class EDataflowTransferNodeSampleScale : uint8
+{
+	/** Bounding volume hierarchy cell size based on max edge length of each geometry group*/
+	Dataflow_Transfer_Component_Edge  UMETA(DisplayName = "Component Max Edge"),
+
+	/** Bounding volume hierarchy cell size based on max edge length of the whole asset*/
+	Dataflow_Transfer_Asset_Edge UMETA(DisplayName = "Asset Max Edge"),
+
+	/** Bounding volume hierarchy cell size based on max length of the bounding box of the whole asset*/
+	Dataflow_Transfer_Asset_Bound UMETA(DisplayName = "Asset Max Bound"),
+
+	//~~~
+	//256th entry
+	Dataflow_Transfer_Dataflow_Max UMETA(Hidden)
+};
+
+UENUM(BlueprintType)
+enum class EDataflowTransferNodeBoundingVolume : uint8
+{
+	/** Bounding volume on vertices of the source triangle mesh*/
+	Dataflow_Transfer_Vertex  UMETA(DisplayName = "Vertex"),
+
+	/** Bounding volume on triangles of the source triangle mesh*/
+	Dataflow_Transfer_Triangle UMETA(DisplayName = "Triangle"),
+
+	//~~~
+	//256th entry
+	Dataflow_Transfer_Dataflow_Max UMETA(Hidden)
+};
+
+/**
+ * Transfer float properties from a sample collection to a target collection. 
+ */
+USTRUCT(meta = (Deprecated = "5.5"))
+struct FGeometryCollectionTransferVertexScalarAttributeNode : public FDataflowNode
+{
+	GENERATED_USTRUCT_BODY()
+	DATAFLOW_NODE_DEFINE_INTERNAL(FGeometryCollectionTransferVertexScalarAttributeNode, "TransferVertexScalarAttribute", "GeometryCollection", "Transfer a named scalar attribute from the Source Collection to the Target Collection")
+	DATAFLOW_NODE_RENDER_TYPE("SurfaceWeightsRender", FGeometryCollection::StaticType(), "Collection", "AttributeKey")
+
+public:
+
+	/* Target collection to transfer vertex float attribute to. */
+	UPROPERTY(Meta = (DataflowInput, DataflowOutput, DisplayName = "Collection", DataflowPassthrough = "Collection"))
+	FManagedArrayCollection Collection;
+
+	/* Sample collection to transfer vertex float attribute from. */
+	UPROPERTY(Meta = (DataflowInput, DisplayName = "FromCollection"))
+	FManagedArrayCollection FromCollection;
+
+	/* The name of the vertex attribute to generate indices from. */
+	UPROPERTY(EditAnywhere, Category = "Dataflow", Meta = (DataflowInput, DataflowOutput, DisplayName = "AttributeKey", DataflowPassthrough = "AttributeKey"))
+	FCollectionAttributeKey AttributeKey;
+
+	/* Bounding volume type for source assets[default: Triangle] */
+	UPROPERTY(EditAnywhere, Category = "Thresholds")
+	EDataflowTransferNodeBoundingVolume BoundingVolumeType = EDataflowTransferNodeBoundingVolume::Dataflow_Transfer_Triangle;
+
+	/* Bounding volume hierarchy cell size for neighboring vertices to transfer into[default: Asset] */
+	UPROPERTY(EditAnywhere, Category = "Thresholds")
+	EDataflowTransferNodeSampleScale SampleScale = EDataflowTransferNodeSampleScale::Dataflow_Transfer_Asset_Bound;
+
+	/* Falloff of sample value based on distance from sample triangle[default: Squared] */
+	UPROPERTY(EditAnywhere, Category = "Thresholds")
+	EDataflowTransferNodeFalloff Falloff = EDataflowTransferNodeFalloff::Dataflow_Transfer_Squared;
+
+	/* Threshold based on distance from sample triangle.Values sampled past the threshold will falloff.[Defaults to 1 percent of triangle size(0.01)] */
+	UPROPERTY(EditAnywhere, Category = "Thresholds")
+	float FalloffThreshold = 0.01f;
+
+	/* Edge multiplier for the Bounding Volume Hierarchy(BVH) target's particle search radius. */
+	UPROPERTY(EditAnywhere, Category = "Thresholds", meta = (EditCondition = "SampleScale == EDataflowTransferNodeSampleScale::Dataflow_Transfer_Asset_Edge || SampleScale == EDataflowTransferNodeSampleScale::Dataflow_Transfer_Component_Edge", EditConditionHides))
+	float EdgeMultiplier = 0.5f;
+
+	/* Max bound multiplier for the Bounding Volume Hierarchy(BVH) target's particle search radius. */
+	UPROPERTY(EditAnywhere, Category = "Thresholds", meta = (
+		EditCondition = "SampleScale == EDataflowTransferNodeSampleScale::Dataflow_Transfer_Asset_Bound", 
+		EditConditionHides))
+	float BoundMultiplier = 0.01f;
+
+	FGeometryCollectionTransferVertexScalarAttributeNode(const Dataflow::FNodeParameters& InParam, FGuid InGuid = FGuid::NewGuid())
+		: FDataflowNode(InParam, InGuid)
+	{
+		RegisterInputConnection(&Collection);
+		RegisterInputConnection(&FromCollection);
+		RegisterInputConnection(&AttributeKey);
+
+		RegisterOutputConnection(&Collection, &Collection);
+		RegisterOutputConnection(&AttributeKey, &AttributeKey);
+	}
+
+private:
+
+	virtual void Evaluate(Dataflow::FContext& Context, const FDataflowOutput* Out) const override;
+
+
+	//
+	// private helpers
+	//
+	TArray<FIntVector2> FindSourceToTargetGeometryMap(const FManagedArrayCollection& SourceCollection, const FManagedArrayCollection& TargetCollection) const;
+	void PairedGeometryTransfer(FCollectionAttributeKey Key, const TArray<FIntVector2>& PairedGeometry,const UE::Private::FTransferFacade& Sample, UE::Private::FTransferFacade& Target, TManagedArray<float>* TargetFloatArray) const;
+	void NearestVertexTransfer(FCollectionAttributeKey Key, const UE::Private::FTransferFacade& Sample, UE::Private::FTransferFacade& Target, TManagedArray<float>* TargetFloatArray) const;
+	static float MaxEdgeLength(TArray<FVector>& Vert, const TManagedArray<FIntVector3>& Tri, int VertexOFfset, int TriStart, int TriCount);
+	static void BuildComponentSpaceVertices(const TManagedArray<FTransform3f>* LocalSpaceTransform, const TManagedArray<int32>* Parent, const TManagedArray<int32>* BoneMapArray, const TManagedArray<FVector3f>* VertexArray, int32 Start, int32 Count, TArray<FVector>& ComponentSpaceVertices);
+	static UE::Private::BVH* BuildParticleSphereBVH(const TArray<FVector>& Vertices, float Radius);
+	static bool TriangleHasWeightsToTransfer(const FIntVector3& T, const TManagedArray<float>& F);
+	static void TriangleToVertexIntersections(UE::Private::BVH& VertexBVH, const TArray<FVector>& ComponentSpaceVertices, const FIntVector3& Triangle, TArray<int32>& OutTargetVertexIntersection);
+	static float CalculateFalloffScale(EDataflowTransferNodeFalloff FalloffSetting, float Threshold, float Distance);
+
+};
+
+
diff --git a/Engine/Plugins/Experimental/GeometryCollectionPlugin/Source/GeometryCollectionDepNodes/Private/Dataflow/SetVertexColorFromFloatArrayDepNode.h b/Engine/Plugins/Experimental/GeometryCollectionPlugin/Source/GeometryCollectionDepNodes/Private/Dataflow/SetVertexColorFromFloatArrayDepNode.h
new file mode 100644
--- /dev/null
+++ b/Engine/Plugins/Experimental/GeometryCollectionPlugin/Source/GeometryCollectionDepNodes/Private/Dataflow/SetVertexColorFromFloatArrayDepNode.h
@@ -0,0 +1,42 @@
+// Copyright Epic Games, Inc. All Rights Reserved.
+
+#pragma once 
+
+#include "CoreMinimal.h"
+#include "Dataflow/DataflowEngine.h"
+#include "GeometryCollection/ManagedArrayCollection.h"
+
+
+#include "SetVertexColorFromFloatArrayDepNode.generated.h"
+
+/** Set collection color based on input float array */
+USTRUCT(meta = (Deprecated = "5.5"))
+struct FSetVertexColorInCollectionFromFloatArrayDataflowNode : public FDataflowNode
+{
+	GENERATED_USTRUCT_BODY()
+		DATAFLOW_NODE_DEFINE_INTERNAL(FSetVertexColorInCollectionFromFloatArrayDataflowNode, "SetVertexColorInCollectionFromFloatArray", "Collection|Utilities", "")
+
+public:
+	/** Collection */
+	UPROPERTY(meta = (DataflowInput, DataflowOutput, DataflowPassthrough = "Collection", DataflowIntrinsic))
+		FManagedArrayCollection Collection;
+
+	/**  */
+	UPROPERTY(meta = (DataflowInput, DataflowIntrinsic))
+		TArray<float> FloatArray;
+
+	/**  */
+	UPROPERTY(EditAnywhere, Category = "Color")
+		float Scale = 1.f;
+
+	FSetVertexColorInCollectionFromFloatArrayDataflowNode(const Dataflow::FNodeParameters& InParam, FGuid InGuid = FGuid::NewGuid())
+		: FDataflowNode(InParam, InGuid)
+	{
+		RegisterInputConnection(&Collection);
+		RegisterInputConnection(&FloatArray);
+		RegisterOutputConnection(&Collection, &Collection);
+	}
+
+	virtual void Evaluate(Dataflow::FContext& Context, const FDataflowOutput* Out) const override;
+
+};
\ No newline at end of file
diff --git a/Engine/Plugins/Experimental/GeometryCollectionPlugin/Source/GeometryCollectionDepNodes/Private/Dataflow/SetVertexColorFromVertexSelectionDepNode.h b/Engine/Plugins/Experimental/GeometryCollectionPlugin/Source/GeometryCollectionDepNodes/Private/Dataflow/SetVertexColorFromVertexSelectionDepNode.h
new file mode 100644
--- /dev/null
+++ b/Engine/Plugins/Experimental/GeometryCollectionPlugin/Source/GeometryCollectionDepNodes/Private/Dataflow/SetVertexColorFromVertexSelectionDepNode.h
@@ -0,0 +1,46 @@
+// Copyright Epic Games, Inc. All Rights Reserved.
+
+#pragma once 
+
+#include "CoreMinimal.h"
+#include "Dataflow/DataflowEngine.h"
+#include "Dataflow/DataflowSelection.h"
+#include "GeometryCollection/ManagedArrayCollection.h"
+
+#include "SetVertexColorFromVertexSelectionDepNode.generated.h"
+
+/** Set the collections vertex color from the selection set. */
+USTRUCT(meta = (Deprecated = "5.5"))
+struct FSetVertexColorInCollectionFromVertexSelectionDataflowNode : public FDataflowNode
+{
+	GENERATED_USTRUCT_BODY()
+	DATAFLOW_NODE_DEFINE_INTERNAL(FSetVertexColorInCollectionFromVertexSelectionDataflowNode, "SetVertexColorInCollectionFromVertexSelection", "Collection|Utilities", "")
+
+public:
+	/** Collection */
+	UPROPERTY(meta = (DataflowInput, DataflowOutput, DataflowPassthrough = "Collection", DataflowIntrinsic))
+		FManagedArrayCollection Collection;
+
+	/**  */
+	UPROPERTY(meta = (DataflowInput, DisplayName = "VertexSelection", DataflowIntrinsic))
+		FDataflowVertexSelection VertexSelection;
+
+	/**  */
+	UPROPERTY(EditAnywhere, Category = "Color")
+		FLinearColor SelectedColor = FLinearColor(FColor::Yellow);
+
+	/**  */
+	UPROPERTY(EditAnywhere, Category = "Color", meta = (DisplayName = "NonSelected Color"))
+		FLinearColor NonSelectedColor = FLinearColor(FColor::Blue);
+
+	FSetVertexColorInCollectionFromVertexSelectionDataflowNode(const Dataflow::FNodeParameters& InParam, FGuid InGuid = FGuid::NewGuid())
+		: FDataflowNode(InParam, InGuid)
+	{
+		RegisterInputConnection(&Collection);
+		RegisterInputConnection(&VertexSelection);
+		RegisterOutputConnection(&Collection, &Collection);
+	}
+
+	virtual void Evaluate(Dataflow::FContext& Context, const FDataflowOutput* Out) const override;
+
+};
diff --git a/Engine/Plugins/Experimental/GeometryCollectionPlugin/Source/GeometryCollectionDepNodes/Public/Dataflow/SetVertexColorFromFloatArrayDepNode.h b/Engine/Plugins/Experimental/GeometryCollectionPlugin/Source/GeometryCollectionDepNodes/Public/Dataflow/SetVertexColorFromFloatArrayDepNode.h
deleted file mode 100644
diff --git a/Engine/Plugins/Experimental/GeometryCollectionPlugin/Source/GeometryCollectionDepNodes/Public/Dataflow/SetVertexColorFromVertexSelectionDepNode.h b/Engine/Plugins/Experimental/GeometryCollectionPlugin/Source/GeometryCollectionDepNodes/Public/Dataflow/SetVertexColorFromVertexSelectionDepNode.h
deleted file mode 100644
diff --git a/Engine/Plugins/Experimental/GeometryCollectionPlugin/Source/GeometryCollectionNodes/Private/Dataflow/GeometryCollectionNodesPlugin.cpp b/Engine/Plugins/Experimental/GeometryCollectionPlugin/Source/GeometryCollectionNodes/Private/Dataflow/GeometryCollectionNodesPlugin.cpp
--- a/Engine/Plugins/Experimental/GeometryCollectionPlugin/Source/GeometryCollectionNodes/Private/Dataflow/GeometryCollectionNodesPlugin.cpp
+++ b/Engine/Plugins/Experimental/GeometryCollectionPlugin/Source/GeometryCollectionNodes/Private/Dataflow/GeometryCollectionNodesPlugin.cpp
@@ -26,7 +26,7 @@
 #include "Dataflow/GeometryCollectionSkeletalMeshToCollectionNode.h"
 #include "Dataflow/GeometryCollectionSkeletonToCollectionNode.h"
 #include "Dataflow/GeometryCollectionStaticMeshToCollectionNode.h"
-#include "Dataflow/GeometryCollectionTransferVertexScalarAttributeNode.h"
+#include "Dataflow/GeometryCollectionTransferVertexAttributeNode.h"
 #include "Dataflow/GeometryCollectionTriangleBoundaryIndicesNode.h"
 #include "Dataflow/GeometryCollectionUtilityNodes.h"
 #include "Dataflow/GeometryCollectionVerticesNodes.h"
@@ -66,7 +66,7 @@ void IGeometryCollectionNodesPlugin::StartupModule()
 	DATAFLOW_NODE_REGISTER_CREATION_FACTORY(FSetVertexColorFromVertexSelectionDataflowNode);
 	DATAFLOW_NODE_REGISTER_CREATION_FACTORY(FSkeletonToCollectionDataflowNode);
 	DATAFLOW_NODE_REGISTER_CREATION_FACTORY(FSkeletalMeshToCollectionDataflowNode);
-	DATAFLOW_NODE_REGISTER_CREATION_FACTORY(FGeometryCollectionTransferVertexScalarAttributeNode);
+	DATAFLOW_NODE_REGISTER_CREATION_FACTORY(FGeometryCollectionTransferVertexAttributeNode);
 	DATAFLOW_NODE_REGISTER_CREATION_FACTORY(FStaticMeshToCollectionDataflowNode);
 	DATAFLOW_NODE_REGISTER_CREATION_FACTORY(FTriangleBoundaryIndicesNode);
 }
diff --git a/Engine/Plugins/Experimental/GeometryCollectionPlugin/Source/GeometryCollectionNodes/Private/Dataflow/GeometryCollectionTransferVertexAttributeNode.cpp b/Engine/Plugins/Experimental/GeometryCollectionPlugin/Source/GeometryCollectionNodes/Private/Dataflow/GeometryCollectionTransferVertexAttributeNode.cpp
new file mode 100644
--- /dev/null
+++ b/Engine/Plugins/Experimental/GeometryCollectionPlugin/Source/GeometryCollectionNodes/Private/Dataflow/GeometryCollectionTransferVertexAttributeNode.cpp
@@ -0,0 +1,502 @@
+// Copyright Epic Games, Inc. All Rights Reserved.
+
+#include "Dataflow/GeometryCollectionTransferVertexAttributeNode.h"
+
+#include "Chaos/Triangle.h"
+#include "Chaos/TriangleMesh.h"
+#include "Chaos/HierarchicalSpatialHash.h"
+#include "Chaos/TriangleCollisionPoint.h"
+#include "Dataflow/DataflowInputOutput.h"
+#include "GeometryCollection/GeometryCollectionAlgo.h"
+#include "GeometryCollection/ManagedArrayAccessor.h"
+
+#include UE_INLINE_GENERATED_CPP_BY_NAME(GeometryCollectionTransferVertexAttributeNode)
+#define LOCTEXT_NAMESPACE "FGeometryCollectionTransferVertexAttributeNode"
+
+namespace UE::Private {
+
+	class FTransferFacade
+	{
+		const FManagedArrayCollection& ConstCollection;
+		FManagedArrayCollection* Collection = nullptr;
+	public:
+		FTransferFacade(FManagedArrayCollection& InCollection)
+			: ConstCollection(InCollection)
+			, Collection(&InCollection)
+			, BoneMap(InCollection, FName("BoneMap"), FName("Vertices"))
+			, Vertex(InCollection, FName("Vertex"), FName("Vertices"))
+			, Indices(InCollection, FName("Indices"), FName("Faces"))
+			, Transform(InCollection, FTransformCollection::TransformAttribute, FTransformCollection::TransformGroup)
+			, Parent(InCollection, FTransformCollection::ParentAttribute, FTransformCollection::TransformGroup)
+			, VertexStart(InCollection, "VertexStart", FGeometryCollection::GeometryGroup)
+			, VertexCount(InCollection, "VertexCount", FGeometryCollection::GeometryGroup)
+			, FaceStart(InCollection, "FaceStart", FGeometryCollection::GeometryGroup)
+			, FaceCount(InCollection, "FaceCount", FGeometryCollection::GeometryGroup)
+		{}
+
+		FTransferFacade(const FManagedArrayCollection& InCollection)
+			: ConstCollection(InCollection)
+			, Collection(nullptr)
+			, BoneMap(InCollection, FName("BoneMap"), FName("Vertices"))
+			, Vertex(InCollection, FName("Vertex"), FName("Vertices"))
+			, Indices(InCollection, FName("Indices"), FName("Faces"))
+			, Transform(InCollection, FTransformCollection::TransformAttribute, FTransformCollection::TransformGroup)
+			, Parent(InCollection, FTransformCollection::ParentAttribute, FTransformCollection::TransformGroup)
+			, VertexStart(InCollection, "VertexStart", FGeometryCollection::GeometryGroup)
+			, VertexCount(InCollection, "VertexCount", FGeometryCollection::GeometryGroup)
+			, FaceStart(InCollection, "FaceStart", FGeometryCollection::GeometryGroup)
+			, FaceCount(InCollection, "FaceCount", FGeometryCollection::GeometryGroup)
+		{}
+
+
+		bool IsValid() const {
+			return BoneMap.IsValid() && Vertex.IsValid() && Indices.IsValid() && Transform.IsValid() && Parent.IsValid() && VertexStart.IsValid() &&
+				VertexCount.IsValid() && FaceStart.IsValid() && FaceCount.IsValid();
+		}
+
+		template<typename T>
+		const TManagedArray<T>* GetAttributeArray(FString AttributeName, FString Group) const
+		{
+			return ConstCollection.FindAttributeTyped<T>(FName(AttributeName), FName(Group));
+		}
+
+		template<typename T>
+		TManagedArray<T>* GetAttributeArray(FString AttributeName, FString Group)
+		{
+			TManagedArray<T>* TargetAttributeArray = nullptr;
+			if (!Collection->HasAttribute(FName(AttributeName), FName(Group)))
+			{
+				Collection->AddAttribute<T>(FName(AttributeName), FName(Group));
+			}
+			return Collection->FindAttributeTyped<T>(FName(AttributeName), FName(Group));
+		}
+
+
+		TManagedArrayAccessor<int32> BoneMap;
+		TManagedArrayAccessor<FVector3f> Vertex;
+		TManagedArrayAccessor<FIntVector3> Indices;
+		TManagedArrayAccessor<FTransform3f> Transform;
+		TManagedArrayAccessor<int32> Parent;
+		TManagedArrayAccessor<int32> VertexStart;
+		TManagedArrayAccessor<int32> VertexCount;
+		TManagedArrayAccessor<int32> FaceStart;
+		TManagedArrayAccessor<int32> FaceCount;
+	};
+
+}
+
+void FGeometryCollectionTransferVertexAttributeNode::Evaluate(Dataflow::FContext& Context, const FDataflowOutput* Out) const
+{
+	FCollectionAttributeKey Key = GetValue(Context, &AttributeKey, AttributeKey);
+
+	if (Out->IsA(&Collection))
+	{
+		FManagedArrayCollection TargetCollection = GetValue(Context, &Collection);
+		const FManagedArrayCollection& SampleCollection = GetValue(Context, &FromCollection);
+
+		UE::Private::FTransferFacade Target(TargetCollection);
+		const UE::Private::FTransferFacade Sample(SampleCollection);
+
+		if (Target.IsValid() && Sample.IsValid())
+		{
+			if (const TManagedArray<float>* SourceAttributeFloatArray = Sample.GetAttributeArray<float>(Key.Attribute, Key.Group))
+			{
+				TManagedArray<float>* TargetAttributeArray = Target.GetAttributeArray<float>(Key.Attribute, Key.Group);
+				TargetAttributeArray->Fill(0.f);
+
+				TArray<FIntVector2> AlignedGeometry = FindSourceToTargetGeometryMap(SampleCollection, TargetCollection);
+				if (AlignedGeometry.Num() == TargetCollection.NumElements(FGeometryCollection::GeometryGroup))
+				{
+					PairedGeometryTransfer<float>(Key, AlignedGeometry, Sample, Target, SourceAttributeFloatArray, TargetAttributeArray);
+				}
+				else
+				{
+					NearestVertexTransfer<float>(Key, Sample, Target, SourceAttributeFloatArray, TargetAttributeArray);
+				}
+			}
+			else if (const TManagedArray<FLinearColor>* SourceAttributeColorArray = Sample.GetAttributeArray<FLinearColor>(Key.Attribute, Key.Group))
+			{
+				TManagedArray<FLinearColor>* TargetAttributeArray = Target.GetAttributeArray<FLinearColor>(Key.Attribute, Key.Group);
+				TargetAttributeArray->Fill(FLinearColor(0, 0, 0, 0));
+
+				TArray<FIntVector2> AlignedGeometry = FindSourceToTargetGeometryMap(SampleCollection, TargetCollection);
+				if (AlignedGeometry.Num() == TargetCollection.NumElements(FGeometryCollection::GeometryGroup))
+				{
+					PairedGeometryTransfer<FLinearColor>(Key, AlignedGeometry, Sample, Target, SourceAttributeColorArray, TargetAttributeArray);
+				}
+				else
+				{
+					NearestVertexTransfer<FLinearColor>(Key, Sample, Target, SourceAttributeColorArray, TargetAttributeArray);
+				}
+			}
+		}
+
+		SetValue(Context, MoveTemp(TargetCollection), &Collection);
+	}
+	else if (Out->IsA(&AttributeKey))
+	{
+		SetValue(Context, MoveTemp(Key), &AttributeKey);
+	}
+}
+
+
+TArray<FIntVector2> FGeometryCollectionTransferVertexAttributeNode::FindSourceToTargetGeometryMap(const FManagedArrayCollection& AttributeCollectionVal, const FManagedArrayCollection& CollectionVal) const
+{
+	TArray<FIntVector2> Mapping;
+	const TManagedArray<FString>* SourceName = AttributeCollectionVal.FindAttribute<FString>(FName("BoneName"), FTransformCollection::TransformGroup);
+	const TManagedArray<int32>* SourceGeometryGroup = AttributeCollectionVal.FindAttribute<int32>(FName("TransformToGeometryIndex"), FTransformCollection::TransformGroup);
+	const TManagedArray<FString>* TargetName = CollectionVal.FindAttribute<FString>(FName("BoneName"), FTransformCollection::TransformGroup);
+	const TManagedArray<int32>* TargetGeometryGroup = CollectionVal.FindAttribute<int32>(FName("TransformToGeometryIndex"), FTransformCollection::TransformGroup);
+	if (SourceName && SourceGeometryGroup && TargetName && TargetGeometryGroup)
+	{
+		for (int i = 0; i < SourceName->Num(); i++)
+		{
+			for (int j = 0; j < TargetName->Num(); j++)
+			{
+				FString TestName = FString::Printf(TEXT("%s_Tet"), *(*SourceName)[i]);
+				if ((*TargetName)[j].StartsWith(TestName))
+				{
+					Mapping.Add(FIntVector2((*SourceGeometryGroup)[i], (*TargetGeometryGroup)[j]));
+					break;
+				}
+			}
+		}
+	}
+
+	return Mapping;
+}
+
+template<typename T>
+void FGeometryCollectionTransferVertexAttributeNode::PairedGeometryTransfer(FCollectionAttributeKey Key, const TArray<FIntVector2>& PairedGeometry,
+	const UE::Private::FTransferFacade& Sample, UE::Private::FTransferFacade& Target, const TManagedArray<T>* SourceAttributeArray, TManagedArray<T>* TargetAttributeArray) const
+{
+	//TargetAttributeArray has to be filled with empty values before (attribute types: float, FLinearColor)
+	if (SourceAttributeArray && TargetAttributeArray)
+	{
+		Chaos::FReal SphereFullRadius;
+
+		if (SampleScale == EDataflowTransferVertexAttributeNodeSampleScale::Asset_Edge || SampleScale == EDataflowTransferVertexAttributeNodeSampleScale::Asset_Bound)
+		{
+			// Build component space vertices for TargetCollection
+			TArray<FVector3f> ComponentSpaceFullTargetVertices; // size of num vertices of the geometry entry. 
+			BuildComponentSpaceVertices(Target.Transform.Get(), Target.Parent.Get(), Target.BoneMap.Get(), Target.Vertex.Get(),
+				0, Target.Vertex.Num(), ComponentSpaceFullTargetVertices);
+
+			// Build component space vertices for TargetCollection
+			TArray<FVector3f> ComponentSpaceFullVertices; // size of num vertices of the geometry entry
+			BuildComponentSpaceVertices(Sample.Transform.Get(), Sample.Parent.Get(), Sample.BoneMap.Get(), Sample.Vertex.Get(),
+				0, Sample.Vertex.Num(), ComponentSpaceFullVertices);
+			if (SampleScale == EDataflowTransferVertexAttributeNodeSampleScale::Asset_Edge)
+			{
+				SphereFullRadius = Chaos::FReal(EdgeMultiplier * FMath::Max(
+					MaxEdgeLength(ComponentSpaceFullTargetVertices, Target.Indices.Get(), 0, 0, Target.Indices.Num()),
+					MaxEdgeLength(ComponentSpaceFullVertices, Sample.Indices.Get(), 0, 0, Sample.Indices.Num())));
+			}
+			else if (SampleScale == EDataflowTransferVertexAttributeNodeSampleScale::Asset_Bound)
+			{
+				Chaos::TVec3<float> CoordMaxs(-FLT_MAX);
+				Chaos::TVec3<float> CoordMins(FLT_MAX);
+				for (int32 i = 0; i < ComponentSpaceFullVertices.Num(); i++)
+				{
+					for (int32 j = 0; j < 3; j++)
+					{
+						if (ComponentSpaceFullVertices[i][j] > CoordMaxs[j])
+						{
+							CoordMaxs[j] = ComponentSpaceFullVertices[i][j];
+						}
+						if (ComponentSpaceFullVertices[i][j] < CoordMins[j])
+						{
+							CoordMins[j] = ComponentSpaceFullVertices[i][j];
+						}
+					}
+				}
+				Chaos::TVec3<float> CoordDiff = (CoordMaxs - CoordMins) * BoundMultiplier;
+				SphereFullRadius = Chaos::FReal(FGenericPlatformMath::Min(CoordDiff[0], FGenericPlatformMath::Min(CoordDiff[1], CoordDiff[2])));
+			}
+		}
+		ParallelFor(PairedGeometry.Num(), [&](int32 Pdx)
+			{
+				int32 AttributeGeometryIndex = PairedGeometry[Pdx][0];
+				int32 TargetGeometryIndex = PairedGeometry[Pdx][1];
+				if (ensure(0 <= AttributeGeometryIndex && AttributeGeometryIndex < Sample.VertexStart.Num()))
+				{
+					if (ensure(0 <= TargetGeometryIndex && TargetGeometryIndex < Target.VertexStart.Num()))
+					{
+						// Build component space vertices for TargetCollection
+						TArray<FVector3f> ComponentSpaceTargetVertices; // size of num vertices of the geometry entry. 
+						BuildComponentSpaceVertices(Target.Transform.Get(), Target.Parent.Get(), Target.BoneMap.Get(), Target.Vertex.Get(),
+							Target.VertexStart[TargetGeometryIndex], Target.VertexCount[TargetGeometryIndex], ComponentSpaceTargetVertices);
+
+						// Build component space vertices for SampleCollection
+						TArray<FVector3f> ComponentSpaceVertices; // size of num vertices of the geometry entry
+						BuildComponentSpaceVertices(Sample.Transform.Get(), Sample.Parent.Get(), Sample.BoneMap.Get(), Sample.Vertex.Get(),
+							Sample.VertexStart[AttributeGeometryIndex], Sample.VertexCount[AttributeGeometryIndex], ComponentSpaceVertices);
+
+						// build Sphere based BVH
+						Chaos::FReal SphereRadius = SphereFullRadius;
+						if (SampleScale == EDataflowTransferVertexAttributeNodeSampleScale::Component_Edge)
+						{
+							SphereRadius = Chaos::FReal(EdgeMultiplier * FMath::Max(
+								MaxEdgeLength(ComponentSpaceTargetVertices, Target.Indices.Get(), Target.VertexStart[TargetGeometryIndex], Target.FaceStart[TargetGeometryIndex], Target.FaceCount[TargetGeometryIndex]),
+								MaxEdgeLength(ComponentSpaceVertices, Sample.Indices.Get(), Sample.VertexStart[AttributeGeometryIndex], Sample.FaceStart[AttributeGeometryIndex], Sample.FaceCount[AttributeGeometryIndex])));
+						}
+
+						int32 TargetVertexStartVal = Target.VertexStart[TargetGeometryIndex];
+						int32 TargetVertexCountVal = Target.VertexCount[TargetGeometryIndex];
+						int32 VertexStartVal = Sample.VertexStart[AttributeGeometryIndex];
+						int32 FaceStartVal = Sample.FaceStart[AttributeGeometryIndex];
+						int32 FaceCountVal = Sample.FaceCount[AttributeGeometryIndex];
+
+						if (BoundingVolumeType == EDataflowTransferVertexAttributeNodeBoundingVolume::Triangle)
+						{
+							TArray<Chaos::TVec3<Chaos::FReal>> ComponentSpaceVerticesTVec3;
+							ComponentSpaceVerticesTVec3.SetNum(ComponentSpaceVertices.Num());
+							for (int32 SourceIndex = 0; SourceIndex < ComponentSpaceVerticesTVec3.Num(); SourceIndex++)
+							{
+								ComponentSpaceVerticesTVec3[SourceIndex] = Chaos::TVec3<Chaos::FReal>(ComponentSpaceVertices[SourceIndex]);
+							}
+							TConstArrayView<Chaos::TVec3<Chaos::FReal>> ConstComponentSpaceVertices(ComponentSpaceVerticesTVec3);
+							Chaos::FTriangleMesh TriangleMesh;
+							TArray<Chaos::TVec3<int32>> SourceElements;
+							SourceElements.SetNum(FaceCountVal);
+							for (int32 ElementIndex = 0; ElementIndex < FaceCountVal; ++ElementIndex)
+							{
+								FIntVector3 Element = Sample.Indices[FaceStartVal + ElementIndex];
+								SourceElements[ElementIndex] = Chaos::TVec3<int32>(Element[0] - VertexStartVal, Element[1] - VertexStartVal, Element[2] - VertexStartVal);
+							}
+							TriangleMesh.Init(SourceElements);
+							Chaos::FTriangleMesh::TSpatialHashType<Chaos::FReal> SpatialHash;
+							TriangleMesh.BuildSpatialHash(ConstComponentSpaceVertices, SpatialHash, SphereRadius);
+							for (int32 TargetIndex = 0; TargetIndex < TargetVertexCountVal; TargetIndex++)
+							{
+								TArray<Chaos::TTriangleCollisionPoint<Chaos::FReal>> Result;
+								if (TriangleMesh.PointClosestTriangleQuery(SpatialHash, ConstComponentSpaceVertices,
+									TargetIndex, Chaos::TVec3<Chaos::FReal>(ComponentSpaceTargetVertices[TargetIndex]), SphereRadius / 2.f, SphereRadius / 2.f,
+									[](const int32 PointIndex, const int32 TriangleIndex)->bool {return true; }, Result))
+								{
+									for (const Chaos::TTriangleCollisionPoint<Chaos::FReal>& CollisionPoint : Result)
+									{
+										float CurrentDistance = abs(CollisionPoint.Phi);
+										float TriRadius = FalloffThreshold * MaxEdgeLength(ComponentSpaceVertices, Sample.Indices.Get(), VertexStartVal, FaceStartVal + CollisionPoint.Indices[1], 1);
+										float FalloffScale = CalculateFalloffScale(Falloff, TriRadius, CurrentDistance);
+										if (!FMath::IsNearlyZero(FalloffScale))
+										{
+											int32 TargetCandidateIndex = CollisionPoint.Indices[0] + TargetVertexStartVal;
+											if (ensure(0 <= TargetCandidateIndex && TargetCandidateIndex < TargetAttributeArray->Num()))
+											{
+												for (int32 k = 0; k < 3; k++)
+												{
+													(*TargetAttributeArray)[TargetCandidateIndex] += FalloffScale * (CollisionPoint.Bary[k+1] * (*SourceAttributeArray)[Sample.Indices[FaceStartVal + CollisionPoint.Indices[1]][k]]);
+												}
+												break;
+											}
+										}
+									}
+								}
+							}
+						}
+						else if (BoundingVolumeType == EDataflowTransferVertexAttributeNodeBoundingVolume::Vertex)
+						{
+							TUniquePtr<BVH> VertexBVH(BuildParticleSphereBVH(ComponentSpaceTargetVertices, SphereRadius));
+							for (int32 i = 0; i < FaceCountVal; i++)
+							{
+								FIntVector3 ComponentTriangle = Sample.Indices[FaceStartVal + i] - FIntVector3(VertexStartVal);
+								TArray<int32> TargetVertexIntersection({});
+								TriangleToVertexIntersections(*VertexBVH, ComponentSpaceVertices, ComponentTriangle, TargetVertexIntersection);
+
+								for (int32 j = 0; j < TargetVertexIntersection.Num(); j++)
+								{
+									Chaos::TVector<float, 3> Bary,
+										TriPos0(ComponentSpaceVertices[ComponentTriangle[0]]),
+										TriPos1(ComponentSpaceVertices[ComponentTriangle[1]]),
+										TriPos2(ComponentSpaceVertices[ComponentTriangle[2]]),
+										ParticlePos(ComponentSpaceTargetVertices[TargetVertexIntersection[j]]);
+
+									Chaos::TVector<float, 3> ClosestPoint = Chaos::FindClosestPointAndBaryOnTriangle(TriPos0, TriPos1, TriPos2, ParticlePos, Bary);
+									float CurrentDistance = (ParticlePos - ClosestPoint).Size();
+									float TriRadius = FalloffThreshold * MaxEdgeLength(ComponentSpaceVertices, Sample.Indices.Get(), VertexStartVal, FaceStartVal + i, 1);
+									float FalloffScale = CalculateFalloffScale(Falloff, TriRadius, CurrentDistance);
+									if (!FMath::IsNearlyZero(FalloffScale))
+									{
+										int32 TargetIndex = TargetVertexIntersection[j] + TargetVertexStartVal;
+										if (ensure(0 <= TargetIndex && TargetIndex < TargetAttributeArray->Num()))
+										{
+											for (int32 k = 0; k < 3; k++)
+											{
+												(*TargetAttributeArray)[TargetIndex] += FalloffScale * (Bary[k] * (*SourceAttributeArray)[ComponentTriangle[k] + VertexStartVal]);
+											}
+											break;
+										}
+									}
+								}
+							}
+						}
+					}
+				}
+			});
+	}
+}
+
+template<typename T>
+void FGeometryCollectionTransferVertexAttributeNode::NearestVertexTransfer(FCollectionAttributeKey Key, const UE::Private::FTransferFacade& Sample, UE::Private::FTransferFacade& Target, const TManagedArray<T>* SourceAttributeArray, TManagedArray<T>* TargetAttributeArray) const
+{
+	//TargetAttributeArray has to be filled with empty values before (attribute types: float, FLinearColor)
+	if (SourceAttributeArray && TargetAttributeArray)
+	{
+		// Build component space vertices for TargetCollection
+		TArray<FVector3f> ComponentSpaceTargetVertices;
+		BuildComponentSpaceVertices(Target.Transform.Get(), Target.Parent.Get(), Target.BoneMap.Get(), Target.Vertex.Get(), 0, Target.Vertex.Num(), ComponentSpaceTargetVertices);
+
+		// Build component space vertices for SourceCollection
+		TArray<FVector3f> ComponentSpaceVertices;
+		BuildComponentSpaceVertices(Sample.Transform.Get(), Sample.Parent.Get(), Sample.BoneMap.Get(), Sample.Vertex.Get(), 0, Sample.Vertex.Num(), ComponentSpaceVertices);
+
+		// build Sphere based BVH
+		Chaos::FReal SphereRadius = Chaos::FReal(EdgeMultiplier * FMath::Max(
+			MaxEdgeLength(ComponentSpaceTargetVertices, Target.Indices.Get(), 0, 0, Target.Indices.Num()),
+			MaxEdgeLength(ComponentSpaceVertices, Sample.Indices.Get(), 0, 0, Sample.Indices.Num())));
+		TUniquePtr<BVH> VertexBVH(BuildParticleSphereBVH(ComponentSpaceTargetVertices, SphereRadius));
+
+		for (int32 i = 0; i < Sample.Indices.Num(); i++)
+		{
+			FIntVector3 Triangle = Sample.Indices[i];
+			TArray<int32> TargetVertexIntersection({});
+			TriangleToVertexIntersections(*VertexBVH, ComponentSpaceVertices, Triangle, TargetVertexIntersection);
+
+			for (int32 j = 0; j < TargetVertexIntersection.Num(); j++)
+			{
+				Chaos::TVector<float, 3> Bary,
+					TriPos0(ComponentSpaceVertices[Sample.Indices[i][0]]),
+					TriPos1(ComponentSpaceVertices[Sample.Indices[i][1]]),
+					TriPos2(ComponentSpaceVertices[Sample.Indices[i][2]]),
+					ParticlePos(ComponentSpaceTargetVertices[TargetVertexIntersection[j]]);
+
+				Chaos::TVector<float, 3> ClosestPoint = Chaos::FindClosestPointAndBaryOnTriangle(TriPos0, TriPos1, TriPos2, ParticlePos, Bary);
+				float CurrentDistance = (ParticlePos - ClosestPoint).Size();
+				float TriRadius = FalloffThreshold * MaxEdgeLength(ComponentSpaceVertices, Sample.Indices.Get(), 0, i, 1);
+				float FalloffScale = CalculateFalloffScale(Falloff, TriRadius, CurrentDistance);
+				if (!FMath::IsNearlyZero(FalloffScale))
+				{
+					int32 TargetIndex = TargetVertexIntersection[j];
+					if (ensure(0 <= TargetIndex && TargetIndex < TargetAttributeArray->Num()))
+					{
+						for (int32 k = 0; k < 3; k++)
+						{
+							(*TargetAttributeArray)[TargetIndex] += FalloffScale * (Bary[k] * (*SourceAttributeArray)[Sample.Indices[i][k]]);
+						}
+						break;
+					}
+				}
+			}
+		}
+	}
+}
+
+
+float FGeometryCollectionTransferVertexAttributeNode::MaxEdgeLength(TArray<FVector3f>& Vert, const TManagedArray<FIntVector3>& Tri, int VertexOffset, int TriStart, int TriCount)
+{
+	auto TriInRange = [](const FIntVector3& T, int Max) {
+		for (int k = 0; k < 3; k++)
+		{
+			if (ensure(0 <= T[k] && T[k] < Max))
+			{
+				return true;
+			}
+		}
+		return false;
+		};
+
+	float Max = 0;
+	int TriStop = TriStart + TriCount;
+	for (int i = TriStart; i < TriStop; i++)
+	{
+		if (TriInRange(Tri[i] - FIntVector3(VertexOffset), Vert.Num()))
+		{
+			Max = FMath::Max(Max, (Vert[Tri[i][0] - VertexOffset] - Vert[Tri[i][1] - VertexOffset]).SquaredLength());
+			Max = FMath::Max(Max, (Vert[Tri[i][0] - VertexOffset] - Vert[Tri[i][2] - VertexOffset]).SquaredLength());
+			Max = FMath::Max(Max, (Vert[Tri[i][1] - VertexOffset] - Vert[Tri[i][2] - VertexOffset]).SquaredLength());
+		}
+	}
+	return FMath::Sqrt(Max);
+}
+
+void FGeometryCollectionTransferVertexAttributeNode::BuildComponentSpaceVertices(const TManagedArray<FTransform3f>& LocalSpaceTransform, const TManagedArray<int32>& Parent, const TManagedArray<int32>& BoneMapArray,
+	const TManagedArray<FVector3f>& VertexArray, int32 Start, int32 Count, TArray<FVector3f>& ComponentSpaceVertices)
+{
+	TArray<FTransform3f> ComponentTransform;
+	GeometryCollectionAlgo::GlobalMatrices(LocalSpaceTransform, Parent, ComponentTransform);
+
+	ComponentSpaceVertices.SetNumUninitialized(Count);
+	for (int i = 0; i < Count; i++)
+	{
+		int j = i + Start;
+		if (0 < (BoneMapArray)[i] && (BoneMapArray)[i] < ComponentTransform.Num())
+		{
+			ComponentSpaceVertices[i] = ComponentTransform[(BoneMapArray)[j]].TransformPosition(VertexArray[j]);
+		}
+		else
+		{
+			ComponentSpaceVertices[i] = VertexArray[j];
+		}
+	}
+}
+
+FGeometryCollectionTransferVertexAttributeNode::BVH* FGeometryCollectionTransferVertexAttributeNode::BuildParticleSphereBVH(const TArray<FVector3f>& Vertices, float Radius)
+{
+	TArray<SphereType*> VertexSpherePtrs;
+	TArray<SphereType> VertexSpheres;
+	VertexSpheres.Init(SphereType(Chaos::FVec3f(0), Radius), Vertices.Num());
+	VertexSpherePtrs.SetNum(Vertices.Num());
+
+	for (int32 i = 0; i < Vertices.Num(); i++)
+	{
+		Chaos::FVec3f SphereCenter(Vertices[i]);
+		SphereType VertexSphere(SphereCenter, Radius);
+		VertexSpheres[i] = SphereType(SphereCenter, Radius);
+		VertexSpherePtrs[i] = &VertexSpheres[i];
+	}
+	return new BVH(VertexSpherePtrs);
+}
+
+void FGeometryCollectionTransferVertexAttributeNode::TriangleToVertexIntersections(
+	const BVH& VertexBVH, const TArray<FVector3f>& ComponentSpaceVertices, const FIntVector3& Triangle, TArray<int32>& OutTargetVertexIntersection)
+{
+	OutTargetVertexIntersection.Empty();
+
+	TArray<int32> TargetVertexIntersection0 = VertexBVH.FindAllIntersections(ComponentSpaceVertices[Triangle[0]]);
+	TArray<int32> TargetVertexIntersection1 = VertexBVH.FindAllIntersections(ComponentSpaceVertices[Triangle[1]]);
+	TArray<int32> TargetVertexIntersection2 = VertexBVH.FindAllIntersections(ComponentSpaceVertices[Triangle[2]]);
+	TargetVertexIntersection0.Sort();
+	TargetVertexIntersection1.Sort();
+	TargetVertexIntersection2.Sort();
+
+	for (int32 k = 0; k < TargetVertexIntersection0.Num(); k++)
+	{
+		if (TargetVertexIntersection1.Contains(TargetVertexIntersection0[k])
+			&& TargetVertexIntersection2.Contains(TargetVertexIntersection0[k]))
+		{
+			OutTargetVertexIntersection.Emplace(TargetVertexIntersection0[k]);
+		}
+	}
+}
+
+
+float FGeometryCollectionTransferVertexAttributeNode::CalculateFalloffScale(EDataflowTransferVertexAttributeNodeFalloff FalloffSetting, float Threshold, float Distance)
+{
+	float Denominator = 1.0;
+	if (Distance > Threshold && !FMath::IsNearlyZero(Threshold))
+	{
+		Denominator = Distance / Threshold;
+	}
+	switch (FalloffSetting)
+	{
+	case EDataflowTransferVertexAttributeNodeFalloff::Linear:
+		return 1. / Denominator;
+	case EDataflowTransferVertexAttributeNodeFalloff::Squared:
+		return 1. / FMath::Square(Denominator);
+	}
+	return 1.0;
+}
+
+
+#undef LOCTEXT_NAMESPACE
diff --git a/Engine/Plugins/Experimental/GeometryCollectionPlugin/Source/GeometryCollectionNodes/Private/Dataflow/GeometryCollectionTransferVertexScalarAttributeNode.cpp b/Engine/Plugins/Experimental/GeometryCollectionPlugin/Source/GeometryCollectionNodes/Private/Dataflow/GeometryCollectionTransferVertexScalarAttributeNode.cpp
deleted file mode 100644
diff --git a/Engine/Plugins/Experimental/GeometryCollectionPlugin/Source/GeometryCollectionNodes/Public/Dataflow/GeometryCollectionTransferVertexAttributeNode.h b/Engine/Plugins/Experimental/GeometryCollectionPlugin/Source/GeometryCollectionNodes/Public/Dataflow/GeometryCollectionTransferVertexAttributeNode.h
new file mode 100644
--- /dev/null
+++ b/Engine/Plugins/Experimental/GeometryCollectionPlugin/Source/GeometryCollectionNodes/Public/Dataflow/GeometryCollectionTransferVertexAttributeNode.h
@@ -0,0 +1,154 @@
+// Copyright Epic Games, Inc. All Rights Reserved.
+
+#pragma once 
+
+#include "CoreMinimal.h"
+#include "Chaos/BoundingVolumeHierarchy.h"
+#include "Dataflow/DataflowNode.h"
+#include "Dataflow/DataflowConnectionTypes.h"
+#include "GeometryCollection/GeometryCollection.h"
+
+#include "GeometryCollectionTransferVertexAttributeNode.generated.h"
+
+
+namespace UE::Private 
+{
+
+	class FTransferFacade;
+};
+
+
+UENUM(BlueprintType)
+enum class EDataflowTransferVertexAttributeNodeFalloff : uint8
+{
+	/** Squared falloff based on distance from triangle*/
+	Squared  UMETA(DisplayName = "Squared"),
+
+	/** Linear falloff based on distance from triangle*/
+	Linear UMETA(DisplayName = "Linear"),
+
+	/** No distance falloff */
+	None UMETA(DisplayName = "None"),
+	//~~~
+	//256th entry
+	Dataflow_Max UMETA(Hidden)
+};
+
+UENUM(BlueprintType)
+enum class EDataflowTransferVertexAttributeNodeSampleScale : uint8
+{
+	/** Bounding volume hierarchy cell size based on max edge length of each geometry group*/
+	Component_Edge  UMETA(DisplayName = "Component Max Edge"),
+
+	/** Bounding volume hierarchy cell size based on max edge length of the whole asset*/
+	Asset_Edge UMETA(DisplayName = "Asset Max Edge"),
+
+	/** Bounding volume hierarchy cell size based on max length of the bounding box of the whole asset*/
+	Asset_Bound UMETA(DisplayName = "Asset Max Bound"),
+
+	//~~~
+	//256th entry
+	Dataflow_Max UMETA(Hidden)
+};
+
+UENUM(BlueprintType)
+enum class EDataflowTransferVertexAttributeNodeBoundingVolume : uint8
+{
+	/** Bounding volume on vertices of the source triangle mesh*/
+	Vertex UMETA(DisplayName = "Vertex"),
+
+	/** Bounding volume on triangles of the source triangle mesh*/
+	Triangle UMETA(DisplayName = "Triangle"),
+
+	//~~~
+	//256th entry
+	Dataflow_Max UMETA(Hidden)
+};
+
+/**
+ * Transfer float properties from a sample collection to a target collection.
+ */
+USTRUCT(meta = (DataflowGeometryCollection))
+struct FGeometryCollectionTransferVertexAttributeNode : public FDataflowNode
+{
+	GENERATED_USTRUCT_BODY()
+	DATAFLOW_NODE_DEFINE_INTERNAL(FGeometryCollectionTransferVertexAttributeNode, "TransferVertexAttribute", "GeometryCollection", "Transfer a named vertex attribute from the Source Collection to the Target Collection")
+	DATAFLOW_NODE_RENDER_TYPE("SurfaceWeightsRender", FGeometryCollection::StaticType(), "Collection", "AttributeKey")
+
+public:
+
+	/* Target collection to transfer vertex attribute to. */
+	UPROPERTY(Meta = (DataflowInput, DataflowOutput, DataflowPassthrough = "Collection"))
+	FManagedArrayCollection Collection;
+
+	/* Sample collection to transfer vertex attribute from. */
+	UPROPERTY(Meta = (DataflowInput, DisplayName = "FromCollection"))
+	FManagedArrayCollection FromCollection;
+
+	/* The name of the vertex attribute to generate indices from. */
+	UPROPERTY(EditAnywhere, Category = "Dataflow", Meta = (DataflowInput, DataflowOutput, DisplayName = "AttributeKey", DataflowPassthrough = "AttributeKey"))
+	FCollectionAttributeKey AttributeKey = FCollectionAttributeKey(FString(""), FString("Vertices"));
+
+	/* Bounding volume type for source assets[default: Triangle] */
+	UPROPERTY(EditAnywhere, Category = "Thresholds")
+	EDataflowTransferVertexAttributeNodeBoundingVolume BoundingVolumeType = EDataflowTransferVertexAttributeNodeBoundingVolume::Triangle;
+
+	/* Bounding volume hierarchy cell size for neighboring vertices to transfer into[default: Asset] */
+	UPROPERTY(EditAnywhere, Category = "Thresholds")
+	EDataflowTransferVertexAttributeNodeSampleScale SampleScale = EDataflowTransferVertexAttributeNodeSampleScale::Asset_Bound;
+
+	/* Falloff of sample value based on distance from sample triangle[default: Squared] */
+	UPROPERTY(EditAnywhere, Category = "Thresholds")
+	EDataflowTransferVertexAttributeNodeFalloff Falloff = EDataflowTransferVertexAttributeNodeFalloff::None;
+
+	/* Threshold based on distance from sample triangle.Values sampled past the threshold will falloff.[Defaults to 1 percent of triangle size(0.01)] */
+	UPROPERTY(EditAnywhere, Category = "Thresholds")
+	float FalloffThreshold = 0.01f;
+
+	/* Edge multiplier for the Bounding Volume Hierarchy(BVH) target's particle search radius. */
+	UPROPERTY(EditAnywhere, Category = "Thresholds", meta = (EditCondition = "SampleScale == EDataflowTransferVertexAttributeNodeSampleScale::Asset_Edge || SampleScale == EDataflowTransferVertexAttributeNodeSampleScale::Component_Edge", EditConditionHides))
+	float EdgeMultiplier = 0.5f;
+
+	/* Max bound multiplier for the Bounding Volume Hierarchy(BVH) target's particle search radius. */
+	UPROPERTY(EditAnywhere, Category = "Thresholds", meta = (
+		EditCondition = "SampleScale == EDataflowTransferVertexAttributeNodeSampleScale::Asset_Bound",
+		EditConditionHides))
+	float BoundMultiplier = 0.01f;
+
+	FGeometryCollectionTransferVertexAttributeNode(const Dataflow::FNodeParameters& InParam, FGuid InGuid = FGuid::NewGuid())
+		: FDataflowNode(InParam, InGuid)
+	{
+		RegisterInputConnection(&Collection);
+		RegisterInputConnection(&FromCollection);
+		RegisterInputConnection(&AttributeKey);
+
+		RegisterOutputConnection(&Collection, &Collection);
+		RegisterOutputConnection(&AttributeKey, &AttributeKey);
+	}
+
+private:
+
+	virtual void Evaluate(Dataflow::FContext& Context, const FDataflowOutput* Out) const override;
+	
+	typedef Chaos::TSphere<Chaos::FReal, 3> SphereType;
+	typedef Chaos::TBoundingVolumeHierarchy<TArray<SphereType*>, TArray<int32>, Chaos::FReal, 3> BVH;
+
+	//
+	// private helpers
+	//
+	TArray<FIntVector2> FindSourceToTargetGeometryMap(const FManagedArrayCollection& SourceCollection, const FManagedArrayCollection& TargetCollection) const;
+
+	template<typename T>
+	void PairedGeometryTransfer(FCollectionAttributeKey Key, const TArray<FIntVector2>& PairedGeometry, const UE::Private::FTransferFacade& Sample, UE::Private::FTransferFacade& Target, const TManagedArray<T>* SourceAttributeArray, TManagedArray<T>* TargetFloatArray) const;
+
+	template<typename T>
+	void NearestVertexTransfer(FCollectionAttributeKey Key, const UE::Private::FTransferFacade& Sample, UE::Private::FTransferFacade& Target, const TManagedArray<T>* SourceAttributeArray, TManagedArray<T>* TargetFloatArray) const;
+	static float MaxEdgeLength(TArray<FVector3f>& Vert, const TManagedArray<FIntVector3>& Tri, int VertexOFfset, int TriStart, int TriCount);
+	static void BuildComponentSpaceVertices(const TManagedArray<FTransform3f>& LocalSpaceTransform, const TManagedArray<int32>& Parent, const TManagedArray<int32>& BoneMapArray, const TManagedArray<FVector3f>& VertexArray, int32 Start, int32 Count, TArray<FVector3f>& ComponentSpaceVertices);
+	static BVH* BuildParticleSphereBVH(const TArray<FVector3f>& Vertices, float Radius);
+	static void TriangleToVertexIntersections(const BVH& VertexBVH, const TArray<FVector3f>& ComponentSpaceVertices, const FIntVector3& Triangle, TArray<int32>& OutTargetVertexIntersection);
+	static float CalculateFalloffScale(EDataflowTransferVertexAttributeNodeFalloff FalloffSetting, float Threshold, float Distance);
+
+};
+
+
diff --git a/Engine/Plugins/Experimental/GeometryCollectionPlugin/Source/GeometryCollectionNodes/Public/Dataflow/GeometryCollectionTransferVertexScalarAttributeNode.h b/Engine/Plugins/Experimental/GeometryCollectionPlugin/Source/GeometryCollectionNodes/Public/Dataflow/GeometryCollectionTransferVertexScalarAttributeNode.h
deleted file mode 100644
diff --git a/Engine/Source/Runtime/Experimental/Chaos/Private/GeometryCollection/GeometryCollection.cpp b/Engine/Source/Runtime/Experimental/Chaos/Private/GeometryCollection/GeometryCollection.cpp
--- a/Engine/Source/Runtime/Experimental/Chaos/Private/GeometryCollection/GeometryCollection.cpp
+++ b/Engine/Source/Runtime/Experimental/Chaos/Private/GeometryCollection/GeometryCollection.cpp
@@ -41,7 +41,7 @@ const FName FGeometryCollection::SimulatableParticlesAttribute("SimulatableParti
 const FName FGeometryCollection::SimulationTypeAttribute("SimulationType");
 const FName FGeometryCollection::StatusFlagsAttribute("StatusFlags");
 const FName FGeometryCollection::ExternalCollisionsAttribute("ExternalCollisions");
-
+const FName FGeometryCollection::ColorAttribute("Color");
 
 bool FGeometryCollection::AreCollisionParticlesEnabled()
 {
@@ -79,7 +79,7 @@ void FGeometryCollection::DefineGeometrySchema(FManagedArrayCollection& InCollec
 	InCollection.AddAttribute<FVector3f>("Vertex", FGeometryCollection::VerticesGroup);
 	InCollection.AddAttribute<FVector3f>("Normal", FGeometryCollection::VerticesGroup);
 	GeometryCollection::UV::DefineUVSchema(InCollection);
-	InCollection.AddAttribute<FLinearColor>("Color", FGeometryCollection::VerticesGroup);
+	InCollection.AddAttribute<FLinearColor>(FGeometryCollection::ColorAttribute, FGeometryCollection::VerticesGroup);
 	InCollection.AddAttribute<FVector3f>("TangentU", FGeometryCollection::VerticesGroup);
 	InCollection.AddAttribute<FVector3f>("TangentV", FGeometryCollection::VerticesGroup);
 	InCollection.AddAttribute<int32>("BoneMap", FGeometryCollection::VerticesGroup, TransformDependency);
@@ -122,7 +122,7 @@ void FGeometryCollection::Construct()
 	// Vertices Group
 	AddExternalAttribute<FVector3f>("Vertex", FGeometryCollection::VerticesGroup, Vertex);
 	AddExternalAttribute<FVector3f>("Normal", FGeometryCollection::VerticesGroup, Normal);
-	AddExternalAttribute<FLinearColor>("Color", FGeometryCollection::VerticesGroup, Color);
+	AddExternalAttribute<FLinearColor>(FGeometryCollection::ColorAttribute, FGeometryCollection::VerticesGroup, Color);
 	AddExternalAttribute<FVector3f>("TangentU", FGeometryCollection::VerticesGroup, TangentU);
 	AddExternalAttribute<FVector3f>("TangentV", FGeometryCollection::VerticesGroup, TangentV);
 	AddExternalAttribute<int32>("BoneMap", FGeometryCollection::VerticesGroup, BoneMap, TransformDependency);
diff --git a/Engine/Source/Runtime/Experimental/Chaos/Public/GeometryCollection/GeometryCollection.h b/Engine/Source/Runtime/Experimental/Chaos/Public/GeometryCollection/GeometryCollection.h
--- a/Engine/Source/Runtime/Experimental/Chaos/Public/GeometryCollection/GeometryCollection.h
+++ b/Engine/Source/Runtime/Experimental/Chaos/Public/GeometryCollection/GeometryCollection.h
@@ -106,7 +106,8 @@ public:
 	static CHAOS_API const FName SimulationTypeAttribute;
 	static CHAOS_API const FName StatusFlagsAttribute;
 	static CHAOS_API const FName ExternalCollisionsAttribute;
-	
+	static CHAOS_API const FName ColorAttribute;
+
 	enum ESimulationTypes : uint8
 	{
 		FST_None = 0,
