Dataflow : make sure clustering node checks if the selection is valid compare to the geometrycollection
if not, the clustering part is skipped and a log warning is emitted

diff --git a/Engine/Plugins/Experimental/GeometryCollectionPlugin/Source/GeometryCollectionNodes/Private/Dataflow/GeometryCollectionClusteringNodes.cpp b/Engine/Plugins/Experimental/GeometryCollectionPlugin/Source/GeometryCollectionNodes/Private/Dataflow/GeometryCollectionClusteringNodes.cpp
--- a/Engine/Plugins/Experimental/GeometryCollectionPlugin/Source/GeometryCollectionNodes/Private/Dataflow/GeometryCollectionClusteringNodes.cpp
+++ b/Engine/Plugins/Experimental/GeometryCollectionPlugin/Source/GeometryCollectionNodes/Private/Dataflow/GeometryCollectionClusteringNodes.cpp
@@ -3,6 +3,7 @@
 #include "Dataflow/GeometryCollectionClusteringNodes.h"
 #include "Dataflow/DataflowCore.h"
 
+#include "ChaosLog.h"
 #include "Engine/StaticMesh.h"
 #include "GeometryCollection/GeometryCollectionObject.h"
 #include "GeometryCollection/ManagedArrayCollection.h"
@@ -81,21 +82,29 @@ void FAutoClusterDataflowNode::Evaluate(Dataflow::FContext& Context, const FData
 			bool bInPreferConvexity = GetValue(Context, &bPreferConvexity);
 			float InConcavityTolerance = GetValue(Context, &ConcavityTolerance);
 
-			TArray<int32> SelectedBones;
-			InTransformSelection.AsArray(SelectedBones);
-
-			FFractureEngineClustering::AutoCluster(*GeomCollection,
-				SelectedBones,
-				(EFractureEngineClusterSizeMethod)InClusterSizeMethod,
-				InClusterSites,
-				InClusterFraction,
-				InSiteSize,
-				InAutoCluster,
-				InAvoidIsolated, 
-				InEnforceSiteParameters,
-				InGridX, InGridY, InGridZ, InMinimumClusterSize, InKMeansIterations, 
-				bInPreferConvexity, InConcavityTolerance);
-
+			// only cluster if the selection matches 
+			if (InTransformSelection.Num() == GeomCollection->NumElements(FTransformCollection::TransformGroup))
+			{
+				TArray<int32> SelectedBones;
+				InTransformSelection.AsArray(SelectedBones);
+
+
+				FFractureEngineClustering::AutoCluster(*GeomCollection,
+					SelectedBones,
+					(EFractureEngineClusterSizeMethod)InClusterSizeMethod,
+					InClusterSites,
+					InClusterFraction,
+					InSiteSize,
+					InAutoCluster,
+					InAvoidIsolated,
+					InEnforceSiteParameters,
+					InGridX, InGridY, InGridZ, InMinimumClusterSize, InKMeansIterations,
+					bInPreferConvexity, InConcavityTolerance);
+			}
+			else
+			{
+				UE_LOG(LogChaos, Warning, TEXT("Dataflow: AutoCluster Node input selection size does not match the collection size, skipping clustering"));
+			}
 			SetValue<const FManagedArrayCollection&>(Context, *GeomCollection, &Collection);
 		}
 	}
