Dataflow : deprecate select by cluster node and create a v2 one to fix the behavior
- the node used SelectCluster that was wrongly selecting the leaves
- make sure the original node behave the same  ( following the fix of SelectCluster in the facade code )
- make v2 to use the fixed version fo SelectCluster

diff --git a/Engine/Plugins/Experimental/GeometryCollectionPlugin/Source/GeometryCollectionNodes/Private/Dataflow/GeometryCollectionSelectionNodes.cpp b/Engine/Plugins/Experimental/GeometryCollectionPlugin/Source/GeometryCollectionNodes/Private/Dataflow/GeometryCollectionSelectionNodes.cpp
--- a/Engine/Plugins/Experimental/GeometryCollectionPlugin/Source/GeometryCollectionNodes/Private/Dataflow/GeometryCollectionSelectionNodes.cpp
+++ b/Engine/Plugins/Experimental/GeometryCollectionPlugin/Source/GeometryCollectionNodes/Private/Dataflow/GeometryCollectionSelectionNodes.cpp
@@ -56,6 +56,7 @@ namespace Dataflow
 		DATAFLOW_NODE_REGISTER_CREATION_FACTORY(FCollectionTransformSelectionContactDataflowNode);
 		DATAFLOW_NODE_REGISTER_CREATION_FACTORY(FCollectionTransformSelectionLeafDataflowNode);
 		DATAFLOW_NODE_REGISTER_CREATION_FACTORY(FCollectionTransformSelectionClusterDataflowNode);
+		DATAFLOW_NODE_REGISTER_CREATION_FACTORY(FCollectionTransformSelectionClusterDataflowNode_v2);
 		DATAFLOW_NODE_REGISTER_CREATION_FACTORY(FCollectionTransformSelectionBySizeDataflowNode);
 		DATAFLOW_NODE_REGISTER_CREATION_FACTORY(FCollectionTransformSelectionByVolumeDataflowNode);
 		DATAFLOW_NODE_REGISTER_CREATION_FACTORY(FCollectionTransformSelectionInBoxDataflowNode);
@@ -626,12 +627,15 @@ void FCollectionTransformSelectionLeafDataflowNode::Evaluate(Dataflow::FContext&
 
 void FCollectionTransformSelectionClusterDataflowNode::Evaluate(Dataflow::FContext& Context, const FDataflowOutput* Out) const
 {
-	if (Out->IsA<FDataflowTransformSelection>(&TransformSelection))
+	if (Out->IsA(&TransformSelection))
 	{
-		const FManagedArrayCollection& InCollection = GetValue<FManagedArrayCollection>(Context, &Collection);
+		const FManagedArrayCollection& InCollection = GetValue(Context, &Collection);
 
 		GeometryCollection::Facades::FCollectionTransformSelectionFacade TransformSelectionFacade(InCollection);
-		const TArray<int32>& SelectionArr = TransformSelectionFacade.SelectCluster();
+		// this node used to use SelectCluster() but this was buggy and woudl select the leaves instead
+		// for this reason this node is now deprecated and we need to keep it doing what it sued to : SelectLeaf()
+		// version 2 of the node properly use the right way 
+		const TArray<int32>& SelectionArr = TransformSelectionFacade.SelectLeaf(); // used to be buggy SelectCluster() - see comment above 
 
 		FDataflowTransformSelection NewTransformSelection;
 		NewTransformSelection.Initialize(InCollection.NumElements(FGeometryCollection::TransformGroup), false);
@@ -639,12 +643,32 @@ void FCollectionTransformSelectionClusterDataflowNode::Evaluate(Dataflow::FConte
 
 		SetValue(Context, MoveTemp(NewTransformSelection), &TransformSelection);
 	}
-	else if (Out->IsA<FManagedArrayCollection>(&Collection))
+	else if (Out->IsA(&Collection))
 	{
 		SafeForwardInput(Context, &Collection, &Collection);
 	}
 }
 
+void FCollectionTransformSelectionClusterDataflowNode_v2::Evaluate(Dataflow::FContext& Context, const FDataflowOutput* Out) const
+{
+	if (Out->IsA(&TransformSelection))
+	{
+		const FManagedArrayCollection& InCollection = GetValue(Context, &Collection);
+
+		GeometryCollection::Facades::FCollectionTransformSelectionFacade TransformSelectionFacade(InCollection);
+		const TArray<int32>& SelectionArr = TransformSelectionFacade.SelectCluster(); 
+
+		FDataflowTransformSelection NewTransformSelection;
+		NewTransformSelection.Initialize(InCollection.NumElements(FGeometryCollection::TransformGroup), false);
+		NewTransformSelection.SetFromArray(SelectionArr);
+
+		SetValue(Context, MoveTemp(NewTransformSelection), &TransformSelection);
+	}
+	else if (Out->IsA(&Collection))
+	{
+		SafeForwardInput(Context, &Collection, &Collection);
+	}
+}
 
 void FCollectionTransformSelectionBySizeDataflowNode::Evaluate(Dataflow::FContext& Context, const FDataflowOutput* Out) const
 {
diff --git a/Engine/Plugins/Experimental/GeometryCollectionPlugin/Source/GeometryCollectionNodes/Public/Dataflow/GeometryCollectionSelectionNodes.h b/Engine/Plugins/Experimental/GeometryCollectionPlugin/Source/GeometryCollectionNodes/Public/Dataflow/GeometryCollectionSelectionNodes.h
--- a/Engine/Plugins/Experimental/GeometryCollectionPlugin/Source/GeometryCollectionNodes/Public/Dataflow/GeometryCollectionSelectionNodes.h
+++ b/Engine/Plugins/Experimental/GeometryCollectionPlugin/Source/GeometryCollectionNodes/Public/Dataflow/GeometryCollectionSelectionNodes.h
@@ -656,9 +656,11 @@ public:
 /**
  *
  * Selects the clusters in the Collection
+ * Deprecated : this node had the wrong behavior and select the leaves instead
+ *				Replace it by CollectionTransformSelectLeaf or use the second version of CollectionTransformSelectCluster
  *
  */
-USTRUCT(meta = (DataflowGeometryCollection))
+USTRUCT(meta = (Deprecated = "5.5"))
 struct FCollectionTransformSelectionClusterDataflowNode : public FDataflowNode
 {
 	GENERATED_USTRUCT_BODY()
@@ -686,6 +688,39 @@ public:
 };
 
 
+/**
+ *
+ * Selects the clusters in the Collection
+ * this version works properly and address the issues found in the deprecated version 1
+ */
+USTRUCT()
+struct FCollectionTransformSelectionClusterDataflowNode_v2 : public FDataflowNode
+{
+	GENERATED_USTRUCT_BODY()
+	DATAFLOW_NODE_DEFINE_INTERNAL(FCollectionTransformSelectionClusterDataflowNode_v2, "CollectionTransformSelectCluster", "GeometryCollection|Selection|Transform", "")
+
+public:
+	/** GeometryCollection for the selection */
+	UPROPERTY(meta = (DataflowInput, DataflowOutput, DataflowPassthrough = "Collection", DataflowIntrinsic))
+	FManagedArrayCollection Collection;
+
+	/** Array of the selected bone indicies */
+	UPROPERTY(meta = (DataflowOutput, DisplayName = "TransformSelection"))
+	FDataflowTransformSelection TransformSelection;
+
+	FCollectionTransformSelectionClusterDataflowNode_v2(const Dataflow::FNodeParameters& InParam, FGuid InGuid = FGuid::NewGuid())
+		: FDataflowNode(InParam, InGuid)
+	{
+		RegisterInputConnection(&Collection);
+		RegisterOutputConnection(&Collection, &Collection);
+		RegisterOutputConnection(&TransformSelection);
+	}
+
+	virtual void Evaluate(Dataflow::FContext& Context, const FDataflowOutput* Out) const override;
+
+};
+
+
 UENUM(BlueprintType)
 enum class ERangeSettingEnum : uint8
 {
