Dataflow:
- Updated FMeshCutterDataflowNode to use a StaticMesh for cutter
- bug fixes

diff --git a/Engine/Plugins/Experimental/Fracture/Source/FractureEngine/Private/FractureEngineFracturing.cpp b/Engine/Plugins/Experimental/Fracture/Source/FractureEngine/Private/FractureEngineFracturing.cpp
--- a/Engine/Plugins/Experimental/Fracture/Source/FractureEngine/Private/FractureEngineFracturing.cpp
+++ b/Engine/Plugins/Experimental/Fracture/Source/FractureEngine/Private/FractureEngineFracturing.cpp
@@ -998,7 +998,7 @@ int32 FFractureEngineFracturing::MeshCutter(TArray<FTransform>& MeshTransforms,
 		for (const FTransform& ScatterTransform : MeshTransforms)
 		{
 			constexpr bool bSetDefaultInternalMaterialsFromCollection = true;
-			int32 Index = CutWithMesh(InDynCuttingMesh, ScatterTransform, InternalSurfaceMaterials, *GeomCollection, TransformSelectionArr, InCollisionSampleSpacing, InTransform, bSetDefaultInternalMaterialsFromCollection, nullptr, InSplitIslands);
+			int32 Index = CutWithMesh(InDynCuttingMesh, ScatterTransform, InternalSurfaceMaterials, *GeomCollection, TransformSelectionArr, InCollisionSampleSpacing, FTransform::Identity, bSetDefaultInternalMaterialsFromCollection, nullptr, InSplitIslands);
 
 			int32 NewLen = Algo::RemoveIf(TransformSelectionArr, [&](int32 Bone)
 				{
diff --git a/Engine/Plugins/Experimental/GeometryCollectionPlugin/Source/GeometryCollectionNodes/Private/Dataflow/GeometryCollectionFracturingNodes.cpp b/Engine/Plugins/Experimental/GeometryCollectionPlugin/Source/GeometryCollectionNodes/Private/Dataflow/GeometryCollectionFracturingNodes.cpp
--- a/Engine/Plugins/Experimental/GeometryCollectionPlugin/Source/GeometryCollectionNodes/Private/Dataflow/GeometryCollectionFracturingNodes.cpp
+++ b/Engine/Plugins/Experimental/GeometryCollectionPlugin/Source/GeometryCollectionNodes/Private/Dataflow/GeometryCollectionFracturingNodes.cpp
@@ -32,6 +32,10 @@
 #include "Dataflow/DataflowSelection.h"
 #include "DynamicMesh/DynamicMesh3.h"
 #include "UDynamicMesh.h"
+#include "MeshDescription.h"
+#include "DynamicMeshToMeshDescription.h"
+#include "MeshDescriptionToDynamicMesh.h"
+#include "StaticMeshAttributes.h"
 
 #include UE_INLINE_GENERATED_CPP_BY_NAME(GeometryCollectionFracturingNodes)
 
@@ -792,6 +796,7 @@ void FMeshCutterDataflowNode::Evaluate(Dataflow::FContext& Context, const FDataf
 		Out->IsA(&TransformSelection) ||
 		Out->IsA(&NewGeometryTransformSelection))
 	{
+#if WITH_EDITORONLY_DATA
 		FDataflowTransformSelection InTransformSelection = GetValue(Context, &TransformSelection);
 		//
 		// If not connected select everything by default
@@ -828,96 +833,112 @@ void FMeshCutterDataflowNode::Evaluate(Dataflow::FContext& Context, const FDataf
 		{
 			FManagedArrayCollection InCollection = GetValue(Context, &Collection);
 
-			if (TObjectPtr<UDynamicMesh> InCuttingMesh = GetValue(Context, &CuttingMesh))
+			if (TObjectPtr<UStaticMesh> InCuttingMesh = GetValue(Context, &CuttingStaticMesh))
 			{
-				const UE::Geometry::FDynamicMesh3& InDynCuttingMesh = InCuttingMesh->GetMeshRef();
-
-				if (InDynCuttingMesh.VertexCount() > 0)
+				if (FMeshDescription* MeshDescription = bUseHiRes ? InCuttingMesh->GetHiResMeshDescription() : InCuttingMesh->GetMeshDescription(LODLevel))
 				{
-					const int32 InRandomSeed = GetValue(Context, &RandomSeed);
-					const int32 InNumberToScatter = GetValue(Context, &NumberToScatter);
-					const int32 InGridX = GetValue(Context, &GridX);
-					const int32 InGridY = GetValue(Context, &GridY);
-					const int32 InGridZ = GetValue(Context, &GridZ);
-					const float InVariability = GetValue(Context, &Variability);
-					const float InMinScaleFactor = GetValue(Context, &MinScaleFactor);
-					const float InMaxScaleFactor = GetValue(Context, &MaxScaleFactor);
-					const float InRollRange = GetValue(Context, &RollRange);
-					const float InPitchRange = GetValue(Context, &PitchRange);
-					const float InYawRange = GetValue(Context, &YawRange);
-					const FTransform InTransform = GetValue(Context, &Transform);
-					const float InChanceToFracture = GetValue(Context, &ChanceToFracture);
-					const float InCollisionSampleSpacing = GetValue(Context, &CollisionSampleSpacing);
-
-					TArray<FTransform> MeshTransforms;
-
-					if (CutDistribution == EMeshCutterCutDistribution::SingleCut)
+					// If HiRes is empty then use LoRes
+					if (bUseHiRes && MeshDescription->Vertices().Num() == 0)
 					{
-						MeshTransforms.Add(InTransform);
-					}
-					else
-					{
-						FFractureEngineFracturing::GenerateMeshTransforms(MeshTransforms,
-							InBoundingBox,
-							InRandomSeed,
-							CutDistribution,
-							InNumberToScatter,
-							InGridX,
-							InGridY,
-							InGridZ,
-							InVariability,
-							InMinScaleFactor,
-							InMaxScaleFactor,
-							bRandomOrientation,
-							InRollRange,
-							InPitchRange,
-							InYawRange);
+						MeshDescription = InCuttingMesh->GetMeshDescription(LODLevel);
 					}
 
-					int32 ResultGeometryIndex = FFractureEngineFracturing::MeshCutter(MeshTransforms,
-						InCollection,
-						InTransformSelection,
-						InDynCuttingMesh,
-						InTransform,
-						InRandomSeed,
-						InChanceToFracture,
-						SplitIslands,
-						InCollisionSampleSpacing);
+					if (MeshDescription->Vertices().Num() > 0)
+					{
+						TObjectPtr<UDynamicMesh> NewMesh = NewObject<UDynamicMesh>();
+						NewMesh->Reset();
 
-					FDataflowTransformSelection NewSelection;
-					FDataflowTransformSelection OriginalSelection;
+						UE::Geometry::FDynamicMesh3& DynCuttingMesh = NewMesh->GetMeshRef();
+						{
+							FMeshDescriptionToDynamicMesh ConverterToDynamicMesh;
+							ConverterToDynamicMesh.Convert(MeshDescription, DynCuttingMesh);
+						}
 
-					if (ResultGeometryIndex != INDEX_NONE)
-					{
-						if (InCollection.HasAttribute("TransformIndex", FGeometryCollection::GeometryGroup))
+						const int32 InRandomSeed = GetValue(Context, &RandomSeed);
+						const int32 InNumberToScatter = GetValue(Context, &NumberToScatter);
+						const int32 InGridX = GetValue(Context, &GridX);
+						const int32 InGridY = GetValue(Context, &GridY);
+						const int32 InGridZ = GetValue(Context, &GridZ);
+						const float InVariability = GetValue(Context, &Variability);
+						const float InMinScaleFactor = GetValue(Context, &MinScaleFactor);
+						const float InMaxScaleFactor = GetValue(Context, &MaxScaleFactor);
+						const float InRollRange = GetValue(Context, &RollRange);
+						const float InPitchRange = GetValue(Context, &PitchRange);
+						const float InYawRange = GetValue(Context, &YawRange);
+						const FTransform InTransform = GetValue(Context, &Transform);
+						const float InChanceToFracture = GetValue(Context, &ChanceToFracture);
+						const float InCollisionSampleSpacing = GetValue(Context, &CollisionSampleSpacing);
+
+						TArray<FTransform> MeshTransforms;
+
+						if (CutDistribution == EMeshCutterCutDistribution::SingleCut)
+						{
+							MeshTransforms.Add(InTransform);
+						}
+						else
 						{
-							const TManagedArray<int32>& TransformIndices = InCollection.GetAttribute<int32>("TransformIndex", FGeometryCollection::GeometryGroup);
+							FFractureEngineFracturing::GenerateMeshTransforms(MeshTransforms,
+								InBoundingBox,
+								InRandomSeed,
+								CutDistribution,
+								InNumberToScatter,
+								InGridX,
+								InGridY,
+								InGridZ,
+								InVariability,
+								InMinScaleFactor,
+								InMaxScaleFactor,
+								bRandomOrientation,
+								InRollRange,
+								InPitchRange,
+								InYawRange);
+						}
 
-							NewSelection.Initialize(TransformIndices.Num(), false);
-							OriginalSelection.Initialize(TransformIndices.Num(), false);
+						int32 ResultGeometryIndex = FFractureEngineFracturing::MeshCutter(MeshTransforms,
+							InCollection,
+							InTransformSelection,
+							DynCuttingMesh,
+							InTransform,
+							InRandomSeed,
+							InChanceToFracture,
+							SplitIslands,
+							InCollisionSampleSpacing);
 
-							// The newly fractured pieces are added to the end of the transform array (starting position is ResultGeometryIndex)
-							for (int32 Idx = ResultGeometryIndex; Idx < TransformIndices.Num(); ++Idx)
-							{
-								int32 BoneIdx = TransformIndices[Idx];
-								NewSelection.SetSelected(BoneIdx);
-							}
+						FDataflowTransformSelection NewSelection;
+						FDataflowTransformSelection OriginalSelection;
 
-							for (int32 Idx = 0; Idx < InTransformSelection.Num(); ++Idx)
+						if (ResultGeometryIndex != INDEX_NONE)
+						{
+							if (InCollection.HasAttribute("TransformIndex", FGeometryCollection::GeometryGroup))
 							{
-								if (InTransformSelection.IsSelected(Idx))
+								const TManagedArray<int32>& TransformIndices = InCollection.GetAttribute<int32>("TransformIndex", FGeometryCollection::GeometryGroup);
+
+								NewSelection.Initialize(TransformIndices.Num(), false);
+								OriginalSelection.Initialize(TransformIndices.Num(), false);
+
+								// The newly fractured pieces are added to the end of the transform array (starting position is ResultGeometryIndex)
+								for (int32 Idx = ResultGeometryIndex; Idx < TransformIndices.Num(); ++Idx)
+								{
+									int32 BoneIdx = TransformIndices[Idx];
+									NewSelection.SetSelected(BoneIdx);
+								}
+
+								for (int32 Idx = 0; Idx < InTransformSelection.Num(); ++Idx)
 								{
-									OriginalSelection.SetSelected(Idx);
+									if (InTransformSelection.IsSelected(Idx))
+									{
+										OriginalSelection.SetSelected(Idx);
+									}
 								}
 							}
 						}
-					}
 
-					SetValue(Context, MoveTemp(InCollection), &Collection);
-					SetValue(Context, OriginalSelection, &TransformSelection);
-					SetValue(Context, NewSelection, &NewGeometryTransformSelection);
+						SetValue(Context, MoveTemp(InCollection), &Collection);
+						SetValue(Context, OriginalSelection, &TransformSelection);
+						SetValue(Context, NewSelection, &NewGeometryTransformSelection);
 
-					return;
+						return;
+					}
 				}
 			}
 		}
@@ -925,6 +946,9 @@ void FMeshCutterDataflowNode::Evaluate(Dataflow::FContext& Context, const FDataf
 		SafeForwardInput(Context, &Collection, &Collection);
 		SetValue(Context, InTransformSelection, &TransformSelection);
 		SetValue(Context, FDataflowTransformSelection(), &NewGeometryTransformSelection);
+#else
+		ensureMsgf(false, TEXT("FMeshCutterDataflowNode is an editor only node."));
+#endif
 	}
 }
 
diff --git a/Engine/Plugins/Experimental/GeometryCollectionPlugin/Source/GeometryCollectionNodes/Public/Dataflow/GeometryCollectionFracturingNodes.h b/Engine/Plugins/Experimental/GeometryCollectionPlugin/Source/GeometryCollectionNodes/Public/Dataflow/GeometryCollectionFracturingNodes.h
--- a/Engine/Plugins/Experimental/GeometryCollectionPlugin/Source/GeometryCollectionNodes/Public/Dataflow/GeometryCollectionFracturingNodes.h
+++ b/Engine/Plugins/Experimental/GeometryCollectionPlugin/Source/GeometryCollectionNodes/Public/Dataflow/GeometryCollectionFracturingNodes.h
@@ -11,7 +11,7 @@
 #include "GeometryCollectionFracturingNodes.generated.h"
 
 class FGeometryCollection;
-class UDynamicMesh;
+class UStaticMesh;
 
 USTRUCT(meta = (DataflowGeometryCollection, Deprecated = "5.5"))
 struct FUniformScatterPointsDataflowNode : public FDataflowNode
@@ -1041,8 +1041,16 @@ public:
 	FTransform Transform = FTransform::Identity;
 
 	/** Mesh to cut with */
-	UPROPERTY(meta = (DataflowInput, DataflowIntrinsic))
-	TObjectPtr<UDynamicMesh> CuttingMesh;
+	UPROPERTY(EditAnywhere, Category = "StaticMesh", meta = (DataflowInput, DataflowIntrinsic))
+	TObjectPtr<UStaticMesh> CuttingStaticMesh;
+
+	/** Output the HiRes representation, if set to true and HiRes doesn't exist it will output empty mesh */
+	UPROPERTY(EditAnywhere, Category = "StaticMesh", meta = (DisplayName = "Use HiRes"));
+	bool bUseHiRes = false;
+
+	/** Specifies the LOD level to use */
+	UPROPERTY(EditAnywhere, Category = "StaticMesh", meta = (DisplayName = "LOD Level"));
+	int32 LODLevel = 0;
 
 	/** How to arrange the mesh cuts in space */
 	UPROPERTY(EditAnywhere, Category = Distribution)
@@ -1131,7 +1139,7 @@ public:
 		RegisterInputConnection(&BoundingBox);
 		RegisterInputConnection(&TransformSelection);
 		RegisterInputConnection(&Transform);
-		RegisterInputConnection(&CuttingMesh);
+		RegisterInputConnection(&CuttingStaticMesh);
 		RegisterInputConnection(&NumberToScatter);
 		RegisterInputConnection(&GridX);
 		RegisterInputConnection(&GridY);
