Fix crash when passing a selection to the fracture node if the geometry

diff --git a/Engine/Plugins/Experimental/GeometryCollectionPlugin/Source/GeometryCollectionNodes/Private/Dataflow/GeometryCollectionFracturingNodes.cpp b/Engine/Plugins/Experimental/GeometryCollectionPlugin/Source/GeometryCollectionNodes/Private/Dataflow/GeometryCollectionFracturingNodes.cpp
--- a/Engine/Plugins/Experimental/GeometryCollectionPlugin/Source/GeometryCollectionNodes/Private/Dataflow/GeometryCollectionFracturingNodes.cpp
+++ b/Engine/Plugins/Experimental/GeometryCollectionPlugin/Source/GeometryCollectionNodes/Private/Dataflow/GeometryCollectionFracturingNodes.cpp
@@ -416,23 +416,24 @@ void FVoronoiFractureDataflowNode_v2::Evaluate(Dataflow::FContext& Context, cons
 			{
 				if (InCollection.HasAttribute("TransformIndex", FGeometryCollection::GeometryGroup))
 				{
-					const TManagedArray<int32>& TransformIndices = InCollection.GetAttribute<int32>("TransformIndex", FGeometryCollection::GeometryGroup);
+					const TManagedArray<int32>& GeometryToTransformIndices = InCollection.GetAttribute<int32>("TransformIndex", FGeometryCollection::GeometryGroup);
 
-					NewSelection.Initialize(TransformIndices.Num(), false);
-					OriginalSelection.Initialize(TransformIndices.Num(), false);
+					const int32 NumTransforms = InCollection.NumElements(FGeometryCollection::TransformGroup);
+					NewSelection.Initialize(NumTransforms, false);
+					OriginalSelection.Initialize(NumTransforms, false);
 
 					// The newly fractured pieces are added to the end of the transform array (starting position is ResultGeometryIndex)
-					for (int32 Idx = ResultGeometryIndex; Idx < TransformIndices.Num(); ++Idx)
+					for (int32 GeometryIdx = ResultGeometryIndex; GeometryIdx < GeometryToTransformIndices.Num(); ++GeometryIdx)
 					{
-						int32 BoneIdx = TransformIndices[Idx];
-						NewSelection.SetSelected(BoneIdx);
+						const int32 TransformIdx = GeometryToTransformIndices[GeometryIdx];
+						NewSelection.SetSelected(TransformIdx);
 					}
 
-					for (int32 Idx = 0; Idx < InTransformSelection.Num(); ++Idx)
+					for (int32 TransformIdx = 0; TransformIdx < InTransformSelection.Num(); ++TransformIdx)
 					{
-						if (InTransformSelection.IsSelected(Idx))
+						if (InTransformSelection.IsSelected(TransformIdx))
 						{
-							OriginalSelection.SetSelected(Idx);
+							OriginalSelection.SetSelected(TransformIdx);
 						}
 					}
 				}
