Change Dataflow namespace to UE::Dataflow to be ok with coding standard

diff --git a/Engine/Plugins/Experimental/Dataflow/Source/DataflowEditor/Private/Dataflow/DataflowSchema.cpp b/Engine/Plugins/Experimental/Dataflow/Source/DataflowEditor/Private/Dataflow/DataflowSchema.cpp
--- a/Engine/Plugins/Experimental/Dataflow/Source/DataflowEditor/Private/Dataflow/DataflowSchema.cpp
+++ b/Engine/Plugins/Experimental/Dataflow/Source/DataflowEditor/Private/Dataflow/DataflowSchema.cpp
@@ -22,7 +22,7 @@
 
 #define LOCTEXT_NAMESPACE "DataflowNode"
 
-namespace Dataflow::Private
+namespace UE::Dataflow::Private
 {
 	static const FName ManagedArrayCollectionType = FName("FManagedArrayCollection");
 	static const FName FloatType = FName("float");
@@ -39,7 +39,7 @@ namespace Dataflow::Private
 	static const FName BoxType = FName("FBox");
 	static const FName SphereType = FName("FSphere");
 	static const FName DataflowAnyTypeType = FName("FDataflowAnyType");
-} // namespace Dataflow::Private
+} // namespace UE::Dataflow::Private
 
 UDataflowSchema::UDataflowSchema()
 {
@@ -108,9 +108,9 @@ void UDataflowSchema::GetContextMenuActions(class UToolMenu* Menu, class UGraphN
 
 void UDataflowSchema::GetGraphContextActions(FGraphContextMenuBuilder& ContextMenuBuilder) const
 {
-	if (Dataflow::FNodeFactory* Factory = Dataflow::FNodeFactory::GetInstance())
+	if (UE::Dataflow::FNodeFactory* Factory = UE::Dataflow::FNodeFactory::GetInstance())
 	{
-		for (Dataflow::FFactoryParameters NodeParameters : Factory->RegisteredParameters())
+		for (UE::Dataflow::FFactoryParameters NodeParameters : Factory->RegisteredParameters())
 		{
 			if (TSharedPtr<FAssetSchemaAction_Dataflow_CreateNode_DataflowEdNode> Action =
 				FAssetSchemaAction_Dataflow_CreateNode_DataflowEdNode::CreateAction(ContextMenuBuilder.CurrentGraph, NodeParameters.TypeName, NodeParameters.DisplayName))
@@ -243,63 +243,63 @@ FLinearColor UDataflowSchema::GetTypeColor(const FName& Type)
 	const UGraphEditorSettings* Settings = GetDefault<UGraphEditorSettings>();
 	const UDataflowSettings* DataflowSettings = GetDefault<UDataflowSettings>();
 
-	if (Type == Dataflow::Private::ManagedArrayCollectionType)
+	if (Type == UE::Dataflow::Private::ManagedArrayCollectionType)
 	{
 		return DataflowSettings->ManagedArrayCollectionPinTypeColor;
 	}
-	else if (Type == Dataflow::Private::FloatType)
+	else if (Type == UE::Dataflow::Private::FloatType)
 	{
 		return Settings->FloatPinTypeColor;
 	}
-	else if (Type == Dataflow::Private::DoubleType)
+	else if (Type == UE::Dataflow::Private::DoubleType)
 	{
 		return Settings->DoublePinTypeColor;
 	}
-	else if (Type == Dataflow::Private::Int32Type)
+	else if (Type == UE::Dataflow::Private::Int32Type)
 	{
 		return Settings->IntPinTypeColor;
 	}
-	else if (Type == Dataflow::Private::BoolType)
+	else if (Type == UE::Dataflow::Private::BoolType)
 	{
 		return Settings->BooleanPinTypeColor;
 	}
-	else if (Type == Dataflow::Private::StringType)
+	else if (Type == UE::Dataflow::Private::StringType)
 	{
 		return Settings->StringPinTypeColor;
 	}
-	else if (Type == Dataflow::Private::NameType)
+	else if (Type == UE::Dataflow::Private::NameType)
 	{
 		return Settings->NamePinTypeColor;
 	}
-	else if (Type == Dataflow::Private::TextType)
+	else if (Type == UE::Dataflow::Private::TextType)
 	{
 		return Settings->TextPinTypeColor;
 	}
-	else if (Type == Dataflow::Private::VectorType)
+	else if (Type == UE::Dataflow::Private::VectorType)
 	{
 		return Settings->VectorPinTypeColor;
 	}
-	else if (Type == Dataflow::Private::TransformType)
+	else if (Type == UE::Dataflow::Private::TransformType)
 	{
 		return Settings->TransformPinTypeColor;
 	}
-	else if (Type == Dataflow::Private::RotatorType)
+	else if (Type == UE::Dataflow::Private::RotatorType)
 	{
 		return Settings->RotatorPinTypeColor;
 	}
-	else if (Type == Dataflow::Private::ArrayType)
+	else if (Type == UE::Dataflow::Private::ArrayType)
 	{
 		return DataflowSettings->ArrayPinTypeColor;
 	}
-	else if (Type == Dataflow::Private::BoxType)
+	else if (Type == UE::Dataflow::Private::BoxType)
 	{
 		return DataflowSettings->BoxPinTypeColor;
 	}
-	else if (Type == Dataflow::Private::SphereType)
+	else if (Type == UE::Dataflow::Private::SphereType)
 	{
 		return DataflowSettings->SpherePinTypeColor;
 	}
-	else if (Type == Dataflow::Private::DataflowAnyTypeType)
+	else if (Type == UE::Dataflow::Private::DataflowAnyTypeType)
 	{
 		return DataflowSettings->DataflowAnyTypePinTypeColor;
 	}
diff --git a/Engine/Plugins/Experimental/Dataflow/Source/DataflowEditor/Private/Dataflow/DataflowSelectionView.cpp b/Engine/Plugins/Experimental/Dataflow/Source/DataflowEditor/Private/Dataflow/DataflowSelectionView.cpp
--- a/Engine/Plugins/Experimental/Dataflow/Source/DataflowEditor/Private/Dataflow/DataflowSelectionView.cpp
+++ b/Engine/Plugins/Experimental/Dataflow/Source/DataflowEditor/Private/Dataflow/DataflowSelectionView.cpp
@@ -44,7 +44,7 @@ void FDataflowSelectionView::UpdateViewData()
 				{
 					if (const TObjectPtr<UDataflowBaseContent> Content = GetEditorContent())
 					{
-						if (TSharedPtr<Dataflow::FEngineContext> Context = Content->GetDataflowContext())
+						if (TSharedPtr<UE::Dataflow::FEngineContext> Context = Content->GetDataflowContext())
 						{
 							TArray<FDataflowOutput*> Outputs = DataflowNode->GetOutputs();
 
diff --git a/Engine/Plugins/Experimental/Dataflow/Source/DataflowEditor/Private/Dataflow/DataflowSimulationControls.cpp b/Engine/Plugins/Experimental/Dataflow/Source/DataflowEditor/Private/Dataflow/DataflowSimulationControls.cpp
--- a/Engine/Plugins/Experimental/Dataflow/Source/DataflowEditor/Private/Dataflow/DataflowSimulationControls.cpp
+++ b/Engine/Plugins/Experimental/Dataflow/Source/DataflowEditor/Private/Dataflow/DataflowSimulationControls.cpp
@@ -19,11 +19,11 @@
 
 namespace UE::Dataflow
 {
-	bool ShouldResetWorld(const TObjectPtr<UDataflow>& SimulationGraph, const TObjectPtr<UWorld>& SimulationWorld, ::Dataflow::FTimestamp& LastTimeStamp)
+	bool ShouldResetWorld(const TObjectPtr<UDataflow>& SimulationGraph, const TObjectPtr<UWorld>& SimulationWorld, UE::Dataflow::FTimestamp& LastTimeStamp)
 	{
-		if(const TSharedPtr<::Dataflow::FGraph> DataflowGraph = SimulationGraph->GetDataflow())
+		if(const TSharedPtr<UE::Dataflow::FGraph> DataflowGraph = SimulationGraph->GetDataflow())
 		{
-			::Dataflow::FTimestamp MaxTimeStamp = ::Dataflow::FTimestamp::Invalid;
+			UE::Dataflow::FTimestamp MaxTimeStamp = UE::Dataflow::FTimestamp::Invalid;
 			for(const TSharedPtr<FDataflowNode>& TerminalNode : DataflowGraph->GetFilteredNodes(FDataflowTerminalNode::StaticType()))
 			{
 				MaxTimeStamp.Value = FMath::Max(MaxTimeStamp.Value, TerminalNode->GetTimestamp().Value);
diff --git a/Engine/Plugins/Experimental/Dataflow/Source/DataflowEditor/Private/Dataflow/DataflowSimulationGenerator.cpp b/Engine/Plugins/Experimental/Dataflow/Source/DataflowEditor/Private/Dataflow/DataflowSimulationGenerator.cpp
--- a/Engine/Plugins/Experimental/Dataflow/Source/DataflowEditor/Private/Dataflow/DataflowSimulationGenerator.cpp
+++ b/Engine/Plugins/Experimental/Dataflow/Source/DataflowEditor/Private/Dataflow/DataflowSimulationGenerator.cpp
@@ -30,7 +30,7 @@ void FDataflowSimulationTask::DoWork()
 			if(bAsyncCaching)
 			{
 				// Compute all the skelmesh animations at the simulation time
-				Dataflow::ComputeSkeletonAnimation(TaskManager->PreviewActor, SimulationTime);
+				UE::Dataflow::ComputeSkeletonAnimation(TaskManager->PreviewActor, SimulationTime);
 				
 				// Background task : run directly the advance simulation data without coming back to the game thread
 				UDataflowSimulationManager* DataflowManager = SimulationWorld->GetSubsystem<UDataflowSimulationManager>();
@@ -47,7 +47,7 @@ void FDataflowSimulationTask::DoWork()
 			else
 			{
 				// Update all the skelmesh animations at the simulation time
-				Dataflow::UpdateSkeletonAnimation(TaskManager->PreviewActor, SimulationTime);
+				UE::Dataflow::UpdateSkeletonAnimation(TaskManager->PreviewActor, SimulationTime);
 				
 				// Foreground task : Run the world ticking 
 				SimulationWorld->Tick(ELevelTick::LEVELTICK_All, DeltaTime);
@@ -75,8 +75,8 @@ bool FDataflowTaskManager::AllocateSimulationResource(const FVector2f& TimeRange
 	
 	CacheManager = SimulationWorld->SpawnActor<AChaosCacheManager>();
 
-	PreviewActor = Dataflow::SpawnSimulatedActor(ActorClass, CacheManager, CacheAsset, true, DataflowContent, BlueprintTransform);
-	Dataflow::SetupSkeletonAnimation(PreviewActor);
+	PreviewActor = UE::Dataflow::SpawnSimulatedActor(ActorClass, CacheManager, CacheAsset, true, DataflowContent, BlueprintTransform);
+	UE::Dataflow::SetupSkeletonAnimation(PreviewActor);
 
 	// Init the cache manager
 	CacheManager->SetObservedComponentProperties(CacheManager->CacheMode);
@@ -98,7 +98,7 @@ bool FDataflowTaskManager::AllocateSimulationResource(const FVector2f& TimeRange
 		if(SimulationTask->GetTask().bAsyncCaching)
 		{
 			// Update all the skelmesh animations at the simulation time
-			Dataflow::UpdateSkeletonAnimation(PreviewActor, SimulationTask->GetTask().MinTime+SimulationTask->GetTask().DeltaTime);
+			UE::Dataflow::UpdateSkeletonAnimation(PreviewActor, SimulationTask->GetTask().MinTime+SimulationTask->GetTask().DeltaTime);
 				
 			// Foreground task : Run the world ticking 
 			SimulationWorld->Tick(ELevelTick::LEVELTICK_All, SimulationTask->GetTask().DeltaTime);
diff --git a/Engine/Plugins/Experimental/Dataflow/Source/DataflowEditor/Private/Dataflow/DataflowSimulationGeometryCache.cpp b/Engine/Plugins/Experimental/Dataflow/Source/DataflowEditor/Private/Dataflow/DataflowSimulationGeometryCache.cpp
--- a/Engine/Plugins/Experimental/Dataflow/Source/DataflowEditor/Private/Dataflow/DataflowSimulationGeometryCache.cpp
+++ b/Engine/Plugins/Experimental/Dataflow/Source/DataflowEditor/Private/Dataflow/DataflowSimulationGeometryCache.cpp
@@ -18,7 +18,7 @@ DEFINE_LOG_CATEGORY(LogDataflowSimulationGeometryCache);
 
 #define LOCTEXT_NAMESPACE "DataflowSimulationGeometryCache"
 	
-namespace DataflowSimulationGeometryCache
+namespace UE::DataflowSimulationGeometryCache
 {
 	int32 GetNumVertices(const FSkeletalMeshLODRenderData& LODData)
 	{
diff --git a/Engine/Plugins/Experimental/Dataflow/Source/DataflowEditor/Private/Dataflow/DataflowSimulationScene.cpp b/Engine/Plugins/Experimental/Dataflow/Source/DataflowEditor/Private/Dataflow/DataflowSimulationScene.cpp
--- a/Engine/Plugins/Experimental/Dataflow/Source/DataflowEditor/Private/Dataflow/DataflowSimulationScene.cpp
+++ b/Engine/Plugins/Experimental/Dataflow/Source/DataflowEditor/Private/Dataflow/DataflowSimulationScene.cpp
@@ -347,12 +347,12 @@ void UDataflowSimulationSceneDescription::GenerateGeometryCache()
 			Cast<AChaosCacheManager>(GetRootActor)->SetStartTime(Time);
 			RenderPositions[Frame] = GeometryCachable->GetGeometryCachePositions(EmbeddedSkeletalMesh);
 		}
-		DataflowSimulationGeometryCache::SaveGeometryCache(*GeometryCacheAsset, *EmbeddedSkeletalMesh, ImportedVertexNumbers, RenderPositions);
-		DataflowSimulationGeometryCache::SavePackage(*GeometryCacheAsset);
+		UE::DataflowSimulationGeometryCache::SaveGeometryCache(*GeometryCacheAsset, *EmbeddedSkeletalMesh, ImportedVertexNumbers, RenderPositions);
+		UE::DataflowSimulationGeometryCache::SavePackage(*GeometryCacheAsset);
 	}
 }
 
-namespace Dataflow::Private
+namespace UE::Dataflow::Private
 {
 	template<class T>
 	T* CreateOrLoad(const FString& PackageName)
@@ -412,7 +412,7 @@ namespace Dataflow::Private
 void UDataflowSimulationSceneDescription::NewGeometryCache()
 {
 	const UObject* const NamingAsset = CacheAsset? CacheAsset.Get() : nullptr;
-	GeometryCacheAsset = Dataflow::Private::NewGeometryCacheDialog(NamingAsset);
+	GeometryCacheAsset = UE::Dataflow::Private::NewGeometryCacheDialog(NamingAsset);
 }
 
 void UDataflowSimulationSceneDescription::SetSimulationScene(FDataflowSimulationScene* InSimulationScene)
diff --git a/Engine/Plugins/Experimental/Dataflow/Source/DataflowEditor/Private/Dataflow/DataflowSkeletonView.cpp b/Engine/Plugins/Experimental/Dataflow/Source/DataflowEditor/Private/Dataflow/DataflowSkeletonView.cpp
--- a/Engine/Plugins/Experimental/Dataflow/Source/DataflowEditor/Private/Dataflow/DataflowSkeletonView.cpp
+++ b/Engine/Plugins/Experimental/Dataflow/Source/DataflowEditor/Private/Dataflow/DataflowSkeletonView.cpp
@@ -88,7 +88,7 @@ void FDataflowSkeletonView::UpdateViewData()
 				{
 					if (const TObjectPtr<UDataflowSkeletalContent> SkeletalContent = Cast<UDataflowSkeletalContent>(GetEditorContent()))
 					{
-						if (TSharedPtr<Dataflow::FEngineContext> Context = SkeletalContent->GetDataflowContext())
+						if (TSharedPtr<UE::Dataflow::FEngineContext> Context = SkeletalContent->GetDataflowContext())
 						{
 							FManagedArrayCollection DefaultCollection;
 							const FManagedArrayCollection& Result = Output->GetValue(*Context, DefaultCollection);
diff --git a/Engine/Plugins/Experimental/Dataflow/Source/DataflowEditor/Private/Dataflow/DataflowToolTarget.cpp b/Engine/Plugins/Experimental/Dataflow/Source/DataflowEditor/Private/Dataflow/DataflowToolTarget.cpp
--- a/Engine/Plugins/Experimental/Dataflow/Source/DataflowEditor/Private/Dataflow/DataflowToolTarget.cpp
+++ b/Engine/Plugins/Experimental/Dataflow/Source/DataflowEditor/Private/Dataflow/DataflowToolTarget.cpp
@@ -178,7 +178,7 @@ UToolTarget* UDataflowReadOnlyToolTargetFactory::BuildTarget(UObject* SourceObje
 		UDataflowReadOnlyToolTarget* Target = NewObject<UDataflowReadOnlyToolTarget>();
 		Target->Asset = BaseContent->GetDataflowOwner();
 		Target->Dataflow = BaseContent->GetDataflowAsset();
-		Target->Context = Dataflow::GetContext(BaseContent);
+		Target->Context = UE::Dataflow::GetContext(BaseContent);
 
 		// @todo(brice) : I needed to comment this out?
 		//checkSlow(Target->Component.IsValid() && Requirements.AreSatisfiedBy(Target));
@@ -220,7 +220,7 @@ UToolTarget* UDataflowToolTargetFactory::BuildTarget(UObject* SourceObject, cons
 		UDataflowToolTarget* Target = NewObject<UDataflowToolTarget>();
 		Target->Asset = BaseContent->GetDataflowOwner();
 		Target->Dataflow = BaseContent->GetDataflowAsset();
-		Target->Context = Dataflow::GetContext(BaseContent);
+		Target->Context = UE::Dataflow::GetContext(BaseContent);
 		//checkSlow(Target->Component.IsValid() && Requirements.AreSatisfiedBy(Target));
 
 		return Target;
diff --git a/Engine/Plugins/Experimental/Dataflow/Source/DataflowEditor/Private/DataflowEditorTools/DataflowEditorWeightMapPaintTool.cpp b/Engine/Plugins/Experimental/Dataflow/Source/DataflowEditor/Private/DataflowEditorTools/DataflowEditorWeightMapPaintTool.cpp
--- a/Engine/Plugins/Experimental/Dataflow/Source/DataflowEditor/Private/DataflowEditorTools/DataflowEditorWeightMapPaintTool.cpp
+++ b/Engine/Plugins/Experimental/Dataflow/Source/DataflowEditor/Private/DataflowEditorTools/DataflowEditorWeightMapPaintTool.cpp
@@ -54,7 +54,7 @@ using namespace UE::Geometry;
 
 #define LOCTEXT_NAMESPACE "UDataflowEditorWeightMapPaintTool"
 
-namespace Dataflow::Private
+namespace UE::Dataflow::Private
 {
 	// probably should be something defined for the whole tool framework...
 #if WITH_EDITOR
@@ -68,15 +68,15 @@ namespace Dataflow::Private
 /*
  * ToolBuilder
  */
-void UDataflowEditorWeightMapPaintToolBuilder::GetSupportedConstructionViewModes(const UDataflowContextObject& ContextObject, TArray<const Dataflow::IDataflowConstructionViewMode*>& Modes) const
+void UDataflowEditorWeightMapPaintToolBuilder::GetSupportedConstructionViewModes(const UDataflowContextObject& ContextObject, TArray<const UE::Dataflow::IDataflowConstructionViewMode*>& Modes) const
 {
-	//Modes.Add(Dataflow::EDataflowPatternVertexType::Sim3D);
-	//Modes.Add(Dataflow::EDataflowPatternVertexType::Sim2D);
+	//Modes.Add(UE::Dataflow::EDataflowPatternVertexType::Sim3D);
+	//Modes.Add(UE::Dataflow::EDataflowPatternVertexType::Sim2D);
 }
 
 bool UDataflowEditorWeightMapPaintToolBuilder::CanBuildTool(const FToolBuilderState& SceneState) const
 {
-	auto HasManagedArrayCollection = [](const FDataflowNode* InDataflowNode, const TSharedPtr<Dataflow::FEngineContext> Context)
+	auto HasManagedArrayCollection = [](const FDataflowNode* InDataflowNode, const TSharedPtr<UE::Dataflow::FEngineContext> Context)
 	{
 		if (InDataflowNode && Context)
 		{
@@ -102,7 +102,7 @@ bool UDataflowEditorWeightMapPaintToolBuilder::CanBuildTool(const FToolBuilderSt
 				{
 					if (ContextObject->GetSelectedNode() == Component->Node)
 					{
-						if (const TSharedPtr<Dataflow::FEngineContext> EvaluationContext = ContextObject->GetDataflowContext())
+						if (const TSharedPtr<UE::Dataflow::FEngineContext> EvaluationContext = ContextObject->GetDataflowContext())
 						{
 							if (const FDataflowNode* PrimarySelection = ContextObject->GetSelectedNodeOfType<FDataflowCollectionAddScalarVertexPropertyNode>())
 							{
@@ -200,12 +200,12 @@ void UDataflowEditorWeightMapPaintTool::Setup()
 	Mesh->Attributes()->PrimaryColors()->CreateFromPredicate([](int ParentVID, int TriIDA, int TriIDB) {return true; }, 0.f);
 	FAxisAlignedBox3d Bounds = Mesh->GetBounds(true);
 
-	TFuture<void> PrecomputeFuture = Async(Dataflow::Private::WeightPaintToolAsyncExecTarget, [&]()
+	TFuture<void> PrecomputeFuture = Async(UE::Dataflow::Private::WeightPaintToolAsyncExecTarget, [&]()
 	{
 		PrecomputeFilterData();
 	});
 
-	TFuture<void> OctreeFuture = Async(Dataflow::Private::WeightPaintToolAsyncExecTarget, [&]()
+	TFuture<void> OctreeFuture = Async(UE::Dataflow::Private::WeightPaintToolAsyncExecTarget, [&]()
 	{
 		// initialize dynamic octree
 		if (Mesh->TriangleCount() > 100000)
@@ -351,7 +351,7 @@ void UDataflowEditorWeightMapPaintTool::Setup()
 	{
 		if (const TSharedPtr<const FManagedArrayCollection> Collection = DataflowEditorContextObject->GetRenderCollection())
 		{
-			using namespace Dataflow;
+			using namespace UE::Dataflow;
 			const FNonManifoldMappingSupport NonManifoldMapping(*Mesh);
 
 			bHaveDynamicMeshToWeightConversion = NonManifoldMapping.IsNonManifoldVertexInSource();
@@ -850,7 +850,7 @@ bool UDataflowEditorWeightMapPaintTool::SyncWeightBufferWithMesh(const FDynamicM
 	return (NumModified > 0);
 }
 
-namespace Dataflow
+namespace UE::Dataflow
 {
 	template<typename RealType>
 	static bool FindPolylineSelfIntersection(
@@ -998,7 +998,7 @@ void UDataflowEditorWeightMapPaintTool::OnPolyLassoFinished(const FCameraPolyLas
 	// Try to clip polyline to be closed, or closed-enough for winding evaluation to work.
 	// If that returns false, the polyline is "too open". In that case we will extend
 	// outwards from the endpoints and then try to create a closed very large polygon
-	if (Dataflow::ApproxSelfClipPolyline(Polyline) == false)
+	if (UE::Dataflow::ApproxSelfClipPolyline(Polyline) == false)
 	{
 		FVector2f StartDirOut = UE::Geometry::Normalized(Polyline[0] - Polyline[1]);
 		FLine2f StartLine(Polyline[0], StartDirOut);
@@ -1555,7 +1555,7 @@ void UDataflowEditorWeightMapPaintTool::OnTick(float DeltaTime)
 
 			// append updated ROI to modified region (async)
 			FDynamicMesh3* Mesh = GetSculptMesh();
-			TFuture<void> AccumulateROI = Async(Dataflow::Private::WeightPaintToolAsyncExecTarget, [&]()
+			TFuture<void> AccumulateROI = Async(UE::Dataflow::Private::WeightPaintToolAsyncExecTarget, [&]()
 			{
 				UE::Geometry::VertexToTriangleOneRing(Mesh, VertexROI, AccumulatedTriangleROI);
 			});
@@ -1784,7 +1784,7 @@ void UDataflowEditorWeightMapPaintTool::UpdateSelectedNode()
 // Change Tracking
 //
 
-namespace DataflowWeightPaintLocals
+namespace UE::DataflowWeightPaintLocals
 {
 
 	/**
@@ -1847,11 +1847,11 @@ void UDataflowEditorWeightMapPaintTool::EndChange()
 
 	TUniquePtr<FDynamicMeshChange> EditResult = ActiveWeightEditChangeTracker->EndChange();
 
-	TUniquePtr<DataflowWeightPaintLocals::FDataflowWeightPaintMeshChange> DataflowWeightPaintMeshChange =
-		MakeUnique<DataflowWeightPaintLocals::FDataflowWeightPaintMeshChange>(DynamicMeshComponent.Get(), MoveTemp(EditResult));
+	TUniquePtr<UE::DataflowWeightPaintLocals::FDataflowWeightPaintMeshChange> DataflowWeightPaintMeshChange =
+		MakeUnique<UE::DataflowWeightPaintLocals::FDataflowWeightPaintMeshChange>(DynamicMeshComponent.Get(), MoveTemp(EditResult));
 	ActiveWeightEditChangeTracker = nullptr;
 
-	TUniquePtr<TWrappedToolCommandChange<DataflowWeightPaintLocals::FDataflowWeightPaintMeshChange>> NewChange = MakeUnique<TWrappedToolCommandChange<DataflowWeightPaintLocals::FDataflowWeightPaintMeshChange>>();
+	TUniquePtr<TWrappedToolCommandChange<UE::DataflowWeightPaintLocals::FDataflowWeightPaintMeshChange>> NewChange = MakeUnique<TWrappedToolCommandChange<UE::DataflowWeightPaintLocals::FDataflowWeightPaintMeshChange>>();
 	NewChange->WrappedChange = MoveTemp(DataflowWeightPaintMeshChange);
 	NewChange->BeforeModify = [this](bool bRevert)
 	{
diff --git a/Engine/Plugins/Experimental/Dataflow/Source/DataflowEditor/Public/Dataflow/AssetDefinition_DataflowAsset.h b/Engine/Plugins/Experimental/Dataflow/Source/DataflowEditor/Public/Dataflow/AssetDefinition_DataflowAsset.h
--- a/Engine/Plugins/Experimental/Dataflow/Source/DataflowEditor/Public/Dataflow/AssetDefinition_DataflowAsset.h
+++ b/Engine/Plugins/Experimental/Dataflow/Source/DataflowEditor/Public/Dataflow/AssetDefinition_DataflowAsset.h
@@ -7,7 +7,7 @@
 #include "AssetDefinition_DataflowAsset.generated.h"
 
 
-namespace DataflowAssetDefinitionHelpers
+namespace UE::DataflowAssetDefinitionHelpers
 {
 	// Return true if we should proceed, false if we should re-open the dialog
 	DATAFLOWEDITOR_API bool CreateNewDataflowAsset(const UObject* Asset, UObject*& OutDataflowAsset);
diff --git a/Engine/Plugins/Experimental/Dataflow/Source/DataflowEditor/Public/Dataflow/DataflowConstructionScene.h b/Engine/Plugins/Experimental/Dataflow/Source/DataflowEditor/Public/Dataflow/DataflowConstructionScene.h
--- a/Engine/Plugins/Experimental/Dataflow/Source/DataflowEditor/Public/Dataflow/DataflowConstructionScene.h
+++ b/Engine/Plugins/Experimental/Dataflow/Source/DataflowEditor/Public/Dataflow/DataflowConstructionScene.h
@@ -93,7 +93,7 @@ private:
 	TArray<TObjectPtr<UInteractiveToolPropertySet>> PropertyObjectsToTick;
 
 	/** Last Rendered Timestamp */
-	Dataflow::FTimestamp LastRenderedTimestamp = Dataflow::FTimestamp::Invalid;
+	UE::Dataflow::FTimestamp LastRenderedTimestamp = UE::Dataflow::FTimestamp::Invalid;
 
 	/** Show the wireframe on render */ 
 	bool bConstructionViewWireframe = true;
diff --git a/Engine/Plugins/Experimental/Dataflow/Source/DataflowEditor/Public/Dataflow/DataflowConstructionViewportClient.h b/Engine/Plugins/Experimental/Dataflow/Source/DataflowEditor/Public/Dataflow/DataflowConstructionViewportClient.h
--- a/Engine/Plugins/Experimental/Dataflow/Source/DataflowEditor/Public/Dataflow/DataflowConstructionViewportClient.h
+++ b/Engine/Plugins/Experimental/Dataflow/Source/DataflowEditor/Public/Dataflow/DataflowConstructionViewportClient.h
@@ -19,7 +19,7 @@ class FDataflowPreviewSceneBase;
 class UInputBehaviorSet;
 class USelection;
 class FDataflowConstructionViewportClient;
-namespace Dataflow
+namespace UE::Dataflow
 {
 	class IDataflowConstructionViewMode;
 }
@@ -32,7 +32,7 @@ public:
 	FDataflowConstructionViewportClient(FEditorModeTools* InModeTools, FPreviewScene* InPreviewScene,  const bool bCouldTickScene,
 								  const TWeakPtr<SEditorViewport> InEditorViewportWidget = nullptr);
 
-	void SetConstructionViewMode(const Dataflow::IDataflowConstructionViewMode* InViewMode);
+	void SetConstructionViewMode(const UE::Dataflow::IDataflowConstructionViewMode* InViewMode);
 
 	// IInputBehaviorSource
 	// virtual const UInputBehaviorSet* GetInputBehaviors() const override;
@@ -71,7 +71,7 @@ private:
 	TWeakPtr<FUICommandList> ToolCommandList;
 
 	/** Construction view mode */
-	const Dataflow::IDataflowConstructionViewMode* ConstructionViewMode = nullptr;
+	const UE::Dataflow::IDataflowConstructionViewMode* ConstructionViewMode = nullptr;
 
 	/** Behavior set for the behavior UI */
 	TObjectPtr<UInputBehaviorSet> BehaviorSet;
diff --git a/Engine/Plugins/Experimental/Dataflow/Source/DataflowEditor/Public/Dataflow/DataflowEditor.h b/Engine/Plugins/Experimental/Dataflow/Source/DataflowEditor/Public/Dataflow/DataflowEditor.h
--- a/Engine/Plugins/Experimental/Dataflow/Source/DataflowEditor/Public/Dataflow/DataflowEditor.h
+++ b/Engine/Plugins/Experimental/Dataflow/Source/DataflowEditor/Public/Dataflow/DataflowEditor.h
@@ -47,7 +47,7 @@ public:
 	virtual void Initialize(const TArray<TObjectPtr<UObject>>& InObjects, const TSubclassOf<AActor>& InPreviewClass = nullptr) override;
 
 	/** Update the terminal contents */
-	void UpdateTerminalContents(const Dataflow::FTimestamp TimeStamp);
+	void UpdateTerminalContents(const UE::Dataflow::FTimestamp TimeStamp);
 	
 	/** Update the editor content */
 	void UpdateEditorContent();
@@ -70,10 +70,10 @@ private :
 	using ValidTerminalsType = TMap<TSharedPtr<FDataflowNode>,TObjectPtr<UDataflowBaseContent>>;
 
 	/** Remove invalid terminal contents from the container */
-	void RemoveTerminalContents(const TSharedPtr<Dataflow::FGraph>& DataflowGraph, ValidTerminalsType& ValidTerminals);
+	void RemoveTerminalContents(const TSharedPtr<UE::Dataflow::FGraph>& DataflowGraph, ValidTerminalsType& ValidTerminals);
 	
 	/** Add valid terminal nodes to the container */
-	void AddTerminalContents(const TSharedPtr<Dataflow::FGraph>& DataflowGraph, ValidTerminalsType& ValidTerminals);
+	void AddTerminalContents(const TSharedPtr<UE::Dataflow::FGraph>& DataflowGraph, ValidTerminalsType& ValidTerminals);
 	
 	// Dataflow editor is the owner of the object list to edit/process and the dataflow mode
 	// is the one holding the dynamic mesh components to be rendered in the viewport
@@ -90,4 +90,4 @@ private :
 	bool bHasTerminalsDirty = false;
 };
 
-DECLARE_LOG_CATEGORY_EXTERN(LogDataflowEditor, Log, All);
\ No newline at end of file
+DECLARE_LOG_CATEGORY_EXTERN(LogDataflowEditor, Log, All);
diff --git a/Engine/Plugins/Experimental/Dataflow/Source/DataflowEditor/Public/Dataflow/DataflowEditorCommands.h b/Engine/Plugins/Experimental/Dataflow/Source/DataflowEditor/Public/Dataflow/DataflowEditorCommands.h
--- a/Engine/Plugins/Experimental/Dataflow/Source/DataflowEditor/Public/Dataflow/DataflowEditorCommands.h
+++ b/Engine/Plugins/Experimental/Dataflow/Source/DataflowEditor/Public/Dataflow/DataflowEditorCommands.h
@@ -119,7 +119,7 @@ public:
 	* @param Asset  When Asset is non null, if the node is a terminal node, and if the node timestamp is more recent than InOutLastNodeTimestamp. then the node SetAssetValue method will be called on this asset.
 	* @return  The node that has been evaluated if any.
 	*/
-	static const FDataflowNode* EvaluateNode(Dataflow::FContext& Context, Dataflow::FTimestamp& InOutLastNodeTimestamp,
+	static const FDataflowNode* EvaluateNode(UE::Dataflow::FContext& Context, UE::Dataflow::FTimestamp& InOutLastNodeTimestamp,
 		const UDataflow* Dataflow, const FDataflowNode* Node, const FDataflowOutput* Output = nullptr, 
 		const FString& NodeName = FString(), UObject* Asset = nullptr);
 
@@ -146,7 +146,7 @@ public:
 	/*
 	*  OnPropertyValueChanged
 	*/
-	static void OnPropertyValueChanged(UDataflow* Graph, TSharedPtr<Dataflow::FEngineContext>& Context, Dataflow::FTimestamp& OutLastNodeTimestamp, const FPropertyChangedEvent& PropertyChangedEvent, const TSet<TObjectPtr<UObject>>& NewSelection = TSet<TObjectPtr<UObject>>());
+	static void OnPropertyValueChanged(UDataflow* Graph, TSharedPtr<UE::Dataflow::FEngineContext>& Context, UE::Dataflow::FTimestamp& OutLastNodeTimestamp, const FPropertyChangedEvent& PropertyChangedEvent, const TSet<TObjectPtr<UObject>>& NewSelection = TSet<TObjectPtr<UObject>>());
 	static void OnAssetPropertyValueChanged(TObjectPtr<UDataflowBaseContent> Content, const FPropertyChangedEvent& PropertyChangedEvent);
 
 	/*
@@ -180,3 +180,4 @@ public:
 	static void RenameNode(const TSharedPtr<SDataflowGraphEditor>& DataflowGraphEditor, UEdGraphNode* EdNode);
 };
 
+
diff --git a/Engine/Plugins/Experimental/Dataflow/Source/DataflowEditor/Public/Dataflow/DataflowEditorMode.h b/Engine/Plugins/Experimental/Dataflow/Source/DataflowEditor/Public/Dataflow/DataflowEditorMode.h
--- a/Engine/Plugins/Experimental/Dataflow/Source/DataflowEditor/Public/Dataflow/DataflowEditorMode.h
+++ b/Engine/Plugins/Experimental/Dataflow/Source/DataflowEditor/Public/Dataflow/DataflowEditorMode.h
@@ -17,7 +17,7 @@ struct FManagedArrayCollection;
 class UEdGraphNode;
 class FViewport;
 class UDataflowEditor;
-namespace Dataflow
+namespace UE::Dataflow
 {
 	class IDataflowConstructionViewMode;
 }
@@ -58,7 +58,7 @@ public:
 	* Construction View Mode
 	*/
 	void SetConstructionViewMode(const FName& NewViewModeName);
-	const Dataflow::IDataflowConstructionViewMode* GetConstructionViewMode() const;
+	const UE::Dataflow::IDataflowConstructionViewMode* GetConstructionViewMode() const;
 	bool CanChangeConstructionViewModeTo(const FName& NewViewModeName) const;
 
 	void ToggleConstructionViewWireframe();
@@ -180,7 +180,7 @@ private:
 	TWeakPtr<SDataflowGraphEditor> DataflowGraphEditor;
 
 	// Points to a view mode owned by FRenderingViewModeFactory
-	const Dataflow::IDataflowConstructionViewMode* ConstructionViewMode = nullptr;
+	const UE::Dataflow::IDataflowConstructionViewMode* ConstructionViewMode = nullptr;
 
 	// Whether we should restore the previous view mode when a tool ends
 	bool bShouldRestoreSavedConstructionViewMode = false;
diff --git a/Engine/Plugins/Experimental/Dataflow/Source/DataflowEditor/Public/Dataflow/DataflowEditorToolkit.h b/Engine/Plugins/Experimental/Dataflow/Source/DataflowEditor/Public/Dataflow/DataflowEditorToolkit.h
--- a/Engine/Plugins/Experimental/Dataflow/Source/DataflowEditor/Public/Dataflow/DataflowEditorToolkit.h
+++ b/Engine/Plugins/Experimental/Dataflow/Source/DataflowEditor/Public/Dataflow/DataflowEditorToolkit.h
@@ -93,7 +93,7 @@ protected:
 	void OnTabClosed(TSharedRef<SDockTab> Tab);
 
 	// Node evaluation
-	void EvaluateNode(const FDataflowNode* Node, const FDataflowOutput* Output, Dataflow::FTimestamp& InOutTimestamp);
+	void EvaluateNode(const FDataflowNode* Node, const FDataflowOutput* Output, UE::Dataflow::FTimestamp& InOutTimestamp);
 	void OnFinishEvaluate();
 
 private:
diff --git a/Engine/Plugins/Experimental/Dataflow/Source/DataflowEditor/Public/Dataflow/DataflowGraphEditor.h b/Engine/Plugins/Experimental/Dataflow/Source/DataflowEditor/Public/Dataflow/DataflowGraphEditor.h
--- a/Engine/Plugins/Experimental/Dataflow/Source/DataflowEditor/Public/Dataflow/DataflowGraphEditor.h
+++ b/Engine/Plugins/Experimental/Dataflow/Source/DataflowEditor/Public/Dataflow/DataflowGraphEditor.h
@@ -25,7 +25,7 @@ class FDataflowSNodeFactory;
 class UDataflow;
 class UDataflowEditor;
 struct FDataflowConnection;
-namespace Dataflow {
+namespace UE::Dataflow {
 	class FContext;
 }
 /**
@@ -146,7 +146,7 @@ public:
 	virtual FString GetReferencerName() const override { return TEXT("SDataflowGraphEditor"); }
 
 	/** FDataflowSNodeInterface */
-	virtual const TSharedPtr<Dataflow::FEngineContext> GetDataflowContext() const override;
+	virtual const TSharedPtr<UE::Dataflow::FEngineContext> GetDataflowContext() const override;
 
 private:
 	/** Add an additional option pin to all selected Dataflow nodes for those that overrides the AddPin function. */
@@ -188,3 +188,4 @@ private:
 	bool LeftAltKeyDown = false;
 	bool RightAltKeyDown = false;
 };
+
diff --git a/Engine/Plugins/Experimental/Dataflow/Source/DataflowEditor/Public/Dataflow/DataflowRenderingFactory.h b/Engine/Plugins/Experimental/Dataflow/Source/DataflowEditor/Public/Dataflow/DataflowRenderingFactory.h
--- a/Engine/Plugins/Experimental/Dataflow/Source/DataflowEditor/Public/Dataflow/DataflowRenderingFactory.h
+++ b/Engine/Plugins/Experimental/Dataflow/Source/DataflowEditor/Public/Dataflow/DataflowRenderingFactory.h
@@ -11,14 +11,14 @@
 #include "Dataflow/DataflowRenderingViewMode.h"
 #include "GeometryCollection/Facades/CollectionRenderingFacade.h"
 
-namespace Dataflow
+namespace UE::Dataflow
 {
 	class FContext;
 	typedef TPair<FString, FName> FRenderKey;
 
 	struct FGraphRenderingState 
 	{
-		FGraphRenderingState(const FGuid InGuid, const FDataflowNode* InNode, const FRenderingParameter& InParameters, Dataflow::FContext& InContext, const Dataflow::IDataflowConstructionViewMode& ViewMode)
+		FGraphRenderingState(const FGuid InGuid, const FDataflowNode* InNode, const FRenderingParameter& InParameters, UE::Dataflow::FContext& InContext, const UE::Dataflow::IDataflowConstructionViewMode& ViewMode)
 			: NodeGuid(InGuid)
 			, Node(InNode)
 			, RenderName(InParameters.Name)
@@ -46,7 +46,7 @@ namespace Dataflow
 			return Default;
 		}
 
-		const Dataflow::IDataflowConstructionViewMode& GetViewMode() const { return ViewMode; }
+		const UE::Dataflow::IDataflowConstructionViewMode& GetViewMode() const { return ViewMode; }
 
 	private:
 		const FGuid NodeGuid;
@@ -56,9 +56,9 @@ namespace Dataflow
 		FName RenderType;
 		TArray<FName> RenderOutputs;
 
-		Dataflow::FContext& Context;
+		UE::Dataflow::FContext& Context;
 
-		const Dataflow::IDataflowConstructionViewMode& ViewMode;
+		const UE::Dataflow::IDataflowConstructionViewMode& ViewMode;
 	};
 
 
@@ -71,7 +71,7 @@ namespace Dataflow
 		public:
 			virtual ~ICallbackInterface() = default;
 			virtual FRenderKey GetRenderKey() const = 0;
-			virtual bool CanRender(const Dataflow::IDataflowConstructionViewMode& ViewMode) const = 0;
+			virtual bool CanRender(const UE::Dataflow::IDataflowConstructionViewMode& ViewMode) const = 0;
 			virtual void Render(GeometryCollection::Facades::FRenderingFacade& RenderData, const FGraphRenderingState& State) = 0;
 		};
 
@@ -98,3 +98,5 @@ namespace Dataflow
 }
 
 
+
+
diff --git a/Engine/Plugins/Experimental/Dataflow/Source/DataflowEditor/Public/Dataflow/DataflowRenderingViewMode.h b/Engine/Plugins/Experimental/Dataflow/Source/DataflowEditor/Public/Dataflow/DataflowRenderingViewMode.h
--- a/Engine/Plugins/Experimental/Dataflow/Source/DataflowEditor/Public/Dataflow/DataflowRenderingViewMode.h
+++ b/Engine/Plugins/Experimental/Dataflow/Source/DataflowEditor/Public/Dataflow/DataflowRenderingViewMode.h
@@ -6,22 +6,22 @@
 #include "Editor/UnrealEdTypes.h"
 #include "ChaosLog.h"
 
-namespace Dataflow
+namespace UE::Dataflow
 {
 	//
 	// Dataflow Construction View Modes allow the user to choose different views on the same node output. The View Mode system serves two main purposes: 
 	// 
 	// 1) In the Dataflow Editor's Construction Viewport, the currently active IDataflowConstructionViewMode will determine which ELevelViewportType should be used for rendering.
 	//    In other words, the the ViewMode will determine whether we have a 3D perspective camera or a 2D orthographic camera.
-	// 2) In the Datflow Rendering Factory (Dataflow::FRenderingFactory), the rendering callbacks that are registered will also be able to return different 
+	// 2) In the Datflow Rendering Factory (UE::Dataflow::FRenderingFactory), the rendering callbacks that are registered will also be able to return different 
 	//    render data based on what the currently active IDataflowConstructionViewMode in the viewport is.
 	//
 	// For example, suppose you have created a new type FMyMeshType that contains both a 3D Mesh representation and a 2D Mesh representation (for, say, UVs.) If you want to
 	// write custom ViewModes for it, here is the process:
 	// 
 	// 1) Define two new classes, e.g. FMyMeshType2DViewMode (inheriting from FDataflowConstruction2DViewModeBase) and FMyMeshType3DViewMode (inheriting from FDataflowConstruction3DViewModeBase)
-	// 2) In your module initialization, register both view modes with the Dataflow::FRenderingViewModeFactory
-	// 3) Create a render function callback class and register it with Dataflow::FRenderingFactory. 
+	// 2) In your module initialization, register both view modes with the UE::Dataflow::FRenderingViewModeFactory
+	// 3) Create a render function callback class and register it with UE::Dataflow::FRenderingFactory. 
 	//    - The CanRender function should return true if the view mode is either of your two new custom view modes.
 	//    - The Render function should return different RenderCollections depending on what the current view mode is (which can be determined from the FGraphRenderingState parameter.)
 	//
diff --git a/Engine/Plugins/Experimental/Dataflow/Source/DataflowEditor/Public/Dataflow/DataflowSEditorInterface.h b/Engine/Plugins/Experimental/Dataflow/Source/DataflowEditor/Public/Dataflow/DataflowSEditorInterface.h
--- a/Engine/Plugins/Experimental/Dataflow/Source/DataflowEditor/Public/Dataflow/DataflowSEditorInterface.h
+++ b/Engine/Plugins/Experimental/Dataflow/Source/DataflowEditor/Public/Dataflow/DataflowSEditorInterface.h
@@ -19,7 +19,7 @@ public:
 	FDataflowSEditorInterface() {}
 
 	/** Dataflow editor content accessors */
-	virtual const TSharedPtr<Dataflow::FEngineContext> GetDataflowContext() const { return TSharedPtr<Dataflow::FEngineContext>(nullptr); };
+	virtual const TSharedPtr<UE::Dataflow::FEngineContext> GetDataflowContext() const { return TSharedPtr<UE::Dataflow::FEngineContext>(nullptr); };
 
 };
 
diff --git a/Engine/Plugins/Experimental/Dataflow/Source/DataflowEditor/Public/Dataflow/DataflowSimulationControls.h b/Engine/Plugins/Experimental/Dataflow/Source/DataflowEditor/Public/Dataflow/DataflowSimulationControls.h
--- a/Engine/Plugins/Experimental/Dataflow/Source/DataflowEditor/Public/Dataflow/DataflowSimulationControls.h
+++ b/Engine/Plugins/Experimental/Dataflow/Source/DataflowEditor/Public/Dataflow/DataflowSimulationControls.h
@@ -14,7 +14,7 @@ namespace UE::Dataflow
 {
 	/** Check if the simulation cache nodes have changed to trigger a reset */
 	bool ShouldResetWorld(const TObjectPtr<UDataflow>& SimulationGraph,
-		const TObjectPtr<UWorld>& SimulationWorld, ::Dataflow::FTimestamp& LastTimeStamp);
+		const TObjectPtr<UWorld>& SimulationWorld, UE::Dataflow::FTimestamp& LastTimeStamp);
 
 	/** Spawn an actor given a class type and attach it to the cache manager */
 	TObjectPtr<AActor> SpawnSimulatedActor(const TSubclassOf<AActor>& ActorClass,
diff --git a/Engine/Plugins/Experimental/Dataflow/Source/DataflowEditor/Public/Dataflow/DataflowSimulationGeometryCache.h b/Engine/Plugins/Experimental/Dataflow/Source/DataflowEditor/Public/Dataflow/DataflowSimulationGeometryCache.h
--- a/Engine/Plugins/Experimental/Dataflow/Source/DataflowEditor/Public/Dataflow/DataflowSimulationGeometryCache.h
+++ b/Engine/Plugins/Experimental/Dataflow/Source/DataflowEditor/Public/Dataflow/DataflowSimulationGeometryCache.h
@@ -11,7 +11,7 @@ class UGeometryCache;
 
 DECLARE_LOG_CATEGORY_EXTERN(LogDataflowSimulationGeometryCache, Log, All);
 
-namespace DataflowSimulationGeometryCache
+namespace UE::DataflowSimulationGeometryCache
 {
 	/** Saves positions into the geometry cache */
 	DATAFLOWEDITOR_API void SaveGeometryCache(UGeometryCache& GeometryCache, const USkinnedAsset& Asset, TConstArrayView<uint32> ImportedVertexNumbers, TArrayView<TArray<FVector3f>> PositionsToMoveFrom);
diff --git a/Engine/Plugins/Experimental/Dataflow/Source/DataflowEditor/Public/Dataflow/DataflowSimulationScene.h b/Engine/Plugins/Experimental/Dataflow/Source/DataflowEditor/Public/Dataflow/DataflowSimulationScene.h
--- a/Engine/Plugins/Experimental/Dataflow/Source/DataflowEditor/Public/Dataflow/DataflowSimulationScene.h
+++ b/Engine/Plugins/Experimental/Dataflow/Source/DataflowEditor/Public/Dataflow/DataflowSimulationScene.h
@@ -165,7 +165,7 @@ private:
 	int32 NumFrames;
 
 	/** Last context time stamp for which we regenerated the world */
-	Dataflow::FTimestamp LastTimeStamp = Dataflow::FTimestamp::Invalid;
+	UE::Dataflow::FTimestamp LastTimeStamp = UE::Dataflow::FTimestamp::Invalid;
 
 	/** Preview actor that will will be used to visualize the result of the simulation graph */
 	TObjectPtr<AActor> PreviewActor;
diff --git a/Engine/Plugins/Experimental/Dataflow/Source/DataflowEditor/Public/Dataflow/DataflowToolRegistry.h b/Engine/Plugins/Experimental/Dataflow/Source/DataflowEditor/Public/Dataflow/DataflowToolRegistry.h
--- a/Engine/Plugins/Experimental/Dataflow/Source/DataflowEditor/Public/Dataflow/DataflowToolRegistry.h
+++ b/Engine/Plugins/Experimental/Dataflow/Source/DataflowEditor/Public/Dataflow/DataflowToolRegistry.h
@@ -11,7 +11,7 @@ class UInteractiveToolBuilder;
 class FUICommandList;
 class UInteractiveTool;
 
-namespace Dataflow
+namespace UE::Dataflow
 {
 	class DATAFLOWEDITOR_API FDataflowToolRegistry
 	{
@@ -60,3 +60,4 @@ namespace Dataflow
 	};
 
 }
+
diff --git a/Engine/Plugins/Experimental/Dataflow/Source/DataflowEditor/Public/Dataflow/DataflowToolTarget.h b/Engine/Plugins/Experimental/Dataflow/Source/DataflowEditor/Public/Dataflow/DataflowToolTarget.h
--- a/Engine/Plugins/Experimental/Dataflow/Source/DataflowEditor/Public/Dataflow/DataflowToolTarget.h
+++ b/Engine/Plugins/Experimental/Dataflow/Source/DataflowEditor/Public/Dataflow/DataflowToolTarget.h
@@ -63,7 +63,7 @@ protected:
 	TObjectPtr<UObject> Asset = nullptr;
 
 	/**  Engine context to be used for dataflow evaluation */
-	TSharedPtr<Dataflow::FEngineContext> Context = nullptr;
+	TSharedPtr<UE::Dataflow::FEngineContext> Context = nullptr;
 };
 
 /**
diff --git a/Engine/Plugins/Experimental/Dataflow/Source/DataflowEditor/Public/DataflowEditorTools/DataflowEditorToolBuilder.h b/Engine/Plugins/Experimental/Dataflow/Source/DataflowEditor/Public/DataflowEditorTools/DataflowEditorToolBuilder.h
--- a/Engine/Plugins/Experimental/Dataflow/Source/DataflowEditor/Public/DataflowEditorTools/DataflowEditorToolBuilder.h
+++ b/Engine/Plugins/Experimental/Dataflow/Source/DataflowEditor/Public/DataflowEditorTools/DataflowEditorToolBuilder.h
@@ -7,7 +7,7 @@
 
 #include "DataflowEditorToolBuilder.generated.h"
 
-namespace Dataflow {
+namespace UE::Dataflow {
 	class IDataflowConstructionViewMode;
 }
 struct FToolBuilderState;
@@ -28,7 +28,7 @@ class IDataflowEditorToolBuilder
 public:
 
 	/** Returns all Construction View modes that this tool can operate in. The first element should be the preferred mode to switch to if necessary. */
-	virtual void GetSupportedConstructionViewModes(const UDataflowContextObject& ContextObject, TArray<const Dataflow::IDataflowConstructionViewMode*>& Modes) const = 0;
+	virtual void GetSupportedConstructionViewModes(const UDataflowContextObject& ContextObject, TArray<const UE::Dataflow::IDataflowConstructionViewMode*>& Modes) const = 0;
 
 	/** Returns whether or not view can be set to wireframe when this tool is active.. */
 	virtual bool CanSetConstructionViewWireframeActive() const { return true; }
diff --git a/Engine/Plugins/Experimental/Dataflow/Source/DataflowEditor/Public/DataflowEditorTools/DataflowEditorWeightMapPaintTool.h b/Engine/Plugins/Experimental/Dataflow/Source/DataflowEditor/Public/DataflowEditorTools/DataflowEditorWeightMapPaintTool.h
--- a/Engine/Plugins/Experimental/Dataflow/Source/DataflowEditor/Public/DataflowEditorTools/DataflowEditorWeightMapPaintTool.h
+++ b/Engine/Plugins/Experimental/Dataflow/Source/DataflowEditor/Public/DataflowEditorTools/DataflowEditorWeightMapPaintTool.h
@@ -63,7 +63,7 @@ public:
 	void SetEditorMode(UDataflowEditorMode* InMode) { Mode = InMode; }
 
 private:
-	virtual void GetSupportedConstructionViewModes(const UDataflowContextObject& ContextObject, TArray<const Dataflow::IDataflowConstructionViewMode*>& Modes) const override;
+	virtual void GetSupportedConstructionViewModes(const UDataflowContextObject& ContextObject, TArray<const UE::Dataflow::IDataflowConstructionViewMode*>& Modes) const override;
 	virtual bool CanBuildTool(const FToolBuilderState& SceneState) const override;
 	virtual UMeshSurfacePointTool* CreateNewTool(const FToolBuilderState& SceneState) const override;
 	virtual bool CanSetConstructionViewWireframeActive() const { return false; }
diff --git a/Engine/Plugins/Experimental/Dataflow/Source/DataflowEnginePlugin/Private/Dataflow/DataflowEngineContextCaching.cpp b/Engine/Plugins/Experimental/Dataflow/Source/DataflowEnginePlugin/Private/Dataflow/DataflowEngineContextCaching.cpp
--- a/Engine/Plugins/Experimental/Dataflow/Source/DataflowEnginePlugin/Private/Dataflow/DataflowEngineContextCaching.cpp
+++ b/Engine/Plugins/Experimental/Dataflow/Source/DataflowEnginePlugin/Private/Dataflow/DataflowEngineContextCaching.cpp
@@ -7,11 +7,11 @@
 #include "Dataflow/DataflowContextCachingFactory.h"
 #include "GeometryCollection/ManagedArrayCollection.h"
 
-namespace Dataflow
+namespace UE::Dataflow
 {
 	void ContextCachingCallbacks()
 	{
-		using namespace Dataflow;
+		using namespace UE::Dataflow;
 
 		/**
 		* Dataflow ContextCaching (FManagedArrayCollection) 
diff --git a/Engine/Plugins/Experimental/Dataflow/Source/DataflowEnginePlugin/Private/Dataflow/DataflowEnginePlugin.cpp b/Engine/Plugins/Experimental/Dataflow/Source/DataflowEnginePlugin/Private/Dataflow/DataflowEnginePlugin.cpp
--- a/Engine/Plugins/Experimental/Dataflow/Source/DataflowEnginePlugin/Private/Dataflow/DataflowEnginePlugin.cpp
+++ b/Engine/Plugins/Experimental/Dataflow/Source/DataflowEnginePlugin/Private/Dataflow/DataflowEnginePlugin.cpp
@@ -28,7 +28,7 @@ IMPLEMENT_MODULE( FDataflowEnginePlugin, DataflowEnginePlugin)
 
 void FDataflowEnginePlugin::StartupModule()
 {
-	Dataflow::ContextCachingCallbacks();
+	UE::Dataflow::ContextCachingCallbacks();
 	FModuleManager::Get().LoadModule("DataflowSimulation");
 }
 
diff --git a/Engine/Plugins/Experimental/Dataflow/Source/DataflowEnginePlugin/Public/Dataflow/DataflowComponent.h b/Engine/Plugins/Experimental/Dataflow/Source/DataflowEnginePlugin/Public/Dataflow/DataflowComponent.h
--- a/Engine/Plugins/Experimental/Dataflow/Source/DataflowEnginePlugin/Public/Dataflow/DataflowComponent.h
+++ b/Engine/Plugins/Experimental/Dataflow/Source/DataflowEnginePlugin/Public/Dataflow/DataflowComponent.h
@@ -10,7 +10,7 @@
 
 #include "DataflowComponent.generated.h"
 
-namespace Dataflow { class IDataflowConstructionViewMode; }
+namespace UE::Dataflow { class IDataflowConstructionViewMode; }
 
 /**
 *	UDataflowComponent
@@ -45,7 +45,7 @@ public:
 	const TArray<const UDataflowEdNode*>& GetRenderTargets() const {return RenderTargets;}
 
 	/** Context */
-	void SetContext(TSharedPtr<Dataflow::FContext> InContext) { Context = InContext; }
+	void SetContext(TSharedPtr<UE::Dataflow::FContext> InContext) { Context = InContext; }
 
 	/** RenderCollection */
 	void SetRenderingCollection(FManagedArrayCollection&& InCollection);
@@ -67,13 +67,13 @@ public:
 	/* View mode */
 	// NOTE: Currently UDataflowComponent is not used in the Dataflow Editor. Instead the FDataflowConstructionScene converts the FRenderingFacade to a UDynamicMeshComponent.
 	// If we do start using UDataflowComponent we will need to update the current View Mode as it's changed using this function.
-	void SetViewMode(const Dataflow::IDataflowConstructionViewMode* InViewMode)
+	void SetViewMode(const UE::Dataflow::IDataflowConstructionViewMode* InViewMode)
 	{
 		ViewMode = InViewMode;
 	}
 
 private:
-	TSharedPtr<Dataflow::FContext> Context;
+	TSharedPtr<UE::Dataflow::FContext> Context;
 	TArray<const UDataflowEdNode*> RenderTargets;
 	TObjectPtr< const UDataflow> Dataflow;
 	FManagedArrayCollection RenderCollection;
@@ -83,6 +83,6 @@ private:
 	bool bBoundsNeedsUpdate = true;
 	FBoxSphereBounds BoundingBox = FBoxSphereBounds(ForceInitToZero);
 	FDataflowSelectionState SelectionState = FDataflowSelectionState(FDataflowSelectionState::EMode::DSS_Dataflow_None);
-	const Dataflow::IDataflowConstructionViewMode* ViewMode;
+	const UE::Dataflow::IDataflowConstructionViewMode* ViewMode;
 };
 
diff --git a/Engine/Plugins/Experimental/Dataflow/Source/DataflowNodes/Private/Dataflow/DataflowCollectionAddScalarVertexPropertyNode.cpp b/Engine/Plugins/Experimental/Dataflow/Source/DataflowNodes/Private/Dataflow/DataflowCollectionAddScalarVertexPropertyNode.cpp
--- a/Engine/Plugins/Experimental/Dataflow/Source/DataflowNodes/Private/Dataflow/DataflowCollectionAddScalarVertexPropertyNode.cpp
+++ b/Engine/Plugins/Experimental/Dataflow/Source/DataflowNodes/Private/Dataflow/DataflowCollectionAddScalarVertexPropertyNode.cpp
@@ -45,13 +45,13 @@ TArray<FName> DataflowAddScalarVertexPropertyCallbackRegistry::GetTargetGroupNam
 	return UniqueNames;
 }
 
-TArray<Dataflow::FRenderingParameter> DataflowAddScalarVertexPropertyCallbackRegistry::GetRenderingParameters() const
+TArray<UE::Dataflow::FRenderingParameter> DataflowAddScalarVertexPropertyCallbackRegistry::GetRenderingParameters() const
 {
-	TArray<Dataflow::FRenderingParameter> UniqueParameters;
+	TArray<UE::Dataflow::FRenderingParameter> UniqueParameters;
 
 	for (const TPair<FName, TUniquePtr<IDataflowAddScalarVertexPropertyCallbacks>>& CallbacksEntry : AllCallbacks)
 	{
-		for (const Dataflow::FRenderingParameter& RenderingParameter : CallbacksEntry.Value->GetRenderingParameters())
+		for (const UE::Dataflow::FRenderingParameter& RenderingParameter : CallbacksEntry.Value->GetRenderingParameters())
 		{
 			UniqueParameters.AddUnique(RenderingParameter);
 		}
@@ -60,7 +60,7 @@ TArray<Dataflow::FRenderingParameter> DataflowAddScalarVertexPropertyCallbackReg
 }
 
 
-FDataflowCollectionAddScalarVertexPropertyNode::FDataflowCollectionAddScalarVertexPropertyNode(const Dataflow::FNodeParameters& InParam, FGuid InGuid)
+FDataflowCollectionAddScalarVertexPropertyNode::FDataflowCollectionAddScalarVertexPropertyNode(const UE::Dataflow::FNodeParameters& InParam, FGuid InGuid)
 	: FDataflowNode(InParam, InGuid)
 {
 	RegisterInputConnection(&Collection);
@@ -68,15 +68,15 @@ FDataflowCollectionAddScalarVertexPropertyNode::FDataflowCollectionAddScalarVert
 	RegisterOutputConnection(&AttributeKey);
 }
 
-TArray<::Dataflow::FRenderingParameter> FDataflowCollectionAddScalarVertexPropertyNode::GetRenderParametersImpl() const
+TArray<UE::Dataflow::FRenderingParameter> FDataflowCollectionAddScalarVertexPropertyNode::GetRenderParametersImpl() const
 {
 	return DataflowAddScalarVertexPropertyCallbackRegistry::Get().GetRenderingParameters();
 }
 
 
-void FDataflowCollectionAddScalarVertexPropertyNode::Evaluate(Dataflow::FContext& Context, const FDataflowOutput* Out) const
+void FDataflowCollectionAddScalarVertexPropertyNode::Evaluate(UE::Dataflow::FContext& Context, const FDataflowOutput* Out) const
 {
-	using namespace Dataflow;
+	using namespace UE::Dataflow;
 
 	if (Out->IsA<FManagedArrayCollection>(&Collection))
 	{
@@ -114,7 +114,7 @@ void FDataflowCollectionAddScalarVertexPropertyNode::Evaluate(Dataflow::FContext
 	}
 }
 
-void FDataflowCollectionAddScalarVertexPropertyNode::OnSelected(Dataflow::FContext& Context)
+void FDataflowCollectionAddScalarVertexPropertyNode::OnSelected(UE::Dataflow::FContext& Context)
 {
 	// Re-evaluate the input collection
 	FManagedArrayCollection InCollection = GetValue<FManagedArrayCollection>(Context, &Collection);
diff --git a/Engine/Plugins/Experimental/Dataflow/Source/DataflowNodes/Private/Dataflow/DataflowCollectionAttributeKeyNodes.cpp b/Engine/Plugins/Experimental/Dataflow/Source/DataflowNodes/Private/Dataflow/DataflowCollectionAttributeKeyNodes.cpp
--- a/Engine/Plugins/Experimental/Dataflow/Source/DataflowNodes/Private/Dataflow/DataflowCollectionAttributeKeyNodes.cpp
+++ b/Engine/Plugins/Experimental/Dataflow/Source/DataflowNodes/Private/Dataflow/DataflowCollectionAttributeKeyNodes.cpp
@@ -7,7 +7,7 @@
 
 #include UE_INLINE_GENERATED_CPP_BY_NAME(DataflowCollectionAttributeKeyNodes)
 
-namespace Dataflow
+namespace UE::Dataflow
 {
 	void DataflowCollectionAttributeKeyNodes()
 	{
@@ -17,14 +17,14 @@ namespace Dataflow
 }
 
 
-void FMakeAttributeKeyDataflowNode::Evaluate(Dataflow::FContext& Context, const FDataflowOutput* Out) const
+void FMakeAttributeKeyDataflowNode::Evaluate(UE::Dataflow::FContext& Context, const FDataflowOutput* Out) const
 {
 	FString AttributeName = GetValue<FString>(Context, &AttributeIn);
 	FString GroupName = GetValue<FString>(Context, &GroupIn);
 	SetValue(Context, FCollectionAttributeKey(GroupName, AttributeName), &AttributeKeyOut);
 }
 
-void FBreakAttributeKeyDataflowNode::Evaluate(Dataflow::FContext& Context, const FDataflowOutput* Out) const
+void FBreakAttributeKeyDataflowNode::Evaluate(UE::Dataflow::FContext& Context, const FDataflowOutput* Out) const
 {
 	FCollectionAttributeKey CollectionKey = GetValue<FCollectionAttributeKey>(Context, &AttributeKeyIn);
 	SetValue(Context, CollectionKey.Attribute, &AttributeOut);
diff --git a/Engine/Plugins/Experimental/Dataflow/Source/DataflowNodes/Private/Dataflow/DataflowContextOverridesNodes.cpp b/Engine/Plugins/Experimental/Dataflow/Source/DataflowNodes/Private/Dataflow/DataflowContextOverridesNodes.cpp
--- a/Engine/Plugins/Experimental/Dataflow/Source/DataflowNodes/Private/Dataflow/DataflowContextOverridesNodes.cpp
+++ b/Engine/Plugins/Experimental/Dataflow/Source/DataflowNodes/Private/Dataflow/DataflowContextOverridesNodes.cpp
@@ -6,7 +6,7 @@
 
 #include UE_INLINE_GENERATED_CPP_BY_NAME(DataflowContextOverridesNodes)
 
-namespace Dataflow
+namespace UE::Dataflow
 {
 	void RegisterContextOverridesNodes()
 	{
@@ -14,14 +14,14 @@ namespace Dataflow
 	}
 }
 
-void FFloatOverrideDataflowNode::Evaluate(Dataflow::FContext& Context, const FDataflowOutput* Out) const
+void FFloatOverrideDataflowNode::Evaluate(UE::Dataflow::FContext& Context, const FDataflowOutput* Out) const
 {
 	if (Out->IsA<float>(&ValueOut))
 	{
 		float Value = 0.f;
-		if (const Dataflow::FEngineContext* EngineContext = Context.AsType<Dataflow::FEngineContext>())
+		if (const UE::Dataflow::FEngineContext* EngineContext = Context.AsType<UE::Dataflow::FEngineContext>())
 		{
-			FString Result = Dataflow::Reflection::FindOverrideProperty< FString >(EngineContext->Owner, PropertyName, KeyName);
+			FString Result = UE::Dataflow::Reflection::FindOverrideProperty< FString >(EngineContext->Owner, PropertyName, KeyName);
 			Value = FCString::Atof(*Result);
 		}
 		SetValue(Context, Value, &ValueOut);
diff --git a/Engine/Plugins/Experimental/Dataflow/Source/DataflowNodes/Private/Dataflow/DataflowNodesPlugin.cpp b/Engine/Plugins/Experimental/Dataflow/Source/DataflowNodes/Private/Dataflow/DataflowNodesPlugin.cpp
--- a/Engine/Plugins/Experimental/Dataflow/Source/DataflowNodes/Private/Dataflow/DataflowNodesPlugin.cpp
+++ b/Engine/Plugins/Experimental/Dataflow/Source/DataflowNodes/Private/Dataflow/DataflowNodesPlugin.cpp
@@ -31,7 +31,7 @@ public:
 		return { FGeometryCollection::VerticesGroup };
 	}
 
-	virtual TArray<Dataflow::FRenderingParameter> GetRenderingParameters() const override
+	virtual TArray<UE::Dataflow::FRenderingParameter> GetRenderingParameters() const override
 	{
 		return { { TEXT("SurfaceRender"), FGeometryCollection::StaticType(), {TEXT("Collection")} } };
 	}
@@ -41,14 +41,14 @@ const FName FGeometryCollectionAddScalarVertexPropertyCallbacks::Name = FName("F
 
 void IDataflowNodesPlugin::StartupModule()
 {
-	Dataflow::RegisterSkeletalMeshNodes();
-	Dataflow::RegisterStaticMeshNodes();
-	Dataflow::RegisterSelectionNodes();
-	Dataflow::RegisterContextOverridesNodes();
-	Dataflow::DataflowCollectionAttributeKeyNodes();
+	UE::Dataflow::RegisterSkeletalMeshNodes();
+	UE::Dataflow::RegisterStaticMeshNodes();
+	UE::Dataflow::RegisterSelectionNodes();
+	UE::Dataflow::RegisterContextOverridesNodes();
+	UE::Dataflow::DataflowCollectionAttributeKeyNodes();
 	DATAFLOW_NODE_REGISTER_CREATION_FACTORY(FDataflowCollectionAddScalarVertexPropertyNode);
 
-	Dataflow::RegisterNodeFilter(FDataflowTerminalNode::StaticType());
+	UE::Dataflow::RegisterNodeFilter(FDataflowTerminalNode::StaticType());
 
 	DataflowAddScalarVertexPropertyCallbackRegistry::Get().RegisterCallbacks(MakeUnique<FGeometryCollectionAddScalarVertexPropertyCallbacks>());
 }
diff --git a/Engine/Plugins/Experimental/Dataflow/Source/DataflowNodes/Private/Dataflow/DataflowSelectionNodes.cpp b/Engine/Plugins/Experimental/Dataflow/Source/DataflowNodes/Private/Dataflow/DataflowSelectionNodes.cpp
--- a/Engine/Plugins/Experimental/Dataflow/Source/DataflowNodes/Private/Dataflow/DataflowSelectionNodes.cpp
+++ b/Engine/Plugins/Experimental/Dataflow/Source/DataflowNodes/Private/Dataflow/DataflowSelectionNodes.cpp
@@ -5,7 +5,7 @@
 
 #include UE_INLINE_GENERATED_CPP_BY_NAME(DataflowSelectionNodes)
 
-namespace Dataflow
+namespace UE::Dataflow
 {
 	void RegisterSelectionNodes()
 	{
@@ -13,7 +13,7 @@ namespace Dataflow
 	}
 };
 
-void FSelectionSetDataflowNode::Evaluate(Dataflow::FContext& Context, const FDataflowOutput* Out) const
+void FSelectionSetDataflowNode::Evaluate(UE::Dataflow::FContext& Context, const FDataflowOutput* Out) const
 {
 	if (Out->IsA<DataType>(&IndicesOut))
 	{
diff --git a/Engine/Plugins/Experimental/Dataflow/Source/DataflowNodes/Private/Dataflow/DataflowSkeletalMeshNodes.cpp b/Engine/Plugins/Experimental/Dataflow/Source/DataflowNodes/Private/Dataflow/DataflowSkeletalMeshNodes.cpp
--- a/Engine/Plugins/Experimental/Dataflow/Source/DataflowNodes/Private/Dataflow/DataflowSkeletalMeshNodes.cpp
+++ b/Engine/Plugins/Experimental/Dataflow/Source/DataflowNodes/Private/Dataflow/DataflowSkeletalMeshNodes.cpp
@@ -11,7 +11,7 @@
 
 #include UE_INLINE_GENERATED_CPP_BY_NAME(DataflowSkeletalMeshNodes)
 
-namespace Dataflow
+namespace UE::Dataflow
 {
 	void RegisterSkeletalMeshNodes()
 	{
@@ -22,7 +22,7 @@ namespace Dataflow
 	}
 }
 
-void FGetSkeletalMeshDataflowNode::Evaluate(Dataflow::FContext& Context, const FDataflowOutput* Out) const
+void FGetSkeletalMeshDataflowNode::Evaluate(UE::Dataflow::FContext& Context, const FDataflowOutput* Out) const
 {
 	typedef TObjectPtr<const USkeletalMesh> DataType;
 	if (Out->IsA<DataType>(&SkeletalMesh))
@@ -31,9 +31,9 @@ void FGetSkeletalMeshDataflowNode::Evaluate(Dataflow::FContext& Context, const F
 
 		if (!SkeletalMesh)
 		{
-			if (const Dataflow::FEngineContext* EngineContext = Context.AsType<Dataflow::FEngineContext>())
+			if (const UE::Dataflow::FEngineContext* EngineContext = Context.AsType<UE::Dataflow::FEngineContext>())
 			{
-				if (const USkeletalMesh* SkeletalMeshFromOwner = Dataflow::Reflection::FindObjectPtrProperty<USkeletalMesh>(
+				if (const USkeletalMesh* SkeletalMeshFromOwner = UE::Dataflow::Reflection::FindObjectPtrProperty<USkeletalMesh>(
 					EngineContext->Owner, PropertyName))
 				{
 					SetValue<DataType>(Context, DataType(SkeletalMeshFromOwner), &SkeletalMesh);
@@ -44,7 +44,7 @@ void FGetSkeletalMeshDataflowNode::Evaluate(Dataflow::FContext& Context, const F
 }
 
 
-void FGetSkeletonDataflowNode::Evaluate(Dataflow::FContext& Context, const FDataflowOutput* Out) const
+void FGetSkeletonDataflowNode::Evaluate(UE::Dataflow::FContext& Context, const FDataflowOutput* Out) const
 {
 	typedef TObjectPtr<const USkeleton> DataType;
 	if (Out->IsA<DataType>(&Skeleton))
@@ -53,9 +53,9 @@ void FGetSkeletonDataflowNode::Evaluate(Dataflow::FContext& Context, const FData
 
 		if (!Skeleton)
 		{
-			if (const Dataflow::FEngineContext* EngineContext = Context.AsType<Dataflow::FEngineContext>())
+			if (const UE::Dataflow::FEngineContext* EngineContext = Context.AsType<UE::Dataflow::FEngineContext>())
 			{
-				if (const USkeleton* SkeletonFromOwner = Dataflow::Reflection::FindObjectPtrProperty<USkeleton>(
+				if (const USkeleton* SkeletonFromOwner = UE::Dataflow::Reflection::FindObjectPtrProperty<USkeleton>(
 					EngineContext->Owner, PropertyName))
 				{
 					SetValue<DataType>(Context, DataType(SkeletonFromOwner), &Skeleton);
@@ -65,7 +65,7 @@ void FGetSkeletonDataflowNode::Evaluate(Dataflow::FContext& Context, const FData
 	}
 }
 
-void FSkeletalMeshBoneDataflowNode::Evaluate(Dataflow::FContext& Context, const FDataflowOutput* Out) const
+void FSkeletalMeshBoneDataflowNode::Evaluate(UE::Dataflow::FContext& Context, const FDataflowOutput* Out) const
 {
 	typedef TObjectPtr<const USkeletalMesh> InDataType;
 	if (Out->IsA<int>(&BoneIndexOut))
@@ -77,9 +77,9 @@ void FSkeletalMeshBoneDataflowNode::Evaluate(Dataflow::FContext& Context, const
 			FName LocalBoneName = BoneName;
 			if (LocalBoneName.IsNone())
 			{
-				if (const Dataflow::FEngineContext* EngineContext = Context.AsType<Dataflow::FEngineContext>())
+				if (const UE::Dataflow::FEngineContext* EngineContext = Context.AsType<UE::Dataflow::FEngineContext>())
 				{
-					LocalBoneName = FName(Dataflow::Reflection::FindOverrideProperty< FString >(EngineContext->Owner, PropertyName, FName("BoneName")));
+					LocalBoneName = FName(UE::Dataflow::Reflection::FindOverrideProperty< FString >(EngineContext->Owner, PropertyName, FName("BoneName")));
 				}
 			}
 
@@ -91,7 +91,7 @@ void FSkeletalMeshBoneDataflowNode::Evaluate(Dataflow::FContext& Context, const
 }
 
 
-void FSkeletalMeshReferenceTransformDataflowNode::Evaluate(Dataflow::FContext& Context, const FDataflowOutput* Out) const
+void FSkeletalMeshReferenceTransformDataflowNode::Evaluate(UE::Dataflow::FContext& Context, const FDataflowOutput* Out) const
 {
 	typedef TObjectPtr<const USkeletalMesh> InDataType;
 	if (Out->IsA<FTransform>(&TransformOut))
@@ -104,7 +104,7 @@ void FSkeletalMeshReferenceTransformDataflowNode::Evaluate(Dataflow::FContext& C
 			if (InDataType SkeletalMesh = GetValue<InDataType>(Context, &SkeletalMeshIn))
 			{
 				TArray<FTransform> ComponentPose;
-				Dataflow::Animation::GlobalTransforms(SkeletalMesh->GetRefSkeleton(), ComponentPose);
+				UE::Dataflow::Animation::GlobalTransforms(SkeletalMesh->GetRefSkeleton(), ComponentPose);
 				if (BoneIndex < ComponentPose.Num())
 				{
 					SetValue(Context, ComponentPose[BoneIndex], &TransformOut);
diff --git a/Engine/Plugins/Experimental/Dataflow/Source/DataflowNodes/Private/Dataflow/DataflowStaticMeshNodes.cpp b/Engine/Plugins/Experimental/Dataflow/Source/DataflowNodes/Private/Dataflow/DataflowStaticMeshNodes.cpp
--- a/Engine/Plugins/Experimental/Dataflow/Source/DataflowNodes/Private/Dataflow/DataflowStaticMeshNodes.cpp
+++ b/Engine/Plugins/Experimental/Dataflow/Source/DataflowNodes/Private/Dataflow/DataflowStaticMeshNodes.cpp
@@ -5,7 +5,7 @@
 
 #include UE_INLINE_GENERATED_CPP_BY_NAME(DataflowStaticMeshNodes)
 
-namespace Dataflow
+namespace UE::Dataflow
 {
 	void RegisterStaticMeshNodes()
 	{
@@ -13,7 +13,7 @@ namespace Dataflow
 	}
 }
 
-void FGetStaticMeshDataflowNode::Evaluate(Dataflow::FContext& Context, const FDataflowOutput* Out) const
+void FGetStaticMeshDataflowNode::Evaluate(UE::Dataflow::FContext& Context, const FDataflowOutput* Out) const
 {
 	typedef TObjectPtr<const UStaticMesh> DataType;
 	if (Out->IsA<DataType>(&StaticMesh))
@@ -24,9 +24,9 @@ void FGetStaticMeshDataflowNode::Evaluate(Dataflow::FContext& Context, const FDa
 		{
 			SetValue(Context, StaticMesh, &StaticMesh);
 		}
-		else if (const Dataflow::FEngineContext* EngineContext = Context.AsType<Dataflow::FEngineContext>())
+		else if (const UE::Dataflow::FEngineContext* EngineContext = Context.AsType<UE::Dataflow::FEngineContext>())
 		{
-			if (const UStaticMesh* StaticMeshFromOwner = Dataflow::Reflection::FindObjectPtrProperty<UStaticMesh>(
+			if (const UStaticMesh* StaticMeshFromOwner = UE::Dataflow::Reflection::FindObjectPtrProperty<UStaticMesh>(
 				EngineContext->Owner, PropertyName))
 			{
 				SetValue(Context, DataType(StaticMeshFromOwner), &StaticMesh);
diff --git a/Engine/Plugins/Experimental/Dataflow/Source/DataflowNodes/Private/Dataflow/DataflowTools.cpp b/Engine/Plugins/Experimental/Dataflow/Source/DataflowNodes/Private/Dataflow/DataflowTools.cpp
--- a/Engine/Plugins/Experimental/Dataflow/Source/DataflowNodes/Private/Dataflow/DataflowTools.cpp
+++ b/Engine/Plugins/Experimental/Dataflow/Source/DataflowNodes/Private/Dataflow/DataflowTools.cpp
@@ -7,7 +7,7 @@
 
 DEFINE_LOG_CATEGORY(LogDataflowNodes);
 
-namespace Dataflow
+namespace UE::Dataflow
 {
 	void FDataflowTools::LogAndToastWarning(const FDataflowNode& DataflowNode, const FText& Headline, const FText& Details)
 	{
@@ -30,4 +30,4 @@ namespace Dataflow
 		do { InOutString.TrimCharInline(TEXT('_'), &bCharsWereRemoved); } while (bCharsWereRemoved);
 	}
 
-}  // End namespace Dataflow
+}  // End namespace UE::Dataflow
diff --git a/Engine/Plugins/Experimental/Dataflow/Source/DataflowNodes/Public/Dataflow/DataflowCollectionAddScalarVertexPropertyNode.h b/Engine/Plugins/Experimental/Dataflow/Source/DataflowNodes/Public/Dataflow/DataflowCollectionAddScalarVertexPropertyNode.h
--- a/Engine/Plugins/Experimental/Dataflow/Source/DataflowNodes/Public/Dataflow/DataflowCollectionAddScalarVertexPropertyNode.h
+++ b/Engine/Plugins/Experimental/Dataflow/Source/DataflowNodes/Public/Dataflow/DataflowCollectionAddScalarVertexPropertyNode.h
@@ -16,7 +16,7 @@ public:
 	virtual ~IDataflowAddScalarVertexPropertyCallbacks() = default;
 	virtual FName GetName() const = 0;
 	virtual TArray<FName> GetTargetGroupNames() const = 0;
-	virtual TArray<Dataflow::FRenderingParameter> GetRenderingParameters() const = 0;
+	virtual TArray<UE::Dataflow::FRenderingParameter> GetRenderingParameters() const = 0;
 };
 
 class DataflowAddScalarVertexPropertyCallbackRegistry
@@ -32,7 +32,7 @@ public:
 
 	DATAFLOWNODES_API TArray<FName> GetTargetGroupNames() const;
 
-	DATAFLOWNODES_API TArray<Dataflow::FRenderingParameter> GetRenderingParameters() const;
+	DATAFLOWNODES_API TArray<UE::Dataflow::FRenderingParameter> GetRenderingParameters() const;
 
 private:
 
@@ -60,7 +60,7 @@ struct DATAFLOWNODES_API FDataflowCollectionAddScalarVertexPropertyNode : public
 	GENERATED_USTRUCT_BODY()
 	DATAFLOW_NODE_DEFINE_INTERNAL(FDataflowCollectionAddScalarVertexPropertyNode, "AddScalarVertexProperty", "Collection", "Add a saved scalar property to a collection")
 
-	virtual TArray<::Dataflow::FRenderingParameter> GetRenderParametersImpl() const override;
+	virtual TArray<UE::Dataflow::FRenderingParameter> GetRenderParametersImpl() const override;
 
 public:
 
@@ -80,14 +80,14 @@ public:
 	UPROPERTY(EditAnywhere, Category = "Vertex Attribute")
 	FScalarVertexPropertyGroup TargetGroup;
 
-	FDataflowCollectionAddScalarVertexPropertyNode(const Dataflow::FNodeParameters& InParam, FGuid InGuid = FGuid::NewGuid());
+	FDataflowCollectionAddScalarVertexPropertyNode(const UE::Dataflow::FNodeParameters& InParam, FGuid InGuid = FGuid::NewGuid());
 
 	/** Return a cached array of all the groups used by the input collection during at the time of the latest evaluation. */
 	const TArray<FName>& GetCachedCollectionGroupNames() const { return CachedCollectionGroupNames; }
 
 private:
-	virtual void Evaluate(Dataflow::FContext& Context, const FDataflowOutput* Out) const override;
-	virtual void OnSelected(Dataflow::FContext& Context) override;
+	virtual void Evaluate(UE::Dataflow::FContext& Context, const FDataflowOutput* Out) const override;
+	virtual void OnSelected(UE::Dataflow::FContext& Context) override;
 	virtual void OnDeselected() override;
 
 	TArray<FName> CachedCollectionGroupNames;
diff --git a/Engine/Plugins/Experimental/Dataflow/Source/DataflowNodes/Public/Dataflow/DataflowCollectionAttributeKeyNodes.h b/Engine/Plugins/Experimental/Dataflow/Source/DataflowNodes/Public/Dataflow/DataflowCollectionAttributeKeyNodes.h
--- a/Engine/Plugins/Experimental/Dataflow/Source/DataflowNodes/Public/Dataflow/DataflowCollectionAttributeKeyNodes.h
+++ b/Engine/Plugins/Experimental/Dataflow/Source/DataflowNodes/Public/Dataflow/DataflowCollectionAttributeKeyNodes.h
@@ -28,7 +28,7 @@ public:
 	UPROPERTY(meta = (DataflowOutput, DisplayName = "AttributeKey"))
 	FCollectionAttributeKey AttributeKeyOut;
 
-	FMakeAttributeKeyDataflowNode(const Dataflow::FNodeParameters& InParam, FGuid InGuid = FGuid::NewGuid())
+	FMakeAttributeKeyDataflowNode(const UE::Dataflow::FNodeParameters& InParam, FGuid InGuid = FGuid::NewGuid())
 		: FDataflowNode(InParam, InGuid)
 	{
 		RegisterInputConnection(&GroupIn);
@@ -36,7 +36,7 @@ public:
 		RegisterOutputConnection(&AttributeKeyOut);
 	}
 
-	virtual void Evaluate(Dataflow::FContext& Context, const FDataflowOutput* Out) const override;
+	virtual void Evaluate(UE::Dataflow::FContext& Context, const FDataflowOutput* Out) const override;
 };
 
 USTRUCT(meta = (DataflowGeometryCollection))
@@ -58,7 +58,7 @@ public:
 
 
 
-	FBreakAttributeKeyDataflowNode(const Dataflow::FNodeParameters& InParam, FGuid InGuid = FGuid::NewGuid())
+	FBreakAttributeKeyDataflowNode(const UE::Dataflow::FNodeParameters& InParam, FGuid InGuid = FGuid::NewGuid())
 		: FDataflowNode(InParam, InGuid)
 	{
 		RegisterInputConnection(&AttributeKeyIn);
@@ -66,10 +66,10 @@ public:
 		RegisterOutputConnection(&GroupOut);
 	}
 
-	virtual void Evaluate(Dataflow::FContext& Context, const FDataflowOutput* Out) const override;
+	virtual void Evaluate(UE::Dataflow::FContext& Context, const FDataflowOutput* Out) const override;
 };
 
-namespace Dataflow
+namespace UE::Dataflow
 {
 	void DataflowCollectionAttributeKeyNodes();
 }
diff --git a/Engine/Plugins/Experimental/Dataflow/Source/DataflowNodes/Public/Dataflow/DataflowContextOverridesNodes.h b/Engine/Plugins/Experimental/Dataflow/Source/DataflowNodes/Public/Dataflow/DataflowContextOverridesNodes.h
--- a/Engine/Plugins/Experimental/Dataflow/Source/DataflowNodes/Public/Dataflow/DataflowContextOverridesNodes.h
+++ b/Engine/Plugins/Experimental/Dataflow/Source/DataflowNodes/Public/Dataflow/DataflowContextOverridesNodes.h
@@ -24,17 +24,17 @@ public:
 	UPROPERTY(meta = (DataflowOutput, DisplayName = "Float"))
 	float ValueOut = 0.f;
 
-	FFloatOverrideDataflowNode(const Dataflow::FNodeParameters& InParam, FGuid InGuid = FGuid::NewGuid())
+	FFloatOverrideDataflowNode(const UE::Dataflow::FNodeParameters& InParam, FGuid InGuid = FGuid::NewGuid())
 		: FDataflowNode(InParam, InGuid)
 	{
 		RegisterOutputConnection(&ValueOut);
 	}
 
-	virtual void Evaluate(Dataflow::FContext& Context, const FDataflowOutput* Out) const override;
+	virtual void Evaluate(UE::Dataflow::FContext& Context, const FDataflowOutput* Out) const override;
 };
 
 
-namespace Dataflow
+namespace UE::Dataflow
 {
 	void RegisterContextOverridesNodes();
 }
diff --git a/Engine/Plugins/Experimental/Dataflow/Source/DataflowNodes/Public/Dataflow/DataflowSelectionNodes.h b/Engine/Plugins/Experimental/Dataflow/Source/DataflowNodes/Public/Dataflow/DataflowSelectionNodes.h
--- a/Engine/Plugins/Experimental/Dataflow/Source/DataflowNodes/Public/Dataflow/DataflowSelectionNodes.h
+++ b/Engine/Plugins/Experimental/Dataflow/Source/DataflowNodes/Public/Dataflow/DataflowSelectionNodes.h
@@ -25,17 +25,17 @@ public:
 	UPROPERTY(meta = (DataflowOutput, DisplayName = "SelectionSet"))
 		TArray<int32> IndicesOut;
 
-	FSelectionSetDataflowNode(const Dataflow::FNodeParameters& InParam, FGuid InGuid = FGuid::NewGuid())
+	FSelectionSetDataflowNode(const UE::Dataflow::FNodeParameters& InParam, FGuid InGuid = FGuid::NewGuid())
 		: FDataflowNode(InParam, InGuid)
 	{
 		RegisterOutputConnection(&IndicesOut);
 	}
 
-	virtual void Evaluate(Dataflow::FContext& Context, const FDataflowOutput* Out) const override;
+	virtual void Evaluate(UE::Dataflow::FContext& Context, const FDataflowOutput* Out) const override;
 };
 
 
-namespace Dataflow
+namespace UE::Dataflow
 {
 	void RegisterSelectionNodes();
 }
diff --git a/Engine/Plugins/Experimental/Dataflow/Source/DataflowNodes/Public/Dataflow/DataflowSkeletalMeshNodes.h b/Engine/Plugins/Experimental/Dataflow/Source/DataflowNodes/Public/Dataflow/DataflowSkeletalMeshNodes.h
--- a/Engine/Plugins/Experimental/Dataflow/Source/DataflowNodes/Public/Dataflow/DataflowSkeletalMeshNodes.h
+++ b/Engine/Plugins/Experimental/Dataflow/Source/DataflowNodes/Public/Dataflow/DataflowSkeletalMeshNodes.h
@@ -26,14 +26,14 @@ public:
 	UPROPERTY(EditAnywhere, Category = "Dataflow" )
 	FName PropertyName = "SkeletalMesh";
 
-	FGetSkeletalMeshDataflowNode(const Dataflow::FNodeParameters& InParam, FGuid InGuid = FGuid::NewGuid())
+	FGetSkeletalMeshDataflowNode(const UE::Dataflow::FNodeParameters& InParam, FGuid InGuid = FGuid::NewGuid())
 		: FDataflowNode(InParam, InGuid)
 	{
 		RegisterOutputConnection(&SkeletalMesh);
 	}
 
 
-	virtual void Evaluate(Dataflow::FContext& Context, const FDataflowOutput* Out) const override;
+	virtual void Evaluate(UE::Dataflow::FContext& Context, const FDataflowOutput* Out) const override;
 };
 
 USTRUCT()
@@ -50,14 +50,14 @@ public:
 	UPROPERTY(EditAnywhere, Category = "Dataflow")
 	FName PropertyName = "Skeleton";
 
-	FGetSkeletonDataflowNode(const Dataflow::FNodeParameters& InParam, FGuid InGuid = FGuid::NewGuid())
+	FGetSkeletonDataflowNode(const UE::Dataflow::FNodeParameters& InParam, FGuid InGuid = FGuid::NewGuid())
 		: FDataflowNode(InParam, InGuid)
 	{
 		RegisterOutputConnection(&Skeleton);
 	}
 
 
-	virtual void Evaluate(Dataflow::FContext& Context, const FDataflowOutput* Out) const override;
+	virtual void Evaluate(UE::Dataflow::FContext& Context, const FDataflowOutput* Out) const override;
 };
 
 
@@ -80,7 +80,7 @@ public:
 	UPROPERTY(EditAnywhere, Category = "Dataflow")
 	FName PropertyName = "Overrides";
 
-	FSkeletalMeshBoneDataflowNode(const Dataflow::FNodeParameters& InParam, FGuid InGuid = FGuid::NewGuid())
+	FSkeletalMeshBoneDataflowNode(const UE::Dataflow::FNodeParameters& InParam, FGuid InGuid = FGuid::NewGuid())
 		: FDataflowNode(InParam, InGuid)
 	{
 		RegisterInputConnection(&SkeletalMesh);
@@ -88,7 +88,7 @@ public:
 	}
 
 
-	virtual void Evaluate(Dataflow::FContext& Context, const FDataflowOutput* Out) const override;
+	virtual void Evaluate(UE::Dataflow::FContext& Context, const FDataflowOutput* Out) const override;
 };
 
 
@@ -109,7 +109,7 @@ public:
 	UPROPERTY(meta = (DataflowOutput, DisplayName = "Transform"))
 	FTransform TransformOut = FTransform::Identity;
 
-	FSkeletalMeshReferenceTransformDataflowNode(const Dataflow::FNodeParameters& InParam, FGuid InGuid = FGuid::NewGuid())
+	FSkeletalMeshReferenceTransformDataflowNode(const UE::Dataflow::FNodeParameters& InParam, FGuid InGuid = FGuid::NewGuid())
 		: FDataflowNode(InParam, InGuid)
 	{
 		RegisterInputConnection(&SkeletalMeshIn);
@@ -118,10 +118,10 @@ public:
 	}
 
 
-	virtual void Evaluate(Dataflow::FContext& Context, const FDataflowOutput* Out) const override;
+	virtual void Evaluate(UE::Dataflow::FContext& Context, const FDataflowOutput* Out) const override;
 };
 
-namespace Dataflow
+namespace UE::Dataflow
 {
 	void RegisterSkeletalMeshNodes();
 }
diff --git a/Engine/Plugins/Experimental/Dataflow/Source/DataflowNodes/Public/Dataflow/DataflowStaticMeshNodes.h b/Engine/Plugins/Experimental/Dataflow/Source/DataflowNodes/Public/Dataflow/DataflowStaticMeshNodes.h
--- a/Engine/Plugins/Experimental/Dataflow/Source/DataflowNodes/Public/Dataflow/DataflowStaticMeshNodes.h
+++ b/Engine/Plugins/Experimental/Dataflow/Source/DataflowNodes/Public/Dataflow/DataflowStaticMeshNodes.h
@@ -26,18 +26,18 @@ public:
 	UPROPERTY(EditAnywhere, Category = "Overrides")
 	FName PropertyName = "StaticMesh";
 
-	FGetStaticMeshDataflowNode(const Dataflow::FNodeParameters& InParam, FGuid InGuid = FGuid::NewGuid())
+	FGetStaticMeshDataflowNode(const UE::Dataflow::FNodeParameters& InParam, FGuid InGuid = FGuid::NewGuid())
 		: FDataflowNode(InParam, InGuid)
 	{
 		RegisterOutputConnection(&StaticMesh);
 	}
 
 
-	virtual void Evaluate(Dataflow::FContext& Context, const FDataflowOutput* Out) const override;
+	virtual void Evaluate(UE::Dataflow::FContext& Context, const FDataflowOutput* Out) const override;
 };
 
 
-namespace Dataflow
+namespace UE::Dataflow
 {
 	void RegisterStaticMeshNodes();
 }
diff --git a/Engine/Plugins/Experimental/Fracture/Source/FractureEngine/Private/FractureEngineUtility.cpp b/Engine/Plugins/Experimental/Fracture/Source/FractureEngine/Private/FractureEngineUtility.cpp
--- a/Engine/Plugins/Experimental/Fracture/Source/FractureEngine/Private/FractureEngineUtility.cpp
+++ b/Engine/Plugins/Experimental/Fracture/Source/FractureEngine/Private/FractureEngineUtility.cpp
@@ -88,7 +88,7 @@ void FFractureEngineUtility::DeconstructMesh(const UE::Geometry::FDynamicMesh3&
 	}
 }
 
-namespace Dataflow::Private
+namespace UE::Dataflow::Private
 {
 	static constexpr double VolDimScale = .01; // compute volumes in meters instead of cm, for saner units at typical scales
 }
@@ -119,7 +119,7 @@ static double GetMinVolume(double TotalVolume,
 	double MinVolume = 0.0;
 	if (InSelectionMethod == EFixTinyGeoGeometrySelectionMethod::VolumeCubeRoot)
 	{
-		MinVolume = InMinVolumeCubeRoot * Dataflow::Private::VolDimScale;
+		MinVolume = InMinVolumeCubeRoot * UE::Dataflow::Private::VolDimScale;
 		MinVolume = MinVolume * MinVolume * MinVolume;
 	}
 	else // EGeometrySelectionMethod::RelativeVolume
@@ -152,7 +152,7 @@ static bool CollectTargetBones(FGeometryCollection& Collection,
 		Collection,
 		TArrayView<int32>(), /*Empty array => use all transforms*/
 		OutVolumes,
-		Dataflow::Private::VolDimScale,
+		UE::Dataflow::Private::VolDimScale,
 		bClusterMode
 	);
 
diff --git a/Engine/Plugins/Experimental/GeometryCollectionPlugin/Source/GeometryCollectionDepNodes/Private/Dataflow/GeometryCollectionTransferVertexScalarAttributeDepNode.cpp b/Engine/Plugins/Experimental/GeometryCollectionPlugin/Source/GeometryCollectionDepNodes/Private/Dataflow/GeometryCollectionTransferVertexScalarAttributeDepNode.cpp
--- a/Engine/Plugins/Experimental/GeometryCollectionPlugin/Source/GeometryCollectionDepNodes/Private/Dataflow/GeometryCollectionTransferVertexScalarAttributeDepNode.cpp
+++ b/Engine/Plugins/Experimental/GeometryCollectionPlugin/Source/GeometryCollectionDepNodes/Private/Dataflow/GeometryCollectionTransferVertexScalarAttributeDepNode.cpp
@@ -83,7 +83,7 @@ namespace UE::Private {
 
 }
 
-void FGeometryCollectionTransferVertexScalarAttributeNode::Evaluate(Dataflow::FContext& Context, const FDataflowOutput* Out) const
+void FGeometryCollectionTransferVertexScalarAttributeNode::Evaluate(UE::Dataflow::FContext& Context, const FDataflowOutput* Out) const
 {
 	typedef TPair<FName, FName> NamePair;
 	FCollectionAttributeKey Key = GetValue(Context, &AttributeKey, AttributeKey);
diff --git a/Engine/Plugins/Experimental/GeometryCollectionPlugin/Source/GeometryCollectionDepNodes/Private/Dataflow/GeometryCollectionTransferVertexScalarAttributeDepNode.h b/Engine/Plugins/Experimental/GeometryCollectionPlugin/Source/GeometryCollectionDepNodes/Private/Dataflow/GeometryCollectionTransferVertexScalarAttributeDepNode.h
--- a/Engine/Plugins/Experimental/GeometryCollectionPlugin/Source/GeometryCollectionDepNodes/Private/Dataflow/GeometryCollectionTransferVertexScalarAttributeDepNode.h
+++ b/Engine/Plugins/Experimental/GeometryCollectionPlugin/Source/GeometryCollectionDepNodes/Private/Dataflow/GeometryCollectionTransferVertexScalarAttributeDepNode.h
@@ -114,7 +114,7 @@ public:
 		EditConditionHides))
 	float BoundMultiplier = 0.01f;
 
-	FGeometryCollectionTransferVertexScalarAttributeNode(const Dataflow::FNodeParameters& InParam, FGuid InGuid = FGuid::NewGuid())
+	FGeometryCollectionTransferVertexScalarAttributeNode(const UE::Dataflow::FNodeParameters& InParam, FGuid InGuid = FGuid::NewGuid())
 		: FDataflowNode(InParam, InGuid)
 	{
 		RegisterInputConnection(&Collection);
@@ -127,7 +127,7 @@ public:
 
 private:
 
-	virtual void Evaluate(Dataflow::FContext& Context, const FDataflowOutput* Out) const override;
+	virtual void Evaluate(UE::Dataflow::FContext& Context, const FDataflowOutput* Out) const override;
 
 
 	//
diff --git a/Engine/Plugins/Experimental/GeometryCollectionPlugin/Source/GeometryCollectionDepNodes/Private/Dataflow/SetVertexColorFromFloatArrayDepNode.cpp b/Engine/Plugins/Experimental/GeometryCollectionPlugin/Source/GeometryCollectionDepNodes/Private/Dataflow/SetVertexColorFromFloatArrayDepNode.cpp
--- a/Engine/Plugins/Experimental/GeometryCollectionPlugin/Source/GeometryCollectionDepNodes/Private/Dataflow/SetVertexColorFromFloatArrayDepNode.cpp
+++ b/Engine/Plugins/Experimental/GeometryCollectionPlugin/Source/GeometryCollectionDepNodes/Private/Dataflow/SetVertexColorFromFloatArrayDepNode.cpp
@@ -6,7 +6,7 @@
 
 #include UE_INLINE_GENERATED_CPP_BY_NAME(SetVertexColorFromFloatArrayDepNode)
 
-void FSetVertexColorInCollectionFromFloatArrayDataflowNode::Evaluate(Dataflow::FContext& Context, const FDataflowOutput* Out) const
+void FSetVertexColorInCollectionFromFloatArrayDataflowNode::Evaluate(UE::Dataflow::FContext& Context, const FDataflowOutput* Out) const
 {
 	if (Out->IsA<FManagedArrayCollection>(&Collection))
 	{
diff --git a/Engine/Plugins/Experimental/GeometryCollectionPlugin/Source/GeometryCollectionDepNodes/Private/Dataflow/SetVertexColorFromFloatArrayDepNode.h b/Engine/Plugins/Experimental/GeometryCollectionPlugin/Source/GeometryCollectionDepNodes/Private/Dataflow/SetVertexColorFromFloatArrayDepNode.h
--- a/Engine/Plugins/Experimental/GeometryCollectionPlugin/Source/GeometryCollectionDepNodes/Private/Dataflow/SetVertexColorFromFloatArrayDepNode.h
+++ b/Engine/Plugins/Experimental/GeometryCollectionPlugin/Source/GeometryCollectionDepNodes/Private/Dataflow/SetVertexColorFromFloatArrayDepNode.h
@@ -29,7 +29,7 @@ public:
 	UPROPERTY(EditAnywhere, Category = "Color")
 		float Scale = 1.f;
 
-	FSetVertexColorInCollectionFromFloatArrayDataflowNode(const Dataflow::FNodeParameters& InParam, FGuid InGuid = FGuid::NewGuid())
+	FSetVertexColorInCollectionFromFloatArrayDataflowNode(const UE::Dataflow::FNodeParameters& InParam, FGuid InGuid = FGuid::NewGuid())
 		: FDataflowNode(InParam, InGuid)
 	{
 		RegisterInputConnection(&Collection);
@@ -37,6 +37,6 @@ public:
 		RegisterOutputConnection(&Collection, &Collection);
 	}
 
-	virtual void Evaluate(Dataflow::FContext& Context, const FDataflowOutput* Out) const override;
+	virtual void Evaluate(UE::Dataflow::FContext& Context, const FDataflowOutput* Out) const override;
 
 };
\ No newline at end of file
diff --git a/Engine/Plugins/Experimental/GeometryCollectionPlugin/Source/GeometryCollectionDepNodes/Private/Dataflow/SetVertexColorFromVertexSelectionDepNode.cpp b/Engine/Plugins/Experimental/GeometryCollectionPlugin/Source/GeometryCollectionDepNodes/Private/Dataflow/SetVertexColorFromVertexSelectionDepNode.cpp
--- a/Engine/Plugins/Experimental/GeometryCollectionPlugin/Source/GeometryCollectionDepNodes/Private/Dataflow/SetVertexColorFromVertexSelectionDepNode.cpp
+++ b/Engine/Plugins/Experimental/GeometryCollectionPlugin/Source/GeometryCollectionDepNodes/Private/Dataflow/SetVertexColorFromVertexSelectionDepNode.cpp
@@ -4,7 +4,7 @@
 
 #include UE_INLINE_GENERATED_CPP_BY_NAME(SetVertexColorFromVertexSelectionDepNode)
 
-void FSetVertexColorInCollectionFromVertexSelectionDataflowNode::Evaluate(Dataflow::FContext& Context, const FDataflowOutput* Out) const
+void FSetVertexColorInCollectionFromVertexSelectionDataflowNode::Evaluate(UE::Dataflow::FContext& Context, const FDataflowOutput* Out) const
 {
 	if (Out->IsA<FManagedArrayCollection>(&Collection))
 	{
diff --git a/Engine/Plugins/Experimental/GeometryCollectionPlugin/Source/GeometryCollectionDepNodes/Private/Dataflow/SetVertexColorFromVertexSelectionDepNode.h b/Engine/Plugins/Experimental/GeometryCollectionPlugin/Source/GeometryCollectionDepNodes/Private/Dataflow/SetVertexColorFromVertexSelectionDepNode.h
--- a/Engine/Plugins/Experimental/GeometryCollectionPlugin/Source/GeometryCollectionDepNodes/Private/Dataflow/SetVertexColorFromVertexSelectionDepNode.h
+++ b/Engine/Plugins/Experimental/GeometryCollectionPlugin/Source/GeometryCollectionDepNodes/Private/Dataflow/SetVertexColorFromVertexSelectionDepNode.h
@@ -33,7 +33,7 @@ public:
 	UPROPERTY(EditAnywhere, Category = "Color", meta = (DisplayName = "NonSelected Color"))
 		FLinearColor NonSelectedColor = FLinearColor(FColor::Blue);
 
-	FSetVertexColorInCollectionFromVertexSelectionDataflowNode(const Dataflow::FNodeParameters& InParam, FGuid InGuid = FGuid::NewGuid())
+	FSetVertexColorInCollectionFromVertexSelectionDataflowNode(const UE::Dataflow::FNodeParameters& InParam, FGuid InGuid = FGuid::NewGuid())
 		: FDataflowNode(InParam, InGuid)
 	{
 		RegisterInputConnection(&Collection);
@@ -41,6 +41,6 @@ public:
 		RegisterOutputConnection(&Collection, &Collection);
 	}
 
-	virtual void Evaluate(Dataflow::FContext& Context, const FDataflowOutput* Out) const override;
+	virtual void Evaluate(UE::Dataflow::FContext& Context, const FDataflowOutput* Out) const override;
 
 };
diff --git a/Engine/Plugins/Experimental/GeometryCollectionPlugin/Source/GeometryCollectionNodes/Private/Dataflow/CreateColorArrayFromFloatArrayNode.cpp b/Engine/Plugins/Experimental/GeometryCollectionPlugin/Source/GeometryCollectionNodes/Private/Dataflow/CreateColorArrayFromFloatArrayNode.cpp
--- a/Engine/Plugins/Experimental/GeometryCollectionPlugin/Source/GeometryCollectionNodes/Private/Dataflow/CreateColorArrayFromFloatArrayNode.cpp
+++ b/Engine/Plugins/Experimental/GeometryCollectionPlugin/Source/GeometryCollectionNodes/Private/Dataflow/CreateColorArrayFromFloatArrayNode.cpp
@@ -5,7 +5,7 @@
 
 #include UE_INLINE_GENERATED_CPP_BY_NAME(CreateColorArrayFromFloatArrayNode)
 
-void FCreateColorArrayFromFloatArrayDataflowNode::Evaluate(Dataflow::FContext& Context, const FDataflowOutput* Out) const
+void FCreateColorArrayFromFloatArrayDataflowNode::Evaluate(UE::Dataflow::FContext& Context, const FDataflowOutput* Out) const
 {
 	if (Out->IsA(&ColorArray))
 	{
diff --git a/Engine/Plugins/Experimental/GeometryCollectionPlugin/Source/GeometryCollectionNodes/Private/Dataflow/GeometryCollectionArrayNodes.cpp b/Engine/Plugins/Experimental/GeometryCollectionPlugin/Source/GeometryCollectionNodes/Private/Dataflow/GeometryCollectionArrayNodes.cpp
--- a/Engine/Plugins/Experimental/GeometryCollectionPlugin/Source/GeometryCollectionNodes/Private/Dataflow/GeometryCollectionArrayNodes.cpp
+++ b/Engine/Plugins/Experimental/GeometryCollectionPlugin/Source/GeometryCollectionNodes/Private/Dataflow/GeometryCollectionArrayNodes.cpp
@@ -6,7 +6,7 @@
 
 //#include UE_INLINE_GENERATED_CPP_BY_NAME(GeometryCollectionArrayNodes)
 
-namespace Dataflow
+namespace UE::Dataflow
 {
 
 	void GeometryCollectionArrayNodes()
@@ -28,7 +28,7 @@ namespace Dataflow
 	}
 }
 
-void FGetFloatArrayElementDataflowNode::Evaluate(Dataflow::FContext& Context, const FDataflowOutput* Out) const
+void FGetFloatArrayElementDataflowNode::Evaluate(UE::Dataflow::FContext& Context, const FDataflowOutput* Out) const
 {
 	if (Out->IsA(&FloatValue))
 	{
@@ -40,7 +40,7 @@ void FGetFloatArrayElementDataflowNode::Evaluate(Dataflow::FContext& Context, co
 	}
 }
 
-void FFloatArrayToIntArrayDataflowNode::Evaluate(Dataflow::FContext& Context, const FDataflowOutput* Out) const
+void FFloatArrayToIntArrayDataflowNode::Evaluate(UE::Dataflow::FContext& Context, const FDataflowOutput* Out) const
 {
 	if (Out->IsA<TArray<int32>>(&IntArray))
 	{
@@ -106,7 +106,7 @@ void FFloatArrayToIntArrayDataflowNode::Evaluate(Dataflow::FContext& Context, co
 	}
 }
 
-void FGetArrayElementDataflowNode::Evaluate(Dataflow::FContext& Context, const FDataflowOutput* Out) const
+void FGetArrayElementDataflowNode::Evaluate(UE::Dataflow::FContext& Context, const FDataflowOutput* Out) const
 {
 	if (Out->IsA<FVector>(&Point))
 	{
@@ -121,7 +121,7 @@ void FGetArrayElementDataflowNode::Evaluate(Dataflow::FContext& Context, const F
 	}
 }
 
-void FGetNumArrayElementsDataflowNode::Evaluate(Dataflow::FContext& Context, const FDataflowOutput* Out) const
+void FGetNumArrayElementsDataflowNode::Evaluate(UE::Dataflow::FContext& Context, const FDataflowOutput* Out) const
 {
 	if (Out->IsA<int32>(&NumElements))
 	{
@@ -150,7 +150,7 @@ void FGetNumArrayElementsDataflowNode::Evaluate(Dataflow::FContext& Context, con
 	}
 }
 
-void FBoolArrayToFaceSelectionDataflowNode::Evaluate(Dataflow::FContext& Context, const FDataflowOutput* Out) const
+void FBoolArrayToFaceSelectionDataflowNode::Evaluate(UE::Dataflow::FContext& Context, const FDataflowOutput* Out) const
 {
 	if (Out->IsA<FDataflowFaceSelection>(&FaceSelection))
 	{
@@ -165,7 +165,7 @@ void FBoolArrayToFaceSelectionDataflowNode::Evaluate(Dataflow::FContext& Context
 }
 
 
-void FFloatArrayToVertexSelectionDataflowNode::Evaluate(Dataflow::FContext& Context, const FDataflowOutput* Out) const
+void FFloatArrayToVertexSelectionDataflowNode::Evaluate(UE::Dataflow::FContext& Context, const FDataflowOutput* Out) const
 {
 	if (Out->IsA<FDataflowVertexSelection>(&VertexSelection))
 	{
@@ -217,7 +217,7 @@ void FFloatArrayToVertexSelectionDataflowNode::Evaluate(Dataflow::FContext& Cont
 	}
 }
 
-void FFloatArrayNormalizeDataflowNode::Evaluate(Dataflow::FContext& Context, const FDataflowOutput* Out) const
+void FFloatArrayNormalizeDataflowNode::Evaluate(UE::Dataflow::FContext& Context, const FDataflowOutput* Out) const
 {
 	if (Out->IsA<TArray<float>>(&OutFloatArray))
 	{
@@ -306,7 +306,7 @@ void FFloatArrayNormalizeDataflowNode::Evaluate(Dataflow::FContext& Context, con
 }
 
 
-void FVectorArrayNormalizeDataflowNode::Evaluate(Dataflow::FContext& Context, const FDataflowOutput* Out) const
+void FVectorArrayNormalizeDataflowNode::Evaluate(UE::Dataflow::FContext& Context, const FDataflowOutput* Out) const
 {
 	if (Out->IsA<TArray<FVector>>(&OutVectorArray))
 	{
@@ -358,7 +358,7 @@ void FVectorArrayNormalizeDataflowNode::Evaluate(Dataflow::FContext& Context, co
 	}
 }
 
-void FUnionIntArraysDataflowNode::Evaluate(Dataflow::FContext& Context, const FDataflowOutput* Out) const
+void FUnionIntArraysDataflowNode::Evaluate(UE::Dataflow::FContext& Context, const FDataflowOutput* Out) const
 {
 	if (Out->IsA<TArray<int32>>(&OutArray))
 	{
@@ -377,7 +377,7 @@ void FUnionIntArraysDataflowNode::Evaluate(Dataflow::FContext& Context, const FD
 	}
 }
 
-void FRemoveFloatArrayElementDataflowNode::Evaluate(Dataflow::FContext& Context, const FDataflowOutput* Out) const
+void FRemoveFloatArrayElementDataflowNode::Evaluate(UE::Dataflow::FContext& Context, const FDataflowOutput* Out) const
 {
 	if (Out->IsA(&FloatArray))
 	{
@@ -531,7 +531,7 @@ static float ComputeSum(const TArray<float>& FloatArray, TArray<int32>& IndexArr
 
 // -----------------------------------------------------------------------------------------------
 
-void FFloatArrayComputeStatisticsDataflowNode::Evaluate(Dataflow::FContext& Context, const FDataflowOutput* Out) const
+void FFloatArrayComputeStatisticsDataflowNode::Evaluate(UE::Dataflow::FContext& Context, const FDataflowOutput* Out) const
 {
 	if (Out->IsA(&Value) || Out->IsA(&Indices))
 	{
@@ -587,7 +587,7 @@ void FFloatArrayComputeStatisticsDataflowNode::Evaluate(Dataflow::FContext& Cont
 	}
 }
 
-void FRandomizeFloatArrayDataflowNode::Evaluate(Dataflow::FContext& Context, const FDataflowOutput* Out) const
+void FRandomizeFloatArrayDataflowNode::Evaluate(UE::Dataflow::FContext& Context, const FDataflowOutput* Out) const
 {
 	if (Out->IsA(&FloatArray))
 	{
diff --git a/Engine/Plugins/Experimental/GeometryCollectionPlugin/Source/GeometryCollectionNodes/Private/Dataflow/GeometryCollectionAssetNodes.cpp b/Engine/Plugins/Experimental/GeometryCollectionPlugin/Source/GeometryCollectionNodes/Private/Dataflow/GeometryCollectionAssetNodes.cpp
--- a/Engine/Plugins/Experimental/GeometryCollectionPlugin/Source/GeometryCollectionNodes/Private/Dataflow/GeometryCollectionAssetNodes.cpp
+++ b/Engine/Plugins/Experimental/GeometryCollectionPlugin/Source/GeometryCollectionNodes/Private/Dataflow/GeometryCollectionAssetNodes.cpp
@@ -13,7 +13,7 @@
 #include "Materials/MaterialInterface.h"
 #include "PreviewScene.h"
 
-namespace Dataflow
+namespace UE::Dataflow
 {
 	void GeometryCollectionEngineAssetNodes()
 	{
@@ -34,7 +34,7 @@ namespace Dataflow
 
 // ===========================================================================================================================
 
-FGeometryCollectionTerminalDataflowNode::FGeometryCollectionTerminalDataflowNode(const Dataflow::FNodeParameters& InParam, FGuid InGuid)
+FGeometryCollectionTerminalDataflowNode::FGeometryCollectionTerminalDataflowNode(const UE::Dataflow::FNodeParameters& InParam, FGuid InGuid)
 	: FDataflowTerminalNode(InParam, InGuid)
 {
 	RegisterInputConnection(&Collection);
@@ -49,7 +49,7 @@ FGeometryCollectionTerminalDataflowNode::FGeometryCollectionTerminalDataflowNode
 
 
 
-void FGeometryCollectionTerminalDataflowNode::SetAssetValue(TObjectPtr<UObject> Asset, Dataflow::FContext& Context) const
+void FGeometryCollectionTerminalDataflowNode::SetAssetValue(TObjectPtr<UObject> Asset, UE::Dataflow::FContext& Context) const
 {
 	using FGeometryCollectionPtr = TSharedPtr<FGeometryCollection, ESPMode::ThreadSafe>;
 	using FMaterialArray = TArray<TObjectPtr<UMaterial>>;
@@ -86,7 +86,7 @@ void FGeometryCollectionTerminalDataflowNode::SetAssetValue(TObjectPtr<UObject>
 	}
 }
 
-void FGeometryCollectionTerminalDataflowNode::Evaluate(Dataflow::FContext& Context) const
+void FGeometryCollectionTerminalDataflowNode::Evaluate(UE::Dataflow::FContext& Context) const
 {
 	// simply forward all inputs to corresponding outputs
 	SafeForwardInput(Context, &Collection, &Collection);
@@ -97,18 +97,18 @@ void FGeometryCollectionTerminalDataflowNode::Evaluate(Dataflow::FContext& Conte
 
 // ===========================================================================================================================
 
-FGetGeometryCollectionAssetDataflowNode::FGetGeometryCollectionAssetDataflowNode(const Dataflow::FNodeParameters& InParam, FGuid InGuid)
+FGetGeometryCollectionAssetDataflowNode::FGetGeometryCollectionAssetDataflowNode(const UE::Dataflow::FNodeParameters& InParam, FGuid InGuid)
 	: FDataflowNode(InParam, InGuid)
 {
 	RegisterOutputConnection(&Asset);
 }
 
-void FGetGeometryCollectionAssetDataflowNode::Evaluate(Dataflow::FContext& Context, const FDataflowOutput* Out) const
+void FGetGeometryCollectionAssetDataflowNode::Evaluate(UE::Dataflow::FContext& Context, const FDataflowOutput* Out) const
 {
 	ensure(Out->IsA(&Asset));
 
 	TObjectPtr<UGeometryCollection> CollectionAsset(nullptr);
-	if (const Dataflow::FEngineContext* EngineContext = Context.AsType<Dataflow::FEngineContext>())
+	if (const UE::Dataflow::FEngineContext* EngineContext = Context.AsType<UE::Dataflow::FEngineContext>())
 	{
 		CollectionAsset = Cast<UGeometryCollection>(EngineContext->Owner);
 	}
@@ -117,14 +117,14 @@ void FGetGeometryCollectionAssetDataflowNode::Evaluate(Dataflow::FContext& Conte
 
 // ===========================================================================================================================
 
-FGetGeometryCollectionSourcesDataflowNode::FGetGeometryCollectionSourcesDataflowNode(const Dataflow::FNodeParameters& InParam, FGuid InGuid)
+FGetGeometryCollectionSourcesDataflowNode::FGetGeometryCollectionSourcesDataflowNode(const UE::Dataflow::FNodeParameters& InParam, FGuid InGuid)
 		: FDataflowNode(InParam, InGuid)
 {
 	RegisterInputConnection(&Asset);
 	RegisterOutputConnection(&Sources);
 }
 
-void FGetGeometryCollectionSourcesDataflowNode::Evaluate(Dataflow::FContext& Context, const FDataflowOutput* Out) const
+void FGetGeometryCollectionSourcesDataflowNode::Evaluate(UE::Dataflow::FContext& Context, const FDataflowOutput* Out) const
 {
 	ensure(Out->IsA(&Sources));
 
@@ -145,7 +145,7 @@ void FGetGeometryCollectionSourcesDataflowNode::Evaluate(Dataflow::FContext& Con
 
 // ===========================================================================================================================
 
-FCreateGeometryCollectionFromSourcesDataflowNode::FCreateGeometryCollectionFromSourcesDataflowNode(const Dataflow::FNodeParameters& InParam, FGuid InGuid)
+FCreateGeometryCollectionFromSourcesDataflowNode::FCreateGeometryCollectionFromSourcesDataflowNode(const UE::Dataflow::FNodeParameters& InParam, FGuid InGuid)
 		: FDataflowNode(InParam, InGuid)
 {
 	RegisterInputConnection(&Sources);
@@ -155,7 +155,7 @@ FCreateGeometryCollectionFromSourcesDataflowNode::FCreateGeometryCollectionFromS
 	RegisterOutputConnection(&InstancedMeshes);
 }
 
-void FCreateGeometryCollectionFromSourcesDataflowNode::Evaluate(Dataflow::FContext& Context, const FDataflowOutput* Out) const
+void FCreateGeometryCollectionFromSourcesDataflowNode::Evaluate(UE::Dataflow::FContext& Context, const FDataflowOutput* Out) const
 {
 	ensure(Out->IsA(&Collection) || Out->IsA(&Materials) || Out->IsA(&InstancedMeshes));
 	
@@ -235,7 +235,7 @@ void FCreateGeometryCollectionFromSourcesDataflowNode::Evaluate(Dataflow::FConte
 
 // ===========================================================================================================================
 
-FGeometryCollectionToCollectionDataflowNode::FGeometryCollectionToCollectionDataflowNode(const Dataflow::FNodeParameters& InParam, FGuid InGuid)
+FGeometryCollectionToCollectionDataflowNode::FGeometryCollectionToCollectionDataflowNode(const UE::Dataflow::FNodeParameters& InParam, FGuid InGuid)
 	: FDataflowNode(InParam, InGuid)
 {
 	RegisterOutputConnection(&Collection);
@@ -244,7 +244,7 @@ FGeometryCollectionToCollectionDataflowNode::FGeometryCollectionToCollectionData
 	RegisterOutputConnection(&InstancedMeshes);
 }
 
-void FGeometryCollectionToCollectionDataflowNode::Evaluate(Dataflow::FContext& Context, const FDataflowOutput* Out) const
+void FGeometryCollectionToCollectionDataflowNode::Evaluate(UE::Dataflow::FContext& Context, const FDataflowOutput* Out) const
 {
 	ensure(Out->IsA(&Collection) || Out->IsA(&Materials) || Out->IsA(&InstancedMeshes));
 
@@ -269,7 +269,7 @@ void FGeometryCollectionToCollectionDataflowNode::Evaluate(Dataflow::FContext& C
 
 // ===========================================================================================================================
 
-FBlueprintToCollectionDataflowNode::FBlueprintToCollectionDataflowNode(const Dataflow::FNodeParameters& InParam, FGuid InGuid)
+FBlueprintToCollectionDataflowNode::FBlueprintToCollectionDataflowNode(const UE::Dataflow::FNodeParameters& InParam, FGuid InGuid)
 	: FDataflowNode(InParam, InGuid)
 {
 	RegisterOutputConnection(&Collection);
@@ -278,7 +278,7 @@ FBlueprintToCollectionDataflowNode::FBlueprintToCollectionDataflowNode(const Dat
 	RegisterOutputConnection(&InstancedMeshes);
 }
 
-void FBlueprintToCollectionDataflowNode::Evaluate(Dataflow::FContext& Context, const FDataflowOutput* Out) const
+void FBlueprintToCollectionDataflowNode::Evaluate(UE::Dataflow::FContext& Context, const FDataflowOutput* Out) const
 {
 	ensure(Out->IsA(&Collection) || Out->IsA(&Materials) || Out->IsA(&InstancedMeshes));
 
