Change Dataflow namespace to UE::Dataflow to be ok with coding standard

diff --git a/Engine/Plugins/Experimental/GeometryCollectionPlugin/Source/GeometryCollectionNodes/Private/Dataflow/GeometryCollectionClusteringNodes.cpp b/Engine/Plugins/Experimental/GeometryCollectionPlugin/Source/GeometryCollectionNodes/Private/Dataflow/GeometryCollectionClusteringNodes.cpp
--- a/Engine/Plugins/Experimental/GeometryCollectionPlugin/Source/GeometryCollectionNodes/Private/Dataflow/GeometryCollectionClusteringNodes.cpp
+++ b/Engine/Plugins/Experimental/GeometryCollectionPlugin/Source/GeometryCollectionNodes/Private/Dataflow/GeometryCollectionClusteringNodes.cpp
@@ -33,7 +33,7 @@
 
 #include UE_INLINE_GENERATED_CPP_BY_NAME(GeometryCollectionClusteringNodes)
 
-namespace Dataflow
+namespace UE::Dataflow
 {
 
 	void GeometryCollectionClusteringNodes()
@@ -55,7 +55,7 @@ namespace Dataflow
 }
 
 
-void FAutoClusterDataflowNode::Evaluate(Dataflow::FContext& Context, const FDataflowOutput* Out) const
+void FAutoClusterDataflowNode::Evaluate(UE::Dataflow::FContext& Context, const FDataflowOutput* Out) const
 {
 	if (Out->IsA<FManagedArrayCollection>(&Collection))
 	{
@@ -111,7 +111,7 @@ void FAutoClusterDataflowNode::Evaluate(Dataflow::FContext& Context, const FData
 }
 
 
-void FClusterFlattenDataflowNode::Evaluate(Dataflow::FContext& Context, const FDataflowOutput* Out) const
+void FClusterFlattenDataflowNode::Evaluate(UE::Dataflow::FContext& Context, const FDataflowOutput* Out) const
 {
 	if (Out->IsA(&Collection) && IsConnected(&Collection))
 	{
@@ -155,7 +155,7 @@ void FClusterFlattenDataflowNode::Evaluate(Dataflow::FContext& Context, const FD
 }
 
 
-void FClusterUnclusterDataflowNode::Evaluate(Dataflow::FContext& Context, const FDataflowOutput* Out) const
+void FClusterUnclusterDataflowNode::Evaluate(UE::Dataflow::FContext& Context, const FDataflowOutput* Out) const
 {
 	if (Out->IsA<FManagedArrayCollection>(&Collection))
 	{
@@ -183,7 +183,7 @@ void FClusterUnclusterDataflowNode::Evaluate(Dataflow::FContext& Context, const
 }
 
 
-void FClusterDataflowNode::Evaluate(Dataflow::FContext& Context, const FDataflowOutput* Out) const
+void FClusterDataflowNode::Evaluate(UE::Dataflow::FContext& Context, const FDataflowOutput* Out) const
 {
 	if (Out->IsA(&Collection))
 	{
@@ -198,7 +198,7 @@ void FClusterDataflowNode::Evaluate(Dataflow::FContext& Context, const FDataflow
 	}
 }
 
-void FClusterMergeToNeighborsDataflowNode::Evaluate(Dataflow::FContext& Context, const FDataflowOutput* Out) const
+void FClusterMergeToNeighborsDataflowNode::Evaluate(UE::Dataflow::FContext& Context, const FDataflowOutput* Out) const
 {
 	if (Out->IsA(&Collection))
 	{
@@ -245,7 +245,7 @@ void FClusterMergeToNeighborsDataflowNode::Evaluate(Dataflow::FContext& Context,
 	}
 }
 
-void FClusterMergeDataflowNode::Evaluate(Dataflow::FContext& Context, const FDataflowOutput* Out) const
+void FClusterMergeDataflowNode::Evaluate(UE::Dataflow::FContext& Context, const FDataflowOutput* Out) const
 {
 	if (Out->IsA(&Collection))
 	{
@@ -263,7 +263,7 @@ void FClusterMergeDataflowNode::Evaluate(Dataflow::FContext& Context, const FDat
 	}
 }
 
-void FClusterIsolatedRootsDataflowNode::Evaluate(Dataflow::FContext& Context, const FDataflowOutput* Out) const
+void FClusterIsolatedRootsDataflowNode::Evaluate(UE::Dataflow::FContext& Context, const FDataflowOutput* Out) const
 {
 	if (Out->IsA(&Collection))
 	{
@@ -283,7 +283,7 @@ void FClusterIsolatedRootsDataflowNode::Evaluate(Dataflow::FContext& Context, co
 	}
 }
 
-void FClusterMagnetDataflowNode::Evaluate(Dataflow::FContext& Context, const FDataflowOutput* Out) const
+void FClusterMagnetDataflowNode::Evaluate(UE::Dataflow::FContext& Context, const FDataflowOutput* Out) const
 {
 	if (Out->IsA(&Collection))
 	{
diff --git a/Engine/Plugins/Experimental/GeometryCollectionPlugin/Source/GeometryCollectionNodes/Private/Dataflow/GeometryCollectionConversionNodes.cpp b/Engine/Plugins/Experimental/GeometryCollectionPlugin/Source/GeometryCollectionNodes/Private/Dataflow/GeometryCollectionConversionNodes.cpp
--- a/Engine/Plugins/Experimental/GeometryCollectionPlugin/Source/GeometryCollectionNodes/Private/Dataflow/GeometryCollectionConversionNodes.cpp
+++ b/Engine/Plugins/Experimental/GeometryCollectionPlugin/Source/GeometryCollectionNodes/Private/Dataflow/GeometryCollectionConversionNodes.cpp
@@ -5,7 +5,7 @@
 
 #include UE_INLINE_GENERATED_CPP_BY_NAME(GeometryCollectionConversionNodes)
 
-namespace Dataflow
+namespace UE::Dataflow
 {
 	void GeometryCollectionConversionNodes()
 	{
@@ -28,7 +28,7 @@ namespace Dataflow
 	}
 }
 
-void FVectorToStringDataflowNode::Evaluate(Dataflow::FContext& Context, const FDataflowOutput* Out) const
+void FVectorToStringDataflowNode::Evaluate(UE::Dataflow::FContext& Context, const FDataflowOutput* Out) const
 {
 	if (Out->IsA<FString>(&String))
 	{
@@ -37,7 +37,7 @@ void FVectorToStringDataflowNode::Evaluate(Dataflow::FContext& Context, const FD
 	}
 }
 
-void FFloatToStringDataflowNode::Evaluate(Dataflow::FContext& Context, const FDataflowOutput* Out) const
+void FFloatToStringDataflowNode::Evaluate(UE::Dataflow::FContext& Context, const FDataflowOutput* Out) const
 {
 	if (Out->IsA<FString>(&String))
 	{
@@ -46,7 +46,7 @@ void FFloatToStringDataflowNode::Evaluate(Dataflow::FContext& Context, const FDa
 	}
 }
 
-void FIntToStringDataflowNode::Evaluate(Dataflow::FContext& Context, const FDataflowOutput* Out) const
+void FIntToStringDataflowNode::Evaluate(UE::Dataflow::FContext& Context, const FDataflowOutput* Out) const
 {
 	if (Out->IsA<FString>(&String))
 	{
@@ -55,7 +55,7 @@ void FIntToStringDataflowNode::Evaluate(Dataflow::FContext& Context, const FData
 	}
 }
 
-void FBoolToStringDataflowNode::Evaluate(Dataflow::FContext& Context, const FDataflowOutput* Out) const
+void FBoolToStringDataflowNode::Evaluate(UE::Dataflow::FContext& Context, const FDataflowOutput* Out) const
 {
 	if (Out->IsA<FString>(&String))
 	{
@@ -64,7 +64,7 @@ void FBoolToStringDataflowNode::Evaluate(Dataflow::FContext& Context, const FDat
 	}
 }
 
-void FIntToFloatDataflowNode::Evaluate(Dataflow::FContext& Context, const FDataflowOutput* Out) const
+void FIntToFloatDataflowNode::Evaluate(UE::Dataflow::FContext& Context, const FDataflowOutput* Out) const
 {
 	if (Out->IsA(&Float))
 	{
@@ -73,7 +73,7 @@ void FIntToFloatDataflowNode::Evaluate(Dataflow::FContext& Context, const FDataf
 	}
 }
 
-void FIntToDoubleDataflowNode::Evaluate(Dataflow::FContext& Context, const FDataflowOutput* Out) const
+void FIntToDoubleDataflowNode::Evaluate(UE::Dataflow::FContext& Context, const FDataflowOutput* Out) const
 {
 	if (Out->IsA(&Double))
 	{
@@ -82,7 +82,7 @@ void FIntToDoubleDataflowNode::Evaluate(Dataflow::FContext& Context, const FData
 	}
 }
 
-void FFloatToDoubleDataflowNode::Evaluate(Dataflow::FContext& Context, const FDataflowOutput* Out) const
+void FFloatToDoubleDataflowNode::Evaluate(UE::Dataflow::FContext& Context, const FDataflowOutput* Out) const
 {
 	if (Out->IsA(&Double))
 	{
@@ -91,7 +91,7 @@ void FFloatToDoubleDataflowNode::Evaluate(Dataflow::FContext& Context, const FDa
 	}
 }
 
-void FFloatToIntDataflowNode::Evaluate(Dataflow::FContext& Context, const FDataflowOutput* Out) const
+void FFloatToIntDataflowNode::Evaluate(UE::Dataflow::FContext& Context, const FDataflowOutput* Out) const
 {
 	if (Out->IsA<int32>(&Int))
 	{
@@ -115,7 +115,7 @@ void FFloatToIntDataflowNode::Evaluate(Dataflow::FContext& Context, const FDataf
 	}
 }
 
-void FIntToBoolDataflowNode::Evaluate(Dataflow::FContext& Context, const FDataflowOutput* Out) const
+void FIntToBoolDataflowNode::Evaluate(UE::Dataflow::FContext& Context, const FDataflowOutput* Out) const
 {
 	if (Out->IsA<bool>(&Bool))
 	{
@@ -124,7 +124,7 @@ void FIntToBoolDataflowNode::Evaluate(Dataflow::FContext& Context, const FDatafl
 	}
 }
 
-void FBoolToIntDataflowNode::Evaluate(Dataflow::FContext& Context, const FDataflowOutput* Out) const
+void FBoolToIntDataflowNode::Evaluate(UE::Dataflow::FContext& Context, const FDataflowOutput* Out) const
 {
 	if (Out->IsA<int32>(&Int))
 	{
diff --git a/Engine/Plugins/Experimental/GeometryCollectionPlugin/Source/GeometryCollectionNodes/Private/Dataflow/GeometryCollectionDebugNodes.cpp b/Engine/Plugins/Experimental/GeometryCollectionPlugin/Source/GeometryCollectionNodes/Private/Dataflow/GeometryCollectionDebugNodes.cpp
--- a/Engine/Plugins/Experimental/GeometryCollectionPlugin/Source/GeometryCollectionNodes/Private/Dataflow/GeometryCollectionDebugNodes.cpp
+++ b/Engine/Plugins/Experimental/GeometryCollectionPlugin/Source/GeometryCollectionNodes/Private/Dataflow/GeometryCollectionDebugNodes.cpp
@@ -14,7 +14,7 @@
 
 #include UE_INLINE_GENERATED_CPP_BY_NAME(GeometryCollectionDebugNodes)
 
-namespace Dataflow
+namespace UE::Dataflow
 {
 	void GeometryCollectionDebugNodes()
 	{
@@ -33,7 +33,7 @@ namespace Dataflow
 	}
 }
 
-void FConvexHullToMeshDataflowNode::Evaluate(Dataflow::FContext& Context, const FDataflowOutput* Out) const
+void FConvexHullToMeshDataflowNode::Evaluate(UE::Dataflow::FContext& Context, const FDataflowOutput* Out) const
 {
 	using namespace UE::Geometry;
 
@@ -58,7 +58,7 @@ void FConvexHullToMeshDataflowNode::Evaluate(Dataflow::FContext& Context, const
 	}
 }
 
-void FSphereCoveringToMeshDataflowNode::Evaluate(Dataflow::FContext& Context, const FDataflowOutput* Out) const
+void FSphereCoveringToMeshDataflowNode::Evaluate(UE::Dataflow::FContext& Context, const FDataflowOutput* Out) const
 {
 	using namespace UE::Geometry;
 
@@ -89,7 +89,7 @@ void FSphereCoveringToMeshDataflowNode::Evaluate(Dataflow::FContext& Context, co
 	}
 }
 
-void FMeshToOBJStringDebugDataflowNode::Evaluate(Dataflow::FContext& Context, const FDataflowOutput* Out) const
+void FMeshToOBJStringDebugDataflowNode::Evaluate(UE::Dataflow::FContext& Context, const FDataflowOutput* Out) const
 {
 	if (Out->IsA(&StringOBJ))
 	{
@@ -124,7 +124,7 @@ void FMeshToOBJStringDebugDataflowNode::Evaluate(Dataflow::FContext& Context, co
 	}
 }
 
-void FWriteStringToFile::Evaluate(Dataflow::FContext& Context, const FDataflowOutput* Out) const
+void FWriteStringToFile::Evaluate(UE::Dataflow::FContext& Context, const FDataflowOutput* Out) const
 {
 	FString InContents = GetValue(Context, &FileContents);
 	FString InPath = GetValue(Context, &FilePath);
diff --git a/Engine/Plugins/Experimental/GeometryCollectionPlugin/Source/GeometryCollectionNodes/Private/Dataflow/GeometryCollectionEditNodes.cpp b/Engine/Plugins/Experimental/GeometryCollectionPlugin/Source/GeometryCollectionNodes/Private/Dataflow/GeometryCollectionEditNodes.cpp
--- a/Engine/Plugins/Experimental/GeometryCollectionPlugin/Source/GeometryCollectionNodes/Private/Dataflow/GeometryCollectionEditNodes.cpp
+++ b/Engine/Plugins/Experimental/GeometryCollectionPlugin/Source/GeometryCollectionNodes/Private/Dataflow/GeometryCollectionEditNodes.cpp
@@ -30,7 +30,7 @@
 
 #include UE_INLINE_GENERATED_CPP_BY_NAME(GeometryCollectionEditNodes)
 
-namespace Dataflow
+namespace UE::Dataflow
 {
 
 	void GeometryCollectionEditNodes()
@@ -47,7 +47,7 @@ namespace Dataflow
 }
 
 
-void FPruneInCollectionDataflowNode::Evaluate(Dataflow::FContext& Context, const FDataflowOutput* Out) const
+void FPruneInCollectionDataflowNode::Evaluate(UE::Dataflow::FContext& Context, const FDataflowOutput* Out) const
 {
 	if (Out->IsA<FManagedArrayCollection>(&Collection))
 	{
@@ -83,7 +83,7 @@ void FPruneInCollectionDataflowNode::Evaluate(Dataflow::FContext& Context, const
 }
 
 
-void FSetVisibilityInCollectionDataflowNode::Evaluate(Dataflow::FContext& Context, const FDataflowOutput* Out) const
+void FSetVisibilityInCollectionDataflowNode::Evaluate(UE::Dataflow::FContext& Context, const FDataflowOutput* Out) const
 {
 	if (Out->IsA<FManagedArrayCollection>(&Collection))
 	{
@@ -153,7 +153,7 @@ void FSetVisibilityInCollectionDataflowNode::Evaluate(Dataflow::FContext& Contex
 }
 
 
-void FMergeInCollectionDataflowNode::Evaluate(Dataflow::FContext& Context, const FDataflowOutput* Out) const
+void FMergeInCollectionDataflowNode::Evaluate(UE::Dataflow::FContext& Context, const FDataflowOutput* Out) const
 {
 	if (Out->IsA<FManagedArrayCollection>(&Collection))
 	{
diff --git a/Engine/Plugins/Experimental/GeometryCollectionPlugin/Source/GeometryCollectionNodes/Private/Dataflow/GeometryCollectionFieldNodes.cpp b/Engine/Plugins/Experimental/GeometryCollectionPlugin/Source/GeometryCollectionNodes/Private/Dataflow/GeometryCollectionFieldNodes.cpp
--- a/Engine/Plugins/Experimental/GeometryCollectionPlugin/Source/GeometryCollectionNodes/Private/Dataflow/GeometryCollectionFieldNodes.cpp
+++ b/Engine/Plugins/Experimental/GeometryCollectionPlugin/Source/GeometryCollectionNodes/Private/Dataflow/GeometryCollectionFieldNodes.cpp
@@ -47,7 +47,7 @@
 
 //#include UE_INLINE_GENERATED_CPP_BY_NAME(GeometryCollectionFieldNodes)
 
-namespace Dataflow
+namespace UE::Dataflow
 {
 	void GeometryCollectionFieldNodes()
 	{
@@ -209,7 +209,7 @@ static void RadialFalloffFieldProcess(const TArray<FVector3f>& InSamplePositions
 	}
 }
 
-void FRadialFalloffFieldDataflowNode::Evaluate(Dataflow::FContext& Context, const FDataflowOutput* Out) const
+void FRadialFalloffFieldDataflowNode::Evaluate(UE::Dataflow::FContext& Context, const FDataflowOutput* Out) const
 {
 	if (Out->IsA<TArray<float>>(&FieldFloatResult) ||
 		Out->IsA<TArray<int32>>(&FieldRemap) ||
@@ -437,7 +437,7 @@ static void BoxFalloffFieldProcess(const TArray<FVector3f>& InSamplePositions,
 	}
 }
 
-void FBoxFalloffFieldDataflowNode::Evaluate(Dataflow::FContext& Context, const FDataflowOutput* Out) const
+void FBoxFalloffFieldDataflowNode::Evaluate(UE::Dataflow::FContext& Context, const FDataflowOutput* Out) const
 {
 	if (Out->IsA<TArray<float>>(&FieldFloatResult) ||
 		Out->IsA<TArray<int32>>(&FieldRemap) ||
@@ -653,7 +653,7 @@ static void PlaneFalloffFieldProcess(const TArray<FVector3f>& InSamplePositions,
 	}
 }
 
-void FPlaneFalloffFieldDataflowNode::Evaluate(Dataflow::FContext& Context, const FDataflowOutput* Out) const
+void FPlaneFalloffFieldDataflowNode::Evaluate(UE::Dataflow::FContext& Context, const FDataflowOutput* Out) const
 {
 	if (Out->IsA<TArray<float>>(&FieldFloatResult) ||
 		Out->IsA<TArray<int32>>(&FieldRemap) ||
@@ -833,7 +833,7 @@ static void RadialIntMaskFieldProcess(const TArray<FVector3f>& InSamplePositions
 	}
 }
 
-void FRadialIntMaskFieldDataflowNode::Evaluate(Dataflow::FContext& Context, const FDataflowOutput* Out) const
+void FRadialIntMaskFieldDataflowNode::Evaluate(UE::Dataflow::FContext& Context, const FDataflowOutput* Out) const
 {
 	if (Out->IsA<TArray<int32>>(&FieldIntResult) ||
 		Out->IsA<TArray<int32>>(&FieldRemap) ||
@@ -963,7 +963,7 @@ static void UniformScalarFieldProcess(const TArray<FVector3f>& InSamplePositions
 	}
 }
 
-void FUniformScalarFieldDataflowNode::Evaluate(Dataflow::FContext& Context, const FDataflowOutput* Out) const
+void FUniformScalarFieldDataflowNode::Evaluate(UE::Dataflow::FContext& Context, const FDataflowOutput* Out) const
 {
 	if (Out->IsA<TArray<float>>(&FieldFloatResult) ||
 		Out->IsA<TArray<int32>>(&FieldRemap) ||
@@ -1078,7 +1078,7 @@ static void UniformVectorFieldProcess(const TArray<FVector3f>& InSamplePositions
 	}
 }
 
-void FUniformVectorFieldDataflowNode::Evaluate(Dataflow::FContext& Context, const FDataflowOutput* Out) const
+void FUniformVectorFieldDataflowNode::Evaluate(UE::Dataflow::FContext& Context, const FDataflowOutput* Out) const
 {
 	if (Out->IsA<TArray<FVector>>(&FieldVectorResult) ||
 		Out->IsA<TArray<int32>>(&FieldRemap) ||
@@ -1194,7 +1194,7 @@ static void RadialVectorFieldProcess(const TArray<FVector3f>& InSamplePositions,
 		OutFieldFloatResult[Idx] = ResultsView[Idx];
 	}
 }
-void FRadialVectorFieldDataflowNode::Evaluate(Dataflow::FContext& Context, const FDataflowOutput* Out) const
+void FRadialVectorFieldDataflowNode::Evaluate(UE::Dataflow::FContext& Context, const FDataflowOutput* Out) const
 {
 	if (Out->IsA<TArray<FVector>>(&FieldVectorResult) ||
 		Out->IsA<TArray<int32>>(&FieldRemap) ||
@@ -1308,7 +1308,7 @@ static void RandomVectorFieldProcess(const TArray<FVector3f>& InSamplePositions,
 	}
 }
 
-void FRandomVectorFieldDataflowNode::Evaluate(Dataflow::FContext& Context, const FDataflowOutput* Out) const
+void FRandomVectorFieldDataflowNode::Evaluate(UE::Dataflow::FContext& Context, const FDataflowOutput* Out) const
 {
 	if (Out->IsA<TArray<FVector>>(&FieldVectorResult) ||
 		Out->IsA<TArray<int32>>(&FieldRemap) ||
@@ -1425,7 +1425,7 @@ static void NoiseFieldProcess(const TArray<FVector3f>& InSamplePositions,
 	}
 }
 
-void FNoiseFieldDataflowNode::Evaluate(Dataflow::FContext& Context, const FDataflowOutput* Out) const
+void FNoiseFieldDataflowNode::Evaluate(UE::Dataflow::FContext& Context, const FDataflowOutput* Out) const
 {
 	if (Out->IsA<TArray<float>>(&FieldFloatResult) ||
 		Out->IsA<TArray<int32>>(&FieldRemap) ||
@@ -1540,7 +1540,7 @@ static void UniformIntegerFieldProcess(const TArray<FVector3f>& InSamplePosition
 	}
 }
 
-void FUniformIntegerFieldDataflowNode::Evaluate(Dataflow::FContext& Context, const FDataflowOutput* Out) const
+void FUniformIntegerFieldDataflowNode::Evaluate(UE::Dataflow::FContext& Context, const FDataflowOutput* Out) const
 {
 	if (Out->IsA<TArray<int32>>(&FieldIntResult) ||
 		Out->IsA<TArray<int32>>(&FieldRemap) ||
@@ -1671,7 +1671,7 @@ static void WaveScalarFieldProcess(const TArray<FVector3f>& InSamplePositions,
 	}
 }
 
-void FWaveScalarFieldDataflowNode::Evaluate(Dataflow::FContext& Context, const FDataflowOutput* Out) const
+void FWaveScalarFieldDataflowNode::Evaluate(UE::Dataflow::FContext& Context, const FDataflowOutput* Out) const
 {
 	if (Out->IsA<TArray<float>>(&FieldFloatResult) ||
 		Out->IsA<TArray<int32>>(&FieldRemap) ||
@@ -1838,7 +1838,7 @@ static void SumScalarEvaluate(const TArray<float>& InFieldFloatLeft,
 }
 
 
-void FSumScalarFieldDataflowNode::Evaluate(Dataflow::FContext& Context, const FDataflowOutput* Out) const
+void FSumScalarFieldDataflowNode::Evaluate(UE::Dataflow::FContext& Context, const FDataflowOutput* Out) const
 {
 	if (Out->IsA<TArray<float>>(&FieldFloatResult) ||
 		Out->IsA<TArray<int32>>(&FieldRemap))
@@ -2067,7 +2067,7 @@ static void SumVectorEvaluate(const TArray<float>* InFieldFloat,
 	}
 }
 
-void FSumVectorFieldDataflowNode::Evaluate(Dataflow::FContext& Context, const FDataflowOutput* Out) const
+void FSumVectorFieldDataflowNode::Evaluate(UE::Dataflow::FContext& Context, const FDataflowOutput* Out) const
 {
 	if (Out->IsA<TArray<FVector>>(&FieldVectorResult) ||
 		Out->IsA<TArray<int32>>(&FieldRemap))
@@ -2457,7 +2457,7 @@ void FSumVectorFieldDataflowNode::Evaluate(Dataflow::FContext& Context, const FD
 
 // ----------------------------------------------------------------------------
 
-void FFieldMakeDenseFloatArrayDataflowNode::Evaluate(Dataflow::FContext& Context, const FDataflowOutput* Out) const
+void FFieldMakeDenseFloatArrayDataflowNode::Evaluate(UE::Dataflow::FContext& Context, const FDataflowOutput* Out) const
 {
 	if (Out->IsA<TArray<float>>(&FieldFloatResult))
 	{
diff --git a/Engine/Plugins/Experimental/GeometryCollectionPlugin/Source/GeometryCollectionNodes/Private/Dataflow/GeometryCollectionFracturingNodes.cpp b/Engine/Plugins/Experimental/GeometryCollectionPlugin/Source/GeometryCollectionNodes/Private/Dataflow/GeometryCollectionFracturingNodes.cpp
--- a/Engine/Plugins/Experimental/GeometryCollectionPlugin/Source/GeometryCollectionNodes/Private/Dataflow/GeometryCollectionFracturingNodes.cpp
+++ b/Engine/Plugins/Experimental/GeometryCollectionPlugin/Source/GeometryCollectionNodes/Private/Dataflow/GeometryCollectionFracturingNodes.cpp
@@ -39,7 +39,7 @@
 
 #include UE_INLINE_GENERATED_CPP_BY_NAME(GeometryCollectionFracturingNodes)
 
-namespace Dataflow
+namespace UE::Dataflow
 {
 
 	void GeometryCollectionFracturingNodes()
@@ -65,7 +65,7 @@ namespace Dataflow
 	}
 }
 
-void FUniformScatterPointsDataflowNode::Evaluate(Dataflow::FContext& Context, const FDataflowOutput* Out) const
+void FUniformScatterPointsDataflowNode::Evaluate(UE::Dataflow::FContext& Context, const FDataflowOutput* Out) const
 {
 	if (Out->IsA<TArray<FVector>>(&Points))
 	{
@@ -94,7 +94,7 @@ void FUniformScatterPointsDataflowNode::Evaluate(Dataflow::FContext& Context, co
 	}
 }
 
-void FUniformScatterPointsDataflowNode_v2::Evaluate(Dataflow::FContext& Context, const FDataflowOutput* Out) const
+void FUniformScatterPointsDataflowNode_v2::Evaluate(UE::Dataflow::FContext& Context, const FDataflowOutput* Out) const
 {
 	if (Out->IsA<TArray<FVector>>(&Points))
 	{
@@ -123,7 +123,7 @@ void FUniformScatterPointsDataflowNode_v2::Evaluate(Dataflow::FContext& Context,
 	}
 }
 
-void FRadialScatterPointsDataflowNode::Evaluate(Dataflow::FContext& Context, const FDataflowOutput* Out) const
+void FRadialScatterPointsDataflowNode::Evaluate(UE::Dataflow::FContext& Context, const FDataflowOutput* Out) const
 {
 	if (Out->IsA<TArray<FVector>>(&Points))
 	{
@@ -153,7 +153,7 @@ void FRadialScatterPointsDataflowNode::Evaluate(Dataflow::FContext& Context, con
 	}
 }
 
-void FRadialScatterPointsDataflowNode_v2::Evaluate(Dataflow::FContext& Context, const FDataflowOutput* Out) const
+void FRadialScatterPointsDataflowNode_v2::Evaluate(UE::Dataflow::FContext& Context, const FDataflowOutput* Out) const
 {
 	if (Out->IsA<TArray<FVector>>(&Points))
 	{
@@ -267,7 +267,7 @@ void FRadialScatterPointsDataflowNode_v2::Evaluate(Dataflow::FContext& Context,
 	}
 }
 
-void FGridScatterPointsDataflowNode::Evaluate(Dataflow::FContext& Context, const FDataflowOutput* Out) const
+void FGridScatterPointsDataflowNode::Evaluate(UE::Dataflow::FContext& Context, const FDataflowOutput* Out) const
 {
 	if (Out->IsA<TArray<FVector>>(&Points))
 	{
@@ -323,7 +323,7 @@ void FGridScatterPointsDataflowNode::Evaluate(Dataflow::FContext& Context, const
 	}
 }
 
-void FVoronoiFractureDataflowNode::Evaluate(Dataflow::FContext& Context, const FDataflowOutput* Out) const
+void FVoronoiFractureDataflowNode::Evaluate(UE::Dataflow::FContext& Context, const FDataflowOutput* Out) const
 {
 	if (Out->IsA<FManagedArrayCollection>(&Collection))
 	{
@@ -363,7 +363,7 @@ void FVoronoiFractureDataflowNode::Evaluate(Dataflow::FContext& Context, const F
 	}
 }
 
-void FVoronoiFractureDataflowNode_v2::Evaluate(Dataflow::FContext& Context, const FDataflowOutput* Out) const
+void FVoronoiFractureDataflowNode_v2::Evaluate(UE::Dataflow::FContext& Context, const FDataflowOutput* Out) const
 {
 	if (Out->IsA<FManagedArrayCollection>(&Collection) ||
 		Out->IsA<FDataflowTransformSelection>(&TransformSelection) ||
@@ -452,7 +452,7 @@ void FVoronoiFractureDataflowNode_v2::Evaluate(Dataflow::FContext& Context, cons
 	}
 }
 
-void FPlaneCutterDataflowNode::Evaluate(Dataflow::FContext& Context, const FDataflowOutput* Out) const
+void FPlaneCutterDataflowNode::Evaluate(UE::Dataflow::FContext& Context, const FDataflowOutput* Out) const
 {
 	if (Out->IsA<FManagedArrayCollection>(&Collection))
 	{
@@ -493,7 +493,7 @@ void FPlaneCutterDataflowNode::Evaluate(Dataflow::FContext& Context, const FData
 	}
 }
 
-void FPlaneCutterDataflowNode_v2::Evaluate(Dataflow::FContext& Context, const FDataflowOutput* Out) const
+void FPlaneCutterDataflowNode_v2::Evaluate(UE::Dataflow::FContext& Context, const FDataflowOutput* Out) const
 {
 	if (Out->IsA<FManagedArrayCollection>(&Collection) ||
 		Out->IsA<FDataflowTransformSelection>(&TransformSelection) ||
@@ -582,7 +582,7 @@ void FPlaneCutterDataflowNode_v2::Evaluate(Dataflow::FContext& Context, const FD
 	}
 }
 
-void FExplodedViewDataflowNode::Evaluate(Dataflow::FContext& Context, const FDataflowOutput* Out) const
+void FExplodedViewDataflowNode::Evaluate(UE::Dataflow::FContext& Context, const FDataflowOutput* Out) const
 {
 	if (Out->IsA<FManagedArrayCollection>(&Collection))
 	{
@@ -594,7 +594,7 @@ void FExplodedViewDataflowNode::Evaluate(Dataflow::FContext& Context, const FDat
 	}
 }
 
-void FSliceCutterDataflowNode::Evaluate(Dataflow::FContext& Context, const FDataflowOutput* Out) const
+void FSliceCutterDataflowNode::Evaluate(UE::Dataflow::FContext& Context, const FDataflowOutput* Out) const
 {
 	if (Out->IsA<FManagedArrayCollection>(&Collection) ||
 		Out->IsA<FDataflowTransformSelection>(&TransformSelection) ||
@@ -686,7 +686,7 @@ void FSliceCutterDataflowNode::Evaluate(Dataflow::FContext& Context, const FData
 	}
 }
 
-void FBrickCutterDataflowNode::Evaluate(Dataflow::FContext& Context, const FDataflowOutput* Out) const
+void FBrickCutterDataflowNode::Evaluate(UE::Dataflow::FContext& Context, const FDataflowOutput* Out) const
 {
 	if (Out->IsA<FManagedArrayCollection>(&Collection) ||
 		Out->IsA<FDataflowTransformSelection>(&TransformSelection) ||
@@ -791,7 +791,7 @@ void FBrickCutterDataflowNode::Evaluate(Dataflow::FContext& Context, const FData
 	}
 }
 
-void FMeshCutterDataflowNode::Evaluate(Dataflow::FContext& Context, const FDataflowOutput* Out) const
+void FMeshCutterDataflowNode::Evaluate(UE::Dataflow::FContext& Context, const FDataflowOutput* Out) const
 {
 	if (Out->IsA(&Collection) ||
 		Out->IsA(&TransformSelection) ||
diff --git a/Engine/Plugins/Experimental/GeometryCollectionPlugin/Source/GeometryCollectionNodes/Private/Dataflow/GeometryCollectionMakeNodes.cpp b/Engine/Plugins/Experimental/GeometryCollectionPlugin/Source/GeometryCollectionNodes/Private/Dataflow/GeometryCollectionMakeNodes.cpp
--- a/Engine/Plugins/Experimental/GeometryCollectionPlugin/Source/GeometryCollectionNodes/Private/Dataflow/GeometryCollectionMakeNodes.cpp
+++ b/Engine/Plugins/Experimental/GeometryCollectionPlugin/Source/GeometryCollectionNodes/Private/Dataflow/GeometryCollectionMakeNodes.cpp
@@ -40,7 +40,7 @@
 
 #include UE_INLINE_GENERATED_CPP_BY_NAME(GeometryCollectionMakeNodes)
 
-namespace Dataflow
+namespace UE::Dataflow
 {
 	void GeometryCollectionMakeNodes()
 	{
@@ -63,7 +63,7 @@ namespace Dataflow
 	}
 }
 
-void FMakeLiteralStringDataflowNode::Evaluate(Dataflow::FContext& Context, const FDataflowOutput* Out) const
+void FMakeLiteralStringDataflowNode::Evaluate(UE::Dataflow::FContext& Context, const FDataflowOutput* Out) const
 {
 	if (Out->IsA<FString>(&String))
 	{
@@ -71,7 +71,7 @@ void FMakeLiteralStringDataflowNode::Evaluate(Dataflow::FContext& Context, const
 	}
 }
 
-void FMakePointsDataflowNode::Evaluate(Dataflow::FContext& Context, const FDataflowOutput* Out) const
+void FMakePointsDataflowNode::Evaluate(UE::Dataflow::FContext& Context, const FDataflowOutput* Out) const
 {
 	if (Out->IsA<TArray<FVector>>(&Points))
 	{
@@ -79,7 +79,7 @@ void FMakePointsDataflowNode::Evaluate(Dataflow::FContext& Context, const FDataf
 	}
 }
 
-void FMakeBoxDataflowNode::Evaluate(Dataflow::FContext& Context, const FDataflowOutput* Out) const
+void FMakeBoxDataflowNode::Evaluate(UE::Dataflow::FContext& Context, const FDataflowOutput* Out) const
 {
 	if (Out->IsA<FBox>(&Box))
 	{
@@ -101,7 +101,7 @@ void FMakeBoxDataflowNode::Evaluate(Dataflow::FContext& Context, const FDataflow
 }
 
 
-void FMakeSphereDataflowNode::Evaluate(Dataflow::FContext& Context, const FDataflowOutput* Out) const
+void FMakeSphereDataflowNode::Evaluate(UE::Dataflow::FContext& Context, const FDataflowOutput* Out) const
 {
 	if (Out->IsA<FSphere>(&Sphere))
 	{
@@ -113,7 +113,7 @@ void FMakeSphereDataflowNode::Evaluate(Dataflow::FContext& Context, const FDataf
 }
 
 
-void FMakeLiteralFloatDataflowNode::Evaluate(Dataflow::FContext& Context, const FDataflowOutput* Out) const
+void FMakeLiteralFloatDataflowNode::Evaluate(UE::Dataflow::FContext& Context, const FDataflowOutput* Out) const
 {
 	if (Out->IsA<float>(&Float))
 	{
@@ -121,7 +121,7 @@ void FMakeLiteralFloatDataflowNode::Evaluate(Dataflow::FContext& Context, const
 	}
 }
 
-void FMakeLiteralIntDataflowNode::Evaluate(Dataflow::FContext& Context, const FDataflowOutput* Out) const
+void FMakeLiteralIntDataflowNode::Evaluate(UE::Dataflow::FContext& Context, const FDataflowOutput* Out) const
 {
 	if (Out->IsA<int32>(&Int))
 	{
@@ -129,7 +129,7 @@ void FMakeLiteralIntDataflowNode::Evaluate(Dataflow::FContext& Context, const FD
 	}
 }
 
-void FMakeLiteralBoolDataflowNode::Evaluate(Dataflow::FContext& Context, const FDataflowOutput* Out) const
+void FMakeLiteralBoolDataflowNode::Evaluate(UE::Dataflow::FContext& Context, const FDataflowOutput* Out) const
 {
 	if (Out->IsA<bool>(&Bool))
 	{
@@ -137,7 +137,7 @@ void FMakeLiteralBoolDataflowNode::Evaluate(Dataflow::FContext& Context, const F
 	}
 }
 
-void FMakeLiteralVectorDataflowNode::Evaluate(Dataflow::FContext& Context, const FDataflowOutput* Out) const
+void FMakeLiteralVectorDataflowNode::Evaluate(UE::Dataflow::FContext& Context, const FDataflowOutput* Out) const
 {
 	if (Out->IsA<FVector>(&Vector))
 	{
@@ -146,7 +146,7 @@ void FMakeLiteralVectorDataflowNode::Evaluate(Dataflow::FContext& Context, const
 	}
 }
 
-void FMakeTransformDataflowNode::Evaluate(Dataflow::FContext& Context, const FDataflowOutput* Out) const
+void FMakeTransformDataflowNode::Evaluate(UE::Dataflow::FContext& Context, const FDataflowOutput* Out) const
 {
 	if (Out->IsA<FTransform>(&OutTransform))
 	{
@@ -158,7 +158,7 @@ void FMakeTransformDataflowNode::Evaluate(Dataflow::FContext& Context, const FDa
 	}
 }
 
-void FMakeQuaternionDataflowNode::Evaluate(Dataflow::FContext& Context, const FDataflowOutput* Out) const
+void FMakeQuaternionDataflowNode::Evaluate(UE::Dataflow::FContext& Context, const FDataflowOutput* Out) const
 {
 	if (Out->IsA<FQuat>(&Quaternion))
 	{
@@ -167,7 +167,7 @@ void FMakeQuaternionDataflowNode::Evaluate(Dataflow::FContext& Context, const FD
 	}
 }
 
-void FMakeFloatArrayDataflowNode::Evaluate(Dataflow::FContext& Context, const FDataflowOutput* Out) const
+void FMakeFloatArrayDataflowNode::Evaluate(UE::Dataflow::FContext& Context, const FDataflowOutput* Out) const
 {
 	if (Out->IsA(&FloatArray))
 	{
diff --git a/Engine/Plugins/Experimental/GeometryCollectionPlugin/Source/GeometryCollectionNodes/Private/Dataflow/GeometryCollectionMaterialNodes.cpp b/Engine/Plugins/Experimental/GeometryCollectionPlugin/Source/GeometryCollectionNodes/Private/Dataflow/GeometryCollectionMaterialNodes.cpp
--- a/Engine/Plugins/Experimental/GeometryCollectionPlugin/Source/GeometryCollectionNodes/Private/Dataflow/GeometryCollectionMaterialNodes.cpp
+++ b/Engine/Plugins/Experimental/GeometryCollectionPlugin/Source/GeometryCollectionNodes/Private/Dataflow/GeometryCollectionMaterialNodes.cpp
@@ -40,7 +40,7 @@
 
 //#include UE_INLINE_GENERATED_CPP_BY_NAME(GeometryCollectionMaterialNodes)
 
-namespace Dataflow
+namespace UE::Dataflow
 {
 	void GeometryCollectionMaterialNodes()
 	{
@@ -60,7 +60,7 @@ namespace Dataflow
 }
 
 
-void FAddMaterialToCollectionDataflowNode::Evaluate(Dataflow::FContext& Context, const FDataflowOutput* Out) const
+void FAddMaterialToCollectionDataflowNode::Evaluate(UE::Dataflow::FContext& Context, const FDataflowOutput* Out) const
 {
 	if (Out->IsA<FManagedArrayCollection>(&Collection) || Out->IsA<TArray<TObjectPtr<UMaterial>>>(&Materials))
 	{
@@ -116,7 +116,7 @@ void FAddMaterialToCollectionDataflowNode::Evaluate(Dataflow::FContext& Context,
 }
 
 
-void FReAssignMaterialInCollectionDataflowNode::Evaluate(Dataflow::FContext& Context, const FDataflowOutput* Out) const
+void FReAssignMaterialInCollectionDataflowNode::Evaluate(UE::Dataflow::FContext& Context, const FDataflowOutput* Out) const
 {
 	if (Out->IsA<FManagedArrayCollection>(&Collection))
 	{
@@ -176,7 +176,7 @@ void FReAssignMaterialInCollectionDataflowNode::Evaluate(Dataflow::FContext& Con
 }
 
 
-void FMaterialsInfoDataflowNode::Evaluate(Dataflow::FContext& Context, const FDataflowOutput* Out) const
+void FMaterialsInfoDataflowNode::Evaluate(UE::Dataflow::FContext& Context, const FDataflowOutput* Out) const
 {
 	if (Out->IsA<FString>(&String))
 	{
@@ -202,7 +202,7 @@ void FMaterialsInfoDataflowNode::Evaluate(Dataflow::FContext& Context, const FDa
 }
 
 
-void FGetMaterialFromMaterialsArrayDataflowNode::Evaluate(Dataflow::FContext& Context, const FDataflowOutput* Out) const
+void FGetMaterialFromMaterialsArrayDataflowNode::Evaluate(UE::Dataflow::FContext& Context, const FDataflowOutput* Out) const
 {
 	if (Out->IsA<TObjectPtr<UMaterial>>(&Material))
 	{
@@ -224,7 +224,7 @@ void FGetMaterialFromMaterialsArrayDataflowNode::Evaluate(Dataflow::FContext& Co
 }
 
 
-void FSetMaterialInMaterialsArrayDataflowNode::Evaluate(Dataflow::FContext& Context, const FDataflowOutput* Out) const
+void FSetMaterialInMaterialsArrayDataflowNode::Evaluate(UE::Dataflow::FContext& Context, const FDataflowOutput* Out) const
 {
 	if (Out->IsA<TArray<TObjectPtr<UMaterial>>>(&Materials))
 	{
@@ -249,7 +249,7 @@ void FSetMaterialInMaterialsArrayDataflowNode::Evaluate(Dataflow::FContext& Cont
 }
 
 
-void FMakeMaterialDataflowNode::Evaluate(Dataflow::FContext& Context, const FDataflowOutput* Out) const
+void FMakeMaterialDataflowNode::Evaluate(UE::Dataflow::FContext& Context, const FDataflowOutput* Out) const
 {
 	if (Out->IsA<TObjectPtr<UMaterial>>(&Material))
 	{
@@ -258,7 +258,7 @@ void FMakeMaterialDataflowNode::Evaluate(Dataflow::FContext& Context, const FDat
 }
 
 
-void FMakeMaterialsArrayDataflowNode::Evaluate(Dataflow::FContext& Context, const FDataflowOutput* Out) const
+void FMakeMaterialsArrayDataflowNode::Evaluate(UE::Dataflow::FContext& Context, const FDataflowOutput* Out) const
 {
 	if (Out->IsA<TArray<TObjectPtr<UMaterial>>>(&Materials))
 	{
diff --git a/Engine/Plugins/Experimental/GeometryCollectionPlugin/Source/GeometryCollectionNodes/Private/Dataflow/GeometryCollectionMathNodes.cpp b/Engine/Plugins/Experimental/GeometryCollectionPlugin/Source/GeometryCollectionNodes/Private/Dataflow/GeometryCollectionMathNodes.cpp
--- a/Engine/Plugins/Experimental/GeometryCollectionPlugin/Source/GeometryCollectionNodes/Private/Dataflow/GeometryCollectionMathNodes.cpp
+++ b/Engine/Plugins/Experimental/GeometryCollectionPlugin/Source/GeometryCollectionNodes/Private/Dataflow/GeometryCollectionMathNodes.cpp
@@ -7,7 +7,7 @@
 
 #include UE_INLINE_GENERATED_CPP_BY_NAME(GeometryCollectionMathNodes)
 
-namespace Dataflow
+namespace UE::Dataflow
 {
 
 	void GeometryCollectionMathNodes()
@@ -77,7 +77,7 @@ namespace Dataflow
 	}
 }
 
-void FAddDataflowNode::Evaluate(Dataflow::FContext& Context, const FDataflowOutput* Out) const
+void FAddDataflowNode::Evaluate(UE::Dataflow::FContext& Context, const FDataflowOutput* Out) const
 {
 	if (Out->IsA<float>(&ReturnValue))
 	{
@@ -89,7 +89,7 @@ void FAddDataflowNode::Evaluate(Dataflow::FContext& Context, const FDataflowOutp
 	}
 }
 
-void FSubtractDataflowNode::Evaluate(Dataflow::FContext& Context, const FDataflowOutput* Out) const
+void FSubtractDataflowNode::Evaluate(UE::Dataflow::FContext& Context, const FDataflowOutput* Out) const
 {
 	if (Out->IsA<float>(&ReturnValue))
 	{
@@ -101,7 +101,7 @@ void FSubtractDataflowNode::Evaluate(Dataflow::FContext& Context, const FDataflo
 	}
 }
 
-void FMultiplyDataflowNode::Evaluate(Dataflow::FContext& Context, const FDataflowOutput* Out) const
+void FMultiplyDataflowNode::Evaluate(UE::Dataflow::FContext& Context, const FDataflowOutput* Out) const
 {
 	if (Out->IsA<float>(&ReturnValue))
 	{
@@ -113,7 +113,7 @@ void FMultiplyDataflowNode::Evaluate(Dataflow::FContext& Context, const FDataflo
 	}
 }
 
-void FSafeDivideDataflowNode::Evaluate(Dataflow::FContext& Context, const FDataflowOutput* Out) const
+void FSafeDivideDataflowNode::Evaluate(UE::Dataflow::FContext& Context, const FDataflowOutput* Out) const
 {
 	if (Out->IsA<float>(&ReturnValue))
 	{
@@ -129,7 +129,7 @@ void FSafeDivideDataflowNode::Evaluate(Dataflow::FContext& Context, const FDataf
 	}
 }
 
-void FDivisionDataflowNode::Evaluate(Dataflow::FContext& Context, const FDataflowOutput* Out) const
+void FDivisionDataflowNode::Evaluate(UE::Dataflow::FContext& Context, const FDataflowOutput* Out) const
 {
 	if (Out->IsA<float>(&Remainder) ||
 		Out->IsA<int32>(&ReturnValue))
@@ -150,7 +150,7 @@ void FDivisionDataflowNode::Evaluate(Dataflow::FContext& Context, const FDataflo
 	}
 }
 
-void FSafeReciprocalDataflowNode::Evaluate(Dataflow::FContext& Context, const FDataflowOutput* Out) const
+void FSafeReciprocalDataflowNode::Evaluate(UE::Dataflow::FContext& Context, const FDataflowOutput* Out) const
 {
 	if (Out->IsA<float>(&ReturnValue))
 	{
@@ -165,7 +165,7 @@ void FSafeReciprocalDataflowNode::Evaluate(Dataflow::FContext& Context, const FD
 	}
 }
 
-void FSquareDataflowNode::Evaluate(Dataflow::FContext& Context, const FDataflowOutput* Out) const
+void FSquareDataflowNode::Evaluate(UE::Dataflow::FContext& Context, const FDataflowOutput* Out) const
 {
 	if (Out->IsA<float>(&ReturnValue))
 	{
@@ -176,7 +176,7 @@ void FSquareDataflowNode::Evaluate(Dataflow::FContext& Context, const FDataflowO
 	}
 }
 
-void FSquareRootDataflowNode::Evaluate(Dataflow::FContext& Context, const FDataflowOutput* Out) const
+void FSquareRootDataflowNode::Evaluate(UE::Dataflow::FContext& Context, const FDataflowOutput* Out) const
 {
 	if (Out->IsA<float>(&ReturnValue))
 	{
@@ -191,7 +191,7 @@ void FSquareRootDataflowNode::Evaluate(Dataflow::FContext& Context, const FDataf
 	}
 }
 
-void FInverseSqrtDataflowNode::Evaluate(Dataflow::FContext& Context, const FDataflowOutput* Out) const
+void FInverseSqrtDataflowNode::Evaluate(UE::Dataflow::FContext& Context, const FDataflowOutput* Out) const
 {
 	if (Out->IsA<float>(&ReturnValue))
 	{
@@ -206,7 +206,7 @@ void FInverseSqrtDataflowNode::Evaluate(Dataflow::FContext& Context, const FData
 	}
 }
 
-void FCubeDataflowNode::Evaluate(Dataflow::FContext& Context, const FDataflowOutput* Out) const
+void FCubeDataflowNode::Evaluate(UE::Dataflow::FContext& Context, const FDataflowOutput* Out) const
 {
 	if (Out->IsA<float>(&ReturnValue))
 	{
@@ -217,7 +217,7 @@ void FCubeDataflowNode::Evaluate(Dataflow::FContext& Context, const FDataflowOut
 	}
 }
 
-void FNegateDataflowNode::Evaluate(Dataflow::FContext& Context, const FDataflowOutput* Out) const
+void FNegateDataflowNode::Evaluate(UE::Dataflow::FContext& Context, const FDataflowOutput* Out) const
 {
 	if (Out->IsA<float>(&ReturnValue))
 	{
@@ -228,7 +228,7 @@ void FNegateDataflowNode::Evaluate(Dataflow::FContext& Context, const FDataflowO
 	}
 }
 
-void FAbsDataflowNode::Evaluate(Dataflow::FContext& Context, const FDataflowOutput* Out) const
+void FAbsDataflowNode::Evaluate(UE::Dataflow::FContext& Context, const FDataflowOutput* Out) const
 {
 	if (Out->IsA<float>(&ReturnValue))
 	{
@@ -239,7 +239,7 @@ void FAbsDataflowNode::Evaluate(Dataflow::FContext& Context, const FDataflowOutp
 	}
 }
 
-void FFloorDataflowNode::Evaluate(Dataflow::FContext& Context, const FDataflowOutput* Out) const
+void FFloorDataflowNode::Evaluate(UE::Dataflow::FContext& Context, const FDataflowOutput* Out) const
 {
 	if (Out->IsA<float>(&ReturnValue))
 	{
@@ -250,7 +250,7 @@ void FFloorDataflowNode::Evaluate(Dataflow::FContext& Context, const FDataflowOu
 	}
 }
 
-void FCeilDataflowNode::Evaluate(Dataflow::FContext& Context, const FDataflowOutput* Out) const
+void FCeilDataflowNode::Evaluate(UE::Dataflow::FContext& Context, const FDataflowOutput* Out) const
 {
 	if (Out->IsA<float>(&ReturnValue))
 	{
@@ -261,7 +261,7 @@ void FCeilDataflowNode::Evaluate(Dataflow::FContext& Context, const FDataflowOut
 	}
 }
 
-void FRoundDataflowNode::Evaluate(Dataflow::FContext& Context, const FDataflowOutput* Out) const
+void FRoundDataflowNode::Evaluate(UE::Dataflow::FContext& Context, const FDataflowOutput* Out) const
 {
 	if (Out->IsA<float>(&ReturnValue))
 	{
@@ -272,7 +272,7 @@ void FRoundDataflowNode::Evaluate(Dataflow::FContext& Context, const FDataflowOu
 	}
 }
 
-void FTruncDataflowNode::Evaluate(Dataflow::FContext& Context, const FDataflowOutput* Out) const
+void FTruncDataflowNode::Evaluate(UE::Dataflow::FContext& Context, const FDataflowOutput* Out) const
 {
 	if (Out->IsA<float>(&ReturnValue))
 	{
@@ -283,7 +283,7 @@ void FTruncDataflowNode::Evaluate(Dataflow::FContext& Context, const FDataflowOu
 	}
 }
 
-void FFracDataflowNode::Evaluate(Dataflow::FContext& Context, const FDataflowOutput* Out) const
+void FFracDataflowNode::Evaluate(UE::Dataflow::FContext& Context, const FDataflowOutput* Out) const
 {
 	if (Out->IsA<float>(&ReturnValue))
 	{
@@ -294,7 +294,7 @@ void FFracDataflowNode::Evaluate(Dataflow::FContext& Context, const FDataflowOut
 	}
 }
 
-void FMinDataflowNode::Evaluate(Dataflow::FContext& Context, const FDataflowOutput* Out) const
+void FMinDataflowNode::Evaluate(UE::Dataflow::FContext& Context, const FDataflowOutput* Out) const
 {
 	if (Out->IsA<float>(&ReturnValue))
 	{
@@ -306,7 +306,7 @@ void FMinDataflowNode::Evaluate(Dataflow::FContext& Context, const FDataflowOutp
 	}
 }
 
-void FMaxDataflowNode::Evaluate(Dataflow::FContext& Context, const FDataflowOutput* Out) const
+void FMaxDataflowNode::Evaluate(UE::Dataflow::FContext& Context, const FDataflowOutput* Out) const
 {
 	if (Out->IsA<float>(&ReturnValue))
 	{
@@ -318,7 +318,7 @@ void FMaxDataflowNode::Evaluate(Dataflow::FContext& Context, const FDataflowOutp
 	}
 }
 
-void FMin3DataflowNode::Evaluate(Dataflow::FContext& Context, const FDataflowOutput* Out) const
+void FMin3DataflowNode::Evaluate(UE::Dataflow::FContext& Context, const FDataflowOutput* Out) const
 {
 	if (Out->IsA<float>(&ReturnValue))
 	{
@@ -331,7 +331,7 @@ void FMin3DataflowNode::Evaluate(Dataflow::FContext& Context, const FDataflowOut
 	}
 }
 
-void FMax3DataflowNode::Evaluate(Dataflow::FContext& Context, const FDataflowOutput* Out) const
+void FMax3DataflowNode::Evaluate(UE::Dataflow::FContext& Context, const FDataflowOutput* Out) const
 {
 	if (Out->IsA<float>(&ReturnValue))
 	{
@@ -344,7 +344,7 @@ void FMax3DataflowNode::Evaluate(Dataflow::FContext& Context, const FDataflowOut
 	}
 }
 
-void FSignDataflowNode::Evaluate(Dataflow::FContext& Context, const FDataflowOutput* Out) const
+void FSignDataflowNode::Evaluate(UE::Dataflow::FContext& Context, const FDataflowOutput* Out) const
 {
 	if (Out->IsA<float>(&ReturnValue))
 	{
@@ -355,7 +355,7 @@ void FSignDataflowNode::Evaluate(Dataflow::FContext& Context, const FDataflowOut
 	}
 }
 
-void FClampDataflowNode::Evaluate(Dataflow::FContext& Context, const FDataflowOutput* Out) const
+void FClampDataflowNode::Evaluate(UE::Dataflow::FContext& Context, const FDataflowOutput* Out) const
 {
 	if (Out->IsA<float>(&ReturnValue))
 	{
@@ -368,7 +368,7 @@ void FClampDataflowNode::Evaluate(Dataflow::FContext& Context, const FDataflowOu
 	}
 }
 
-void FFitDataflowNode::Evaluate(Dataflow::FContext& Context, const FDataflowOutput* Out) const
+void FFitDataflowNode::Evaluate(UE::Dataflow::FContext& Context, const FDataflowOutput* Out) const
 {
 	if (Out->IsA<float>(&ReturnValue))
 	{
@@ -389,7 +389,7 @@ void FFitDataflowNode::Evaluate(Dataflow::FContext& Context, const FDataflowOutp
 	}
 }
 
-void FEFitDataflowNode::Evaluate(Dataflow::FContext& Context, const FDataflowOutput* Out) const
+void FEFitDataflowNode::Evaluate(UE::Dataflow::FContext& Context, const FDataflowOutput* Out) const
 {
 	if (Out->IsA<float>(&ReturnValue))
 	{
@@ -409,7 +409,7 @@ void FEFitDataflowNode::Evaluate(Dataflow::FContext& Context, const FDataflowOut
 	}
 }
 
-void FPowDataflowNode::Evaluate(Dataflow::FContext& Context, const FDataflowOutput* Out) const
+void FPowDataflowNode::Evaluate(UE::Dataflow::FContext& Context, const FDataflowOutput* Out) const
 {
 	if (Out->IsA<float>(&ReturnValue))
 	{
@@ -421,7 +421,7 @@ void FPowDataflowNode::Evaluate(Dataflow::FContext& Context, const FDataflowOutp
 	}
 }
 
-void FLogDataflowNode::Evaluate(Dataflow::FContext& Context, const FDataflowOutput* Out) const
+void FLogDataflowNode::Evaluate(UE::Dataflow::FContext& Context, const FDataflowOutput* Out) const
 {
 	if (Out->IsA<float>(&ReturnValue))
 	{
@@ -437,7 +437,7 @@ void FLogDataflowNode::Evaluate(Dataflow::FContext& Context, const FDataflowOutp
 	}
 }
 
-void FLogeDataflowNode::Evaluate(Dataflow::FContext& Context, const FDataflowOutput* Out) const
+void FLogeDataflowNode::Evaluate(UE::Dataflow::FContext& Context, const FDataflowOutput* Out) const
 {
 	if (Out->IsA<float>(&ReturnValue))
 	{
@@ -448,7 +448,7 @@ void FLogeDataflowNode::Evaluate(Dataflow::FContext& Context, const FDataflowOut
 	}
 }
 
-void FLerpDataflowNode::Evaluate(Dataflow::FContext& Context, const FDataflowOutput* Out) const
+void FLerpDataflowNode::Evaluate(UE::Dataflow::FContext& Context, const FDataflowOutput* Out) const
 {
 	if (Out->IsA<float>(&ReturnValue))
 	{
@@ -465,7 +465,7 @@ void FLerpDataflowNode::Evaluate(Dataflow::FContext& Context, const FDataflowOut
 	}
 }
 
-void FWrapDataflowNode::Evaluate(Dataflow::FContext& Context, const FDataflowOutput* Out) const
+void FWrapDataflowNode::Evaluate(UE::Dataflow::FContext& Context, const FDataflowOutput* Out) const
 {
 	if (Out->IsA<float>(&ReturnValue))
 	{
@@ -482,7 +482,7 @@ void FWrapDataflowNode::Evaluate(Dataflow::FContext& Context, const FDataflowOut
 	}
 }
 
-void FExpDataflowNode::Evaluate(Dataflow::FContext& Context, const FDataflowOutput* Out) const
+void FExpDataflowNode::Evaluate(UE::Dataflow::FContext& Context, const FDataflowOutput* Out) const
 {
 	if (Out->IsA<float>(&ReturnValue))
 	{
@@ -493,7 +493,7 @@ void FExpDataflowNode::Evaluate(Dataflow::FContext& Context, const FDataflowOutp
 	}
 }
 
-void FSinDataflowNode::Evaluate(Dataflow::FContext& Context, const FDataflowOutput* Out) const
+void FSinDataflowNode::Evaluate(UE::Dataflow::FContext& Context, const FDataflowOutput* Out) const
 {
 	if (Out->IsA<float>(&ReturnValue))
 	{
@@ -504,7 +504,7 @@ void FSinDataflowNode::Evaluate(Dataflow::FContext& Context, const FDataflowOutp
 	}
 }
 
-void FArcSinDataflowNode::Evaluate(Dataflow::FContext& Context, const FDataflowOutput* Out) const
+void FArcSinDataflowNode::Evaluate(UE::Dataflow::FContext& Context, const FDataflowOutput* Out) const
 {
 	if (Out->IsA<float>(&ReturnValue))
 	{
@@ -519,7 +519,7 @@ void FArcSinDataflowNode::Evaluate(Dataflow::FContext& Context, const FDataflowO
 	}
 }
 
-void FCosDataflowNode::Evaluate(Dataflow::FContext& Context, const FDataflowOutput* Out) const
+void FCosDataflowNode::Evaluate(UE::Dataflow::FContext& Context, const FDataflowOutput* Out) const
 {
 	if (Out->IsA<float>(&ReturnValue))
 	{
@@ -530,7 +530,7 @@ void FCosDataflowNode::Evaluate(Dataflow::FContext& Context, const FDataflowOutp
 	}
 }
 
-void FArcCosDataflowNode::Evaluate(Dataflow::FContext& Context, const FDataflowOutput* Out) const
+void FArcCosDataflowNode::Evaluate(UE::Dataflow::FContext& Context, const FDataflowOutput* Out) const
 {
 	if (Out->IsA<float>(&ReturnValue))
 	{
@@ -545,7 +545,7 @@ void FArcCosDataflowNode::Evaluate(Dataflow::FContext& Context, const FDataflowO
 	}
 }
 
-void FTanDataflowNode::Evaluate(Dataflow::FContext& Context, const FDataflowOutput* Out) const
+void FTanDataflowNode::Evaluate(UE::Dataflow::FContext& Context, const FDataflowOutput* Out) const
 {
 	if (Out->IsA<float>(&ReturnValue))
 	{
@@ -556,7 +556,7 @@ void FTanDataflowNode::Evaluate(Dataflow::FContext& Context, const FDataflowOutp
 	}
 }
 
-void FArcTanDataflowNode::Evaluate(Dataflow::FContext& Context, const FDataflowOutput* Out) const
+void FArcTanDataflowNode::Evaluate(UE::Dataflow::FContext& Context, const FDataflowOutput* Out) const
 {
 	if (Out->IsA<float>(&ReturnValue))
 	{
@@ -567,7 +567,7 @@ void FArcTanDataflowNode::Evaluate(Dataflow::FContext& Context, const FDataflowO
 	}
 }
 
-void FArcTan2DataflowNode::Evaluate(Dataflow::FContext& Context, const FDataflowOutput* Out) const
+void FArcTan2DataflowNode::Evaluate(UE::Dataflow::FContext& Context, const FDataflowOutput* Out) const
 {
 	if (Out->IsA<float>(&ReturnValue))
 	{
@@ -579,7 +579,7 @@ void FArcTan2DataflowNode::Evaluate(Dataflow::FContext& Context, const FDataflow
 	}
 }
 
-void FNormalizeToRangeDataflowNode::Evaluate(Dataflow::FContext& Context, const FDataflowOutput* Out) const
+void FNormalizeToRangeDataflowNode::Evaluate(UE::Dataflow::FContext& Context, const FDataflowOutput* Out) const
 {
 	if (Out->IsA<float>(&ReturnValue))
 	{
@@ -611,7 +611,7 @@ void FNormalizeToRangeDataflowNode::Evaluate(Dataflow::FContext& Context, const
 }
 
 
-void FScaleVectorDataflowNode::Evaluate(Dataflow::FContext& Context, const FDataflowOutput* Out) const
+void FScaleVectorDataflowNode::Evaluate(UE::Dataflow::FContext& Context, const FDataflowOutput* Out) const
 {
 	if (Out->IsA(&ScaledVector))
 	{
@@ -623,7 +623,7 @@ void FScaleVectorDataflowNode::Evaluate(Dataflow::FContext& Context, const FData
 	}
 }
 
-void FDotProductDataflowNode::Evaluate(Dataflow::FContext& Context, const FDataflowOutput* Out) const
+void FDotProductDataflowNode::Evaluate(UE::Dataflow::FContext& Context, const FDataflowOutput* Out) const
 {
 	if (Out->IsA<float>(&ReturnValue))
 	{
@@ -635,7 +635,7 @@ void FDotProductDataflowNode::Evaluate(Dataflow::FContext& Context, const FDataf
 	}
 }
 
-void FCrossProductDataflowNode::Evaluate(Dataflow::FContext& Context, const FDataflowOutput* Out) const
+void FCrossProductDataflowNode::Evaluate(UE::Dataflow::FContext& Context, const FDataflowOutput* Out) const
 {
 	if (Out->IsA<FVector>(&ReturnValue))
 	{
@@ -647,7 +647,7 @@ void FCrossProductDataflowNode::Evaluate(Dataflow::FContext& Context, const FDat
 	}
 }
 
-void FNormalizeDataflowNode::Evaluate(Dataflow::FContext& Context, const FDataflowOutput* Out) const
+void FNormalizeDataflowNode::Evaluate(UE::Dataflow::FContext& Context, const FDataflowOutput* Out) const
 {
 	if (Out->IsA<FVector>(&ReturnValue))
 	{
@@ -659,7 +659,7 @@ void FNormalizeDataflowNode::Evaluate(Dataflow::FContext& Context, const FDatafl
 	}
 }
 
-void FLengthDataflowNode::Evaluate(Dataflow::FContext& Context, const FDataflowOutput* Out) const
+void FLengthDataflowNode::Evaluate(UE::Dataflow::FContext& Context, const FDataflowOutput* Out) const
 {
 	if (Out->IsA<float>(&ReturnValue))
 	{
@@ -670,7 +670,7 @@ void FLengthDataflowNode::Evaluate(Dataflow::FContext& Context, const FDataflowO
 	}
 }
 
-void FDistanceDataflowNode::Evaluate(Dataflow::FContext& Context, const FDataflowOutput* Out) const
+void FDistanceDataflowNode::Evaluate(UE::Dataflow::FContext& Context, const FDataflowOutput* Out) const
 {
 	if (Out->IsA<float>(&ReturnValue))
 	{
@@ -682,7 +682,7 @@ void FDistanceDataflowNode::Evaluate(Dataflow::FContext& Context, const FDataflo
 	}
 }
 
-void FIsNearlyZeroDataflowNode::Evaluate(Dataflow::FContext& Context, const FDataflowOutput* Out) const
+void FIsNearlyZeroDataflowNode::Evaluate(UE::Dataflow::FContext& Context, const FDataflowOutput* Out) const
 {
 	if (Out->IsA<bool>(&ReturnValue))
 	{
@@ -693,7 +693,7 @@ void FIsNearlyZeroDataflowNode::Evaluate(Dataflow::FContext& Context, const FDat
 	}
 }
 
-void FRandomFloatDataflowNode::Evaluate(Dataflow::FContext& Context, const FDataflowOutput* Out) const
+void FRandomFloatDataflowNode::Evaluate(UE::Dataflow::FContext& Context, const FDataflowOutput* Out) const
 {
 	if (Out->IsA<float>(&ReturnValue))
 	{
@@ -711,7 +711,7 @@ void FRandomFloatDataflowNode::Evaluate(Dataflow::FContext& Context, const FData
 	}
 }
 
-void FRandomFloatInRangeDataflowNode::Evaluate(Dataflow::FContext& Context, const FDataflowOutput* Out) const
+void FRandomFloatInRangeDataflowNode::Evaluate(UE::Dataflow::FContext& Context, const FDataflowOutput* Out) const
 {
 	if (Out->IsA<float>(&ReturnValue))
 	{
@@ -732,7 +732,7 @@ void FRandomFloatInRangeDataflowNode::Evaluate(Dataflow::FContext& Context, cons
 	}
 }
 
-void FRandomUnitVectorDataflowNode::Evaluate(Dataflow::FContext& Context, const FDataflowOutput* Out) const
+void FRandomUnitVectorDataflowNode::Evaluate(UE::Dataflow::FContext& Context, const FDataflowOutput* Out) const
 {
 	if (Out->IsA<FVector>(&ReturnValue))
 	{
@@ -750,7 +750,7 @@ void FRandomUnitVectorDataflowNode::Evaluate(Dataflow::FContext& Context, const
 	}
 }
 
-void FRandomUnitVectorInConeDataflowNode::Evaluate(Dataflow::FContext& Context, const FDataflowOutput* Out) const
+void FRandomUnitVectorInConeDataflowNode::Evaluate(UE::Dataflow::FContext& Context, const FDataflowOutput* Out) const
 {
 	if (Out->IsA<FVector>(&ReturnValue))
 	{
@@ -771,7 +771,7 @@ void FRandomUnitVectorInConeDataflowNode::Evaluate(Dataflow::FContext& Context,
 	}
 }
 
-void FRadiansToDegreesDataflowNode::Evaluate(Dataflow::FContext& Context, const FDataflowOutput* Out) const
+void FRadiansToDegreesDataflowNode::Evaluate(UE::Dataflow::FContext& Context, const FDataflowOutput* Out) const
 {
 	if (Out->IsA<float>(&Degrees))
 	{
@@ -779,7 +779,7 @@ void FRadiansToDegreesDataflowNode::Evaluate(Dataflow::FContext& Context, const
 	}
 }
 
-void FDegreesToRadiansDataflowNode::Evaluate(Dataflow::FContext& Context, const FDataflowOutput* Out) const
+void FDegreesToRadiansDataflowNode::Evaluate(UE::Dataflow::FContext& Context, const FDataflowOutput* Out) const
 {
 	if (Out->IsA<float>(&Radians))
 	{
@@ -787,7 +787,7 @@ void FDegreesToRadiansDataflowNode::Evaluate(Dataflow::FContext& Context, const
 	}
 }
 
-void FMathConstantsDataflowNode::Evaluate(Dataflow::FContext& Context, const FDataflowOutput* Out) const
+void FMathConstantsDataflowNode::Evaluate(UE::Dataflow::FContext& Context, const FDataflowOutput* Out) const
 {
 	if (Out->IsA<float>(&ReturnValue))
 	{
@@ -850,7 +850,7 @@ void FMathConstantsDataflowNode::Evaluate(Dataflow::FContext& Context, const FDa
 	}
 }
 
-void FOneMinusDataflowNode::Evaluate(Dataflow::FContext& Context, const FDataflowOutput* Out) const
+void FOneMinusDataflowNode::Evaluate(UE::Dataflow::FContext& Context, const FDataflowOutput* Out) const
 {
 	if (Out->IsA<float>(&ReturnValue))
 	{
@@ -861,7 +861,7 @@ void FOneMinusDataflowNode::Evaluate(Dataflow::FContext& Context, const FDataflo
 	}
 }
 
-FFloatMathExpressionDataflowNode::FFloatMathExpressionDataflowNode(const Dataflow::FNodeParameters& InParam, FGuid InGuid)
+FFloatMathExpressionDataflowNode::FFloatMathExpressionDataflowNode(const UE::Dataflow::FNodeParameters& InParam, FGuid InGuid)
 	: FDataflowNode(InParam, InGuid)
 {
 	RegisterInputConnection(&A);
@@ -871,7 +871,7 @@ FFloatMathExpressionDataflowNode::FFloatMathExpressionDataflowNode(const Dataflo
 	RegisterOutputConnection(&ReturnValue);
 }
 
-void FFloatMathExpressionDataflowNode::Evaluate(Dataflow::FContext& Context, const FDataflowOutput* Out) const
+void FFloatMathExpressionDataflowNode::Evaluate(UE::Dataflow::FContext& Context, const FDataflowOutput* Out) const
 {
 	check(Out->IsA(&ReturnValue));
 
@@ -925,7 +925,7 @@ void FFloatMathExpressionDataflowNode::Evaluate(Dataflow::FContext& Context, con
 
 
 
-FMathExpressionDataflowNode::FMathExpressionDataflowNode(const Dataflow::FNodeParameters& InParam, FGuid InGuid)
+FMathExpressionDataflowNode::FMathExpressionDataflowNode(const UE::Dataflow::FNodeParameters& InParam, FGuid InGuid)
 	: FDataflowNode(InParam, InGuid)
 {
 	RegisterInputConnection(&A);
@@ -935,7 +935,7 @@ FMathExpressionDataflowNode::FMathExpressionDataflowNode(const Dataflow::FNodePa
 	RegisterOutputConnection(&ReturnValue);
 }
 
-void FMathExpressionDataflowNode::Evaluate(Dataflow::FContext& Context, const FDataflowOutput* Out) const
+void FMathExpressionDataflowNode::Evaluate(UE::Dataflow::FContext& Context, const FDataflowOutput* Out) const
 {
 	check(Out->IsA(&ReturnValue));
 
diff --git a/Engine/Plugins/Experimental/GeometryCollectionPlugin/Source/GeometryCollectionNodes/Private/Dataflow/GeometryCollectionMeshNodes.cpp b/Engine/Plugins/Experimental/GeometryCollectionPlugin/Source/GeometryCollectionNodes/Private/Dataflow/GeometryCollectionMeshNodes.cpp
--- a/Engine/Plugins/Experimental/GeometryCollectionPlugin/Source/GeometryCollectionNodes/Private/Dataflow/GeometryCollectionMeshNodes.cpp
+++ b/Engine/Plugins/Experimental/GeometryCollectionPlugin/Source/GeometryCollectionNodes/Private/Dataflow/GeometryCollectionMeshNodes.cpp
@@ -33,7 +33,7 @@
 
 #include UE_INLINE_GENERATED_CPP_BY_NAME(GeometryCollectionMeshNodes)
 
-namespace Dataflow
+namespace UE::Dataflow
 {
 
 	void GeometryCollectionMeshNodes()
@@ -57,7 +57,7 @@ namespace Dataflow
 }
 
 
-void FPointsToMeshDataflowNode::Evaluate(Dataflow::FContext& Context, const FDataflowOutput* Out) const
+void FPointsToMeshDataflowNode::Evaluate(UE::Dataflow::FContext& Context, const FDataflowOutput* Out) const
 {
 	if (Out->IsA<TObjectPtr<UDynamicMesh>>(&Mesh) || Out->IsA<int32>(&TriangleCount))
 	{
@@ -87,7 +87,7 @@ void FPointsToMeshDataflowNode::Evaluate(Dataflow::FContext& Context, const FDat
 }
 
 
-void FBoxToMeshDataflowNode::Evaluate(Dataflow::FContext& Context, const FDataflowOutput* Out) const
+void FBoxToMeshDataflowNode::Evaluate(UE::Dataflow::FContext& Context, const FDataflowOutput* Out) const
 {
 	if (Out->IsA<TObjectPtr<UDynamicMesh>>(&Mesh) || Out->IsA<int32>(&TriangleCount))
 	{
@@ -110,7 +110,7 @@ void FBoxToMeshDataflowNode::Evaluate(Dataflow::FContext& Context, const FDatafl
 }
 
 
-void FMeshInfoDataflowNode::Evaluate(Dataflow::FContext& Context, const FDataflowOutput* Out) const
+void FMeshInfoDataflowNode::Evaluate(UE::Dataflow::FContext& Context, const FDataflowOutput* Out) const
 {
 	if (Out->IsA<FString>(&InfoString))
 	{
@@ -128,7 +128,7 @@ void FMeshInfoDataflowNode::Evaluate(Dataflow::FContext& Context, const FDataflo
 }
 
 
-void FMeshToCollectionDataflowNode::Evaluate(Dataflow::FContext& Context, const FDataflowOutput* Out) const
+void FMeshToCollectionDataflowNode::Evaluate(UE::Dataflow::FContext& Context, const FDataflowOutput* Out) const
 {
 	if (Out->IsA<FManagedArrayCollection>(&Collection))
 	{
@@ -162,7 +162,7 @@ void FMeshToCollectionDataflowNode::Evaluate(Dataflow::FContext& Context, const
 }
 
 
-void FCollectionToMeshDataflowNode::Evaluate(Dataflow::FContext& Context, const FDataflowOutput* Out) const
+void FCollectionToMeshDataflowNode::Evaluate(UE::Dataflow::FContext& Context, const FDataflowOutput* Out) const
 {
 #if WITH_EDITORONLY_DATA
 	if (Out->IsA<TObjectPtr<UDynamicMesh>>(&Mesh))
@@ -213,7 +213,7 @@ void FCollectionToMeshDataflowNode::Evaluate(Dataflow::FContext& Context, const
 }
 
 
-void FStaticMeshToMeshDataflowNode::Evaluate(Dataflow::FContext& Context, const FDataflowOutput* Out) const
+void FStaticMeshToMeshDataflowNode::Evaluate(UE::Dataflow::FContext& Context, const FDataflowOutput* Out) const
 {
 #if WITH_EDITORONLY_DATA
 	if (Out->IsA(&Mesh))
@@ -243,7 +243,7 @@ void FStaticMeshToMeshDataflowNode::Evaluate(Dataflow::FContext& Context, const
 }
 
 
-void FMeshAppendDataflowNode::Evaluate(Dataflow::FContext& Context, const FDataflowOutput* Out) const
+void FMeshAppendDataflowNode::Evaluate(UE::Dataflow::FContext& Context, const FDataflowOutput* Out) const
 {
 	if (Out->IsA<TObjectPtr<UDynamicMesh>>(&Mesh))
 	{
@@ -281,7 +281,7 @@ void FMeshAppendDataflowNode::Evaluate(Dataflow::FContext& Context, const FDataf
 }
 
 
-void FMeshBooleanDataflowNode::Evaluate(Dataflow::FContext& Context, const FDataflowOutput* Out) const
+void FMeshBooleanDataflowNode::Evaluate(UE::Dataflow::FContext& Context, const FDataflowOutput* Out) const
 {
 	if (Out->IsA<TObjectPtr<UDynamicMesh>>(&Mesh))
 	{
@@ -333,7 +333,7 @@ void FMeshBooleanDataflowNode::Evaluate(Dataflow::FContext& Context, const FData
 	}
 }
 
-void FMeshCopyToPointsDataflowNode::Evaluate(Dataflow::FContext& Context, const FDataflowOutput* Out) const
+void FMeshCopyToPointsDataflowNode::Evaluate(UE::Dataflow::FContext& Context, const FDataflowOutput* Out) const
 {
 	if (Out->IsA<TObjectPtr<UDynamicMesh>>(&Mesh))
 	{
@@ -380,7 +380,7 @@ void FMeshCopyToPointsDataflowNode::Evaluate(Dataflow::FContext& Context, const
 }
 
 
-void FGetMeshDataDataflowNode::Evaluate(Dataflow::FContext& Context, const FDataflowOutput* Out) const
+void FGetMeshDataDataflowNode::Evaluate(UE::Dataflow::FContext& Context, const FDataflowOutput* Out) const
 {
 	if (Out->IsA<int32>(&VertexCount))
 	{
diff --git a/Engine/Plugins/Experimental/GeometryCollectionPlugin/Source/GeometryCollectionNodes/Private/Dataflow/GeometryCollectionNodes.cpp b/Engine/Plugins/Experimental/GeometryCollectionPlugin/Source/GeometryCollectionNodes/Private/Dataflow/GeometryCollectionNodes.cpp
--- a/Engine/Plugins/Experimental/GeometryCollectionPlugin/Source/GeometryCollectionNodes/Private/Dataflow/GeometryCollectionNodes.cpp
+++ b/Engine/Plugins/Experimental/GeometryCollectionPlugin/Source/GeometryCollectionNodes/Private/Dataflow/GeometryCollectionNodes.cpp
@@ -41,7 +41,7 @@
 
 #include UE_INLINE_GENERATED_CPP_BY_NAME(GeometryCollectionNodes)
 
-namespace Dataflow
+namespace UE::Dataflow
 {
 	void GeometryCollectionEngineNodes()
 	{
@@ -98,7 +98,7 @@ namespace Dataflow
 	}
 }
 
-void FGetCollectionFromAssetDataflowNode::Evaluate(Dataflow::FContext& Context, const FDataflowOutput* Out) const
+void FGetCollectionFromAssetDataflowNode::Evaluate(UE::Dataflow::FContext& Context, const FDataflowOutput* Out) const
 {
 	if (Out->IsA<FManagedArrayCollection>(&Collection))
 	{
@@ -121,7 +121,7 @@ void FGetCollectionFromAssetDataflowNode::Evaluate(Dataflow::FContext& Context,
 }
 
 
-void FAppendCollectionAssetsDataflowNode::Evaluate(Dataflow::FContext& Context, const FDataflowOutput* Out) const
+void FAppendCollectionAssetsDataflowNode::Evaluate(UE::Dataflow::FContext& Context, const FDataflowOutput* Out) const
 {
 	if (Out->IsA<DataType>(&Collection1))
 	{
@@ -144,7 +144,7 @@ void FAppendCollectionAssetsDataflowNode::Evaluate(Dataflow::FContext& Context,
 }
 
 
-void FPrintStringDataflowNode::Evaluate(Dataflow::FContext& Context, const FDataflowOutput* Out) const
+void FPrintStringDataflowNode::Evaluate(UE::Dataflow::FContext& Context, const FDataflowOutput* Out) const
 {
 	FString Value = GetValue<FString>(Context, &String);
 
@@ -158,7 +158,7 @@ void FPrintStringDataflowNode::Evaluate(Dataflow::FContext& Context, const FData
 	}
 }
 
-void FLogStringDataflowNode::Evaluate(Dataflow::FContext& Context, const FDataflowOutput* Out) const
+void FLogStringDataflowNode::Evaluate(UE::Dataflow::FContext& Context, const FDataflowOutput* Out) const
 {
 	if (bPrintToLog)
 	{
@@ -169,7 +169,7 @@ void FLogStringDataflowNode::Evaluate(Dataflow::FContext& Context, const FDatafl
 
 
 
-void FBoundingBoxDataflowNode::Evaluate(Dataflow::FContext& Context, const FDataflowOutput* Out) const
+void FBoundingBoxDataflowNode::Evaluate(UE::Dataflow::FContext& Context, const FDataflowOutput* Out) const
 {
 	if (Out->IsA<FBox>(&BoundingBox))
 	{
@@ -182,7 +182,7 @@ void FBoundingBoxDataflowNode::Evaluate(Dataflow::FContext& Context, const FData
 	}
 }
 
-void FGetBoxLengthsDataflowNode::Evaluate(Dataflow::FContext& Context, const FDataflowOutput* Out) const
+void FGetBoxLengthsDataflowNode::Evaluate(UE::Dataflow::FContext& Context, const FDataflowOutput* Out) const
 {
 	if (Out->IsA(&Lengths))
 	{
@@ -201,7 +201,7 @@ void FGetBoxLengthsDataflowNode::Evaluate(Dataflow::FContext& Context, const FDa
 }
 
 
-void FExpandBoundingBoxDataflowNode::Evaluate(Dataflow::FContext& Context, const FDataflowOutput* Out) const
+void FExpandBoundingBoxDataflowNode::Evaluate(UE::Dataflow::FContext& Context, const FDataflowOutput* Out) const
 {
 	FBox BBox = GetValue<FBox>(Context, &BoundingBox);
 
@@ -228,7 +228,7 @@ void FExpandBoundingBoxDataflowNode::Evaluate(Dataflow::FContext& Context, const
 }
 
 
-void FExpandVectorDataflowNode::Evaluate(Dataflow::FContext& Context, const FDataflowOutput* Out) const
+void FExpandVectorDataflowNode::Evaluate(UE::Dataflow::FContext& Context, const FDataflowOutput* Out) const
 {
 	FVector VectorVal = GetValue<FVector>(Context, &Vector);
 
@@ -246,7 +246,7 @@ void FExpandVectorDataflowNode::Evaluate(Dataflow::FContext& Context, const FDat
 	}
 }
 
-void FStringAppendDataflowNode::Evaluate(Dataflow::FContext& Context, const FDataflowOutput* Out) const
+void FStringAppendDataflowNode::Evaluate(UE::Dataflow::FContext& Context, const FDataflowOutput* Out) const
 {
 	if (Out->IsA<FString>(&String))
 	{
@@ -256,7 +256,7 @@ void FStringAppendDataflowNode::Evaluate(Dataflow::FContext& Context, const FDat
 }
 
 
-void FHashStringDataflowNode::Evaluate(Dataflow::FContext& Context, const FDataflowOutput* Out) const
+void FHashStringDataflowNode::Evaluate(UE::Dataflow::FContext& Context, const FDataflowOutput* Out) const
 {
 	if (Out->IsA<int32>(&Hash))
 	{
@@ -264,7 +264,7 @@ void FHashStringDataflowNode::Evaluate(Dataflow::FContext& Context, const FDataf
 	}
 }
 
-void FHashVectorDataflowNode::Evaluate(Dataflow::FContext& Context, const FDataflowOutput* Out) const
+void FHashVectorDataflowNode::Evaluate(UE::Dataflow::FContext& Context, const FDataflowOutput* Out) const
 {
 	if (Out->IsA<int32>(&Hash))
 	{
@@ -273,7 +273,7 @@ void FHashVectorDataflowNode::Evaluate(Dataflow::FContext& Context, const FDataf
 }
 
 
-void FGetBoundingBoxesFromCollectionDataflowNode::Evaluate(Dataflow::FContext& Context, const FDataflowOutput* Out) const
+void FGetBoundingBoxesFromCollectionDataflowNode::Evaluate(UE::Dataflow::FContext& Context, const FDataflowOutput* Out) const
 {
 	if (Out->IsA<TArray<FBox>>(&BoundingBoxes))
 	{
@@ -312,7 +312,7 @@ void FGetBoundingBoxesFromCollectionDataflowNode::Evaluate(Dataflow::FContext& C
 	}
 }
 
-void FGetRootIndexFromCollectionDataflowNode::Evaluate(Dataflow::FContext& Context, const FDataflowOutput* Out) const
+void FGetRootIndexFromCollectionDataflowNode::Evaluate(UE::Dataflow::FContext& Context, const FDataflowOutput* Out) const
 {
 	if (Out->IsA(&RootIndex))
 	{
@@ -323,7 +323,7 @@ void FGetRootIndexFromCollectionDataflowNode::Evaluate(Dataflow::FContext& Conte
 	}
 }
 
-void FGetCentroidsFromCollectionDataflowNode::Evaluate(Dataflow::FContext& Context, const FDataflowOutput* Out) const
+void FGetCentroidsFromCollectionDataflowNode::Evaluate(UE::Dataflow::FContext& Context, const FDataflowOutput* Out) const
 {
 	if (Out->IsA<TArray<FVector>>(&Centroids))
 	{
@@ -362,7 +362,7 @@ void FGetCentroidsFromCollectionDataflowNode::Evaluate(Dataflow::FContext& Conte
 }
 
 
-void FTransformCollectionDataflowNode::Evaluate(Dataflow::FContext& Context, const FDataflowOutput* Out) const
+void FTransformCollectionDataflowNode::Evaluate(UE::Dataflow::FContext& Context, const FDataflowOutput* Out) const
 {
 	if (Out->IsA<FManagedArrayCollection>(&Collection))
 	{
@@ -385,7 +385,7 @@ void FTransformCollectionDataflowNode::Evaluate(Dataflow::FContext& Context, con
 }
 
 
-void FBakeTransformsInCollectionDataflowNode::Evaluate(Dataflow::FContext& Context, const FDataflowOutput* Out) const
+void FBakeTransformsInCollectionDataflowNode::Evaluate(UE::Dataflow::FContext& Context, const FDataflowOutput* Out) const
 {
 	if (Out->IsA<FManagedArrayCollection>(&Collection))
 	{
@@ -409,7 +409,7 @@ void FBakeTransformsInCollectionDataflowNode::Evaluate(Dataflow::FContext& Conte
 }
 
 
-void FTransformMeshDataflowNode::Evaluate(Dataflow::FContext& Context, const FDataflowOutput* Out) const
+void FTransformMeshDataflowNode::Evaluate(UE::Dataflow::FContext& Context, const FDataflowOutput* Out) const
 {
 	if (Out->IsA<TObjectPtr<UDynamicMesh>>(&Mesh))
 	{
@@ -468,7 +468,7 @@ namespace
 }
 
 
-void FCompareIntDataflowNode::Evaluate(Dataflow::FContext& Context, const FDataflowOutput* Out) const
+void FCompareIntDataflowNode::Evaluate(UE::Dataflow::FContext& Context, const FDataflowOutput* Out) const
 {
 	if (Out->IsA<bool>(&Result))
 	{
@@ -481,7 +481,7 @@ void FCompareIntDataflowNode::Evaluate(Dataflow::FContext& Context, const FDataf
 }
 
 
-void FCompareFloatDataflowNode::Evaluate(Dataflow::FContext& Context, const FDataflowOutput* Out) const
+void FCompareFloatDataflowNode::Evaluate(UE::Dataflow::FContext& Context, const FDataflowOutput* Out) const
 {
 	if (Out->IsA(&Result))
 	{
@@ -494,7 +494,7 @@ void FCompareFloatDataflowNode::Evaluate(Dataflow::FContext& Context, const FDat
 }
 
 
-void FBranchMeshDataflowNode::Evaluate(Dataflow::FContext& Context, const FDataflowOutput* Out) const
+void FBranchMeshDataflowNode::Evaluate(UE::Dataflow::FContext& Context, const FDataflowOutput* Out) const
 {
 	if (Out->IsA<TObjectPtr<UDynamicMesh>>(&Mesh))
 	{
@@ -524,7 +524,7 @@ void FBranchMeshDataflowNode::Evaluate(Dataflow::FContext& Context, const FDataf
 }
 
 
-void FBranchCollectionDataflowNode::Evaluate(Dataflow::FContext& Context, const FDataflowOutput* Out) const
+void FBranchCollectionDataflowNode::Evaluate(UE::Dataflow::FContext& Context, const FDataflowOutput* Out) const
 {
 	if (Out->IsA(&ChosenCollection))
 	{
@@ -566,7 +566,7 @@ namespace {
 	}
 }
 
-void FGetSchemaDataflowNode::Evaluate(Dataflow::FContext& Context, const FDataflowOutput* Out) const
+void FGetSchemaDataflowNode::Evaluate(UE::Dataflow::FContext& Context, const FDataflowOutput* Out) const
 {
 	if (Out->IsA<FString>(&String))
 	{
@@ -602,7 +602,7 @@ void FGetSchemaDataflowNode::Evaluate(Dataflow::FContext& Context, const FDatafl
 }
 
 
-void FRemoveOnBreakDataflowNode::Evaluate(Dataflow::FContext& Context, const FDataflowOutput* Out) const
+void FRemoveOnBreakDataflowNode::Evaluate(UE::Dataflow::FContext& Context, const FDataflowOutput* Out) const
 {
 	if (Out->IsA<FManagedArrayCollection>(&Collection))
 	{
@@ -641,7 +641,7 @@ void FRemoveOnBreakDataflowNode::Evaluate(Dataflow::FContext& Context, const FDa
 	}
 }
 
-void FSetAnchorStateDataflowNode::Evaluate(Dataflow::FContext& Context, const FDataflowOutput* Out) const
+void FSetAnchorStateDataflowNode::Evaluate(UE::Dataflow::FContext& Context, const FDataflowOutput* Out) const
 {
 	if (Out->IsA<FManagedArrayCollection>(&Collection))
 	{
@@ -681,7 +681,7 @@ void FSetAnchorStateDataflowNode::Evaluate(Dataflow::FContext& Context, const FD
 }
 
 
-void FProximityDataflowNode::Evaluate(Dataflow::FContext& Context, const FDataflowOutput* Out) const
+void FProximityDataflowNode::Evaluate(UE::Dataflow::FContext& Context, const FDataflowOutput* Out) const
 {
 	if (Out->IsA<FManagedArrayCollection>(&Collection))
 	{
@@ -718,7 +718,7 @@ void FProximityDataflowNode::Evaluate(Dataflow::FContext& Context, const FDatafl
 }
 
 
-void FCollectionSetPivotDataflowNode::Evaluate(Dataflow::FContext& Context, const FDataflowOutput* Out) const
+void FCollectionSetPivotDataflowNode::Evaluate(UE::Dataflow::FContext& Context, const FDataflowOutput* Out) const
 {
 	if (Out->IsA<FManagedArrayCollection>(&Collection))
 	{
@@ -772,7 +772,7 @@ static void AddAndFillAttribute(FManagedArrayCollection& InCollection, FName Att
 	CustomAttribute.AddAndFill(DefaultValue);
 }
 
-void FAddCustomCollectionAttributeDataflowNode::Evaluate(Dataflow::FContext& Context, const FDataflowOutput* Out) const
+void FAddCustomCollectionAttributeDataflowNode::Evaluate(UE::Dataflow::FContext& Context, const FDataflowOutput* Out) const
 {
 	if (Out->IsA<FManagedArrayCollection>(&Collection))
 	{
@@ -882,7 +882,7 @@ void FAddCustomCollectionAttributeDataflowNode::Evaluate(Dataflow::FContext& Con
 }
 
 
-void FGetNumElementsInCollectionGroupDataflowNode::Evaluate(Dataflow::FContext& Context, const FDataflowOutput* Out) const
+void FGetNumElementsInCollectionGroupDataflowNode::Evaluate(UE::Dataflow::FContext& Context, const FDataflowOutput* Out) const
 {
 	if (Out->IsA<int32>(&NumElements))
 	{
@@ -913,7 +913,7 @@ void FGetNumElementsInCollectionGroupDataflowNode::Evaluate(Dataflow::FContext&
 }
 
 
-void FGetCollectionAttributeDataTypedDataflowNode::Evaluate(Dataflow::FContext& Context, const FDataflowOutput* Out) const
+void FGetCollectionAttributeDataTypedDataflowNode::Evaluate(UE::Dataflow::FContext& Context, const FDataflowOutput* Out) const
 {
 	if (Out->IsA<TArray<bool>>(&BoolAttributeData) ||
 		Out->IsA<TArray<float>>(&FloatAttributeData) ||
@@ -1020,7 +1020,7 @@ void FGetCollectionAttributeDataTypedDataflowNode::Evaluate(Dataflow::FContext&
 }
 
 template<typename T>
-static void SetAttributeData(const FDataflowNode* DataflowNode, Dataflow::FContext& Context, FManagedArrayCollection& InCollection, const TArray<T>& Property, FName AttributeName, FName GroupName)
+static void SetAttributeData(const FDataflowNode* DataflowNode, UE::Dataflow::FContext& Context, FManagedArrayCollection& InCollection, const TArray<T>& Property, FName AttributeName, FName GroupName)
 {
 	if (DataflowNode && DataflowNode->IsConnected<TArray<T>>(&Property))
 	{
@@ -1040,7 +1040,7 @@ static void SetAttributeData(const FDataflowNode* DataflowNode, Dataflow::FConte
 	}
 }
 
-void FSetCollectionAttributeDataTypedDataflowNode::Evaluate(Dataflow::FContext& Context, const FDataflowOutput* Out) const
+void FSetCollectionAttributeDataTypedDataflowNode::Evaluate(UE::Dataflow::FContext& Context, const FDataflowOutput* Out) const
 {
 	if (Out->IsA<FManagedArrayCollection>(&Collection))
 	{
@@ -1109,13 +1109,13 @@ void FSetCollectionAttributeDataTypedDataflowNode::Evaluate(Dataflow::FContext&
 	}
 }
 
-void FSelectionToVertexListDataflowNode::Evaluate(Dataflow::FContext& Context, const FDataflowOutput* Out) const
+void FSelectionToVertexListDataflowNode::Evaluate(UE::Dataflow::FContext& Context, const FDataflowOutput* Out) const
 {
 	const FDataflowVertexSelection& InVertexSelection = GetValue<FDataflowVertexSelection>(Context, &VertexSelection);
 	SetValue(Context, InVertexSelection.AsArray(), &VertexList);
 }
 
-void FMultiplyTransformDataflowNode::Evaluate(Dataflow::FContext& Context, const FDataflowOutput* Out) const
+void FMultiplyTransformDataflowNode::Evaluate(UE::Dataflow::FContext& Context, const FDataflowOutput* Out) const
 {
 	if (Out->IsA<FTransform>(&OutTransform))
 	{
@@ -1126,7 +1126,7 @@ void FMultiplyTransformDataflowNode::Evaluate(Dataflow::FContext& Context, const
 	}
 }
 
-void FInvertTransformDataflowNode::Evaluate(Dataflow::FContext& Context, const FDataflowOutput* Out) const
+void FInvertTransformDataflowNode::Evaluate(UE::Dataflow::FContext& Context, const FDataflowOutput* Out) const
 {
 	if (Out->IsA<FTransform>(&OutTransform))
 	{
@@ -1136,7 +1136,7 @@ void FInvertTransformDataflowNode::Evaluate(Dataflow::FContext& Context, const F
 	}
 }
 
-void FBranchFloatDataflowNode::Evaluate(Dataflow::FContext& Context, const FDataflowOutput* Out) const
+void FBranchFloatDataflowNode::Evaluate(UE::Dataflow::FContext& Context, const FDataflowOutput* Out) const
 {
 	if (Out->IsA<float>(&ReturnValue))
 	{
@@ -1157,7 +1157,7 @@ void FBranchFloatDataflowNode::Evaluate(Dataflow::FContext& Context, const FData
 	}
 }
 
-void FBranchIntDataflowNode::Evaluate(Dataflow::FContext& Context, const FDataflowOutput* Out) const
+void FBranchIntDataflowNode::Evaluate(UE::Dataflow::FContext& Context, const FDataflowOutput* Out) const
 {
 	if (Out->IsA<int32>(&ReturnValue))
 	{
diff --git a/Engine/Plugins/Experimental/GeometryCollectionPlugin/Source/GeometryCollectionNodes/Private/Dataflow/GeometryCollectionNodesPlugin.cpp b/Engine/Plugins/Experimental/GeometryCollectionPlugin/Source/GeometryCollectionNodes/Private/Dataflow/GeometryCollectionNodesPlugin.cpp
--- a/Engine/Plugins/Experimental/GeometryCollectionPlugin/Source/GeometryCollectionNodes/Private/Dataflow/GeometryCollectionNodesPlugin.cpp
+++ b/Engine/Plugins/Experimental/GeometryCollectionPlugin/Source/GeometryCollectionNodes/Private/Dataflow/GeometryCollectionNodesPlugin.cpp
@@ -40,25 +40,25 @@
 
 void IGeometryCollectionNodesPlugin::StartupModule()
 {
-	Dataflow::GeometryCollectionEngineNodes();
-	Dataflow::GeometryCollectionEngineAssetNodes();
-	Dataflow::GeometryCollectionProcessingNodes();
-	Dataflow::GeometryCollectionSelectionNodes();
-	Dataflow::GeometryCollectionMeshNodes();
-	Dataflow::GeometryCollectionClusteringNodes();
-	Dataflow::GeometryCollectionFracturingNodes();
-	Dataflow::GeometryCollectionEditNodes();
-	Dataflow::GeometryCollectionUtilityNodes();
-	Dataflow::GeometryCollectionMaterialNodes();
-	Dataflow::GeometryCollectionFieldNodes();
-	Dataflow::GeometryCollectionOverrideNodes();
-	Dataflow::GeometryCollectionMakeNodes();
-	Dataflow::GeometryCollectionMathNodes();
-	Dataflow::GeometryCollectionConversionNodes();
-	Dataflow::GeometryCollectionVerticesNodes();
-	Dataflow::GeometryCollectionArrayNodes();
-	Dataflow::GeometryCollectionDebugNodes();
-	Dataflow::GeometryCollectionSamplingNodes();
+	UE::Dataflow::GeometryCollectionEngineNodes();
+	UE::Dataflow::GeometryCollectionEngineAssetNodes();
+	UE::Dataflow::GeometryCollectionProcessingNodes();
+	UE::Dataflow::GeometryCollectionSelectionNodes();
+	UE::Dataflow::GeometryCollectionMeshNodes();
+	UE::Dataflow::GeometryCollectionClusteringNodes();
+	UE::Dataflow::GeometryCollectionFracturingNodes();
+	UE::Dataflow::GeometryCollectionEditNodes();
+	UE::Dataflow::GeometryCollectionUtilityNodes();
+	UE::Dataflow::GeometryCollectionMaterialNodes();
+	UE::Dataflow::GeometryCollectionFieldNodes();
+	UE::Dataflow::GeometryCollectionOverrideNodes();
+	UE::Dataflow::GeometryCollectionMakeNodes();
+	UE::Dataflow::GeometryCollectionMathNodes();
+	UE::Dataflow::GeometryCollectionConversionNodes();
+	UE::Dataflow::GeometryCollectionVerticesNodes();
+	UE::Dataflow::GeometryCollectionArrayNodes();
+	UE::Dataflow::GeometryCollectionDebugNodes();
+	UE::Dataflow::GeometryCollectionSamplingNodes();
 	DATAFLOW_NODE_REGISTER_CREATION_FACTORY(FCreateColorArrayFromFloatArrayDataflowNode);
 	DATAFLOW_NODE_REGISTER_CREATION_FACTORY(FGeometryCollectionVertexScalarToVertexIndicesNode);
 	DATAFLOW_NODE_REGISTER_CREATION_FACTORY(FSetVertexColorFromFloatArrayDataflowNode);
diff --git a/Engine/Plugins/Experimental/GeometryCollectionPlugin/Source/GeometryCollectionNodes/Private/Dataflow/GeometryCollectionOverrideNodes.cpp b/Engine/Plugins/Experimental/GeometryCollectionPlugin/Source/GeometryCollectionNodes/Private/Dataflow/GeometryCollectionOverrideNodes.cpp
--- a/Engine/Plugins/Experimental/GeometryCollectionPlugin/Source/GeometryCollectionNodes/Private/Dataflow/GeometryCollectionOverrideNodes.cpp
+++ b/Engine/Plugins/Experimental/GeometryCollectionPlugin/Source/GeometryCollectionNodes/Private/Dataflow/GeometryCollectionOverrideNodes.cpp
@@ -7,7 +7,7 @@
 
 #include UE_INLINE_GENERATED_CPP_BY_NAME(GeometryCollectionOverrideNodes)
 
-namespace Dataflow
+namespace UE::Dataflow
 {
 	void GeometryCollectionOverrideNodes()
 	{
@@ -24,7 +24,7 @@ namespace Dataflow
 }
 
 
-void FGetFloatOverrideFromAssetDataflowNode::Evaluate(Dataflow::FContext& Context, const FDataflowOutput* Out) const
+void FGetFloatOverrideFromAssetDataflowNode::Evaluate(UE::Dataflow::FContext& Context, const FDataflowOutput* Out) const
 {
 	if (Out->IsA(&Float) || Out->IsA(&FloatDefault) || Out->IsA(&IsOverriden))
 	{
@@ -32,7 +32,7 @@ void FGetFloatOverrideFromAssetDataflowNode::Evaluate(Dataflow::FContext& Contex
 		float NewValue = DefaultValue;
 		bool bIsOverriden = false;
 
-		if (const Dataflow::FEngineContext* EngineContext = Context.AsType<Dataflow::FEngineContext>())
+		if (const UE::Dataflow::FEngineContext* EngineContext = Context.AsType<UE::Dataflow::FEngineContext>())
 		{
 			const FString ValueFromAsset = GetValueFromAsset(Context, EngineContext->Owner);
 
@@ -49,7 +49,7 @@ void FGetFloatOverrideFromAssetDataflowNode::Evaluate(Dataflow::FContext& Contex
 	}
 }
 
-void FGetIntOverrideFromAssetDataflowNode::Evaluate(Dataflow::FContext& Context, const FDataflowOutput* Out) const
+void FGetIntOverrideFromAssetDataflowNode::Evaluate(UE::Dataflow::FContext& Context, const FDataflowOutput* Out) const
 {
 	if (Out->IsA(&Int) || Out->IsA(&IntDefault) || Out->IsA(&IsOverriden))
 	{
@@ -57,7 +57,7 @@ void FGetIntOverrideFromAssetDataflowNode::Evaluate(Dataflow::FContext& Context,
 		int32 NewValue = DefaultValue;
 		bool bIsOverriden = false;
 
-		if (const Dataflow::FEngineContext* EngineContext = Context.AsType<Dataflow::FEngineContext>())
+		if (const UE::Dataflow::FEngineContext* EngineContext = Context.AsType<UE::Dataflow::FEngineContext>())
 		{
 			const FString ValueFromAsset = GetValueFromAsset(Context, EngineContext->Owner);
 
@@ -93,7 +93,7 @@ static bool StringToBool(const FString& InString, bool InDefault)
 	return Result;
 }
 
-void FGetBoolOverrideFromAssetDataflowNode::Evaluate(Dataflow::FContext& Context, const FDataflowOutput* Out) const
+void FGetBoolOverrideFromAssetDataflowNode::Evaluate(UE::Dataflow::FContext& Context, const FDataflowOutput* Out) const
 {
 	if (Out->IsA(&Bool) || Out->IsA(&BoolDefault) || Out->IsA(&IsOverriden))
 	{
@@ -101,7 +101,7 @@ void FGetBoolOverrideFromAssetDataflowNode::Evaluate(Dataflow::FContext& Context
 		bool NewValue = DefaultValue;
 		bool bIsOverriden = false;
 
-		if (const Dataflow::FEngineContext* EngineContext = Context.AsType<Dataflow::FEngineContext>())
+		if (const UE::Dataflow::FEngineContext* EngineContext = Context.AsType<UE::Dataflow::FEngineContext>())
 		{
 			const FString ValueFromAsset = GetValueFromAsset(Context, EngineContext->Owner);
 
@@ -118,7 +118,7 @@ void FGetBoolOverrideFromAssetDataflowNode::Evaluate(Dataflow::FContext& Context
 	}
 }
 
-void FGetStringOverrideFromAssetDataflowNode::Evaluate(Dataflow::FContext& Context, const FDataflowOutput* Out) const
+void FGetStringOverrideFromAssetDataflowNode::Evaluate(UE::Dataflow::FContext& Context, const FDataflowOutput* Out) const
 {
 	if (Out->IsA(&String) || Out->IsA(&StringDefault) || Out->IsA(&IsOverriden))
 	{
@@ -126,7 +126,7 @@ void FGetStringOverrideFromAssetDataflowNode::Evaluate(Dataflow::FContext& Conte
 		FString NewValue = DefaultValue;
 		bool bIsOverriden = false;
 
-		if (const Dataflow::FEngineContext* EngineContext = Context.AsType<Dataflow::FEngineContext>())
+		if (const UE::Dataflow::FEngineContext* EngineContext = Context.AsType<UE::Dataflow::FEngineContext>())
 		{
 			const FString ValueFromAsset = GetValueFromAsset(Context, EngineContext->Owner);
 
diff --git a/Engine/Plugins/Experimental/GeometryCollectionPlugin/Source/GeometryCollectionNodes/Private/Dataflow/GeometryCollectionProcessingNodes.cpp b/Engine/Plugins/Experimental/GeometryCollectionPlugin/Source/GeometryCollectionNodes/Private/Dataflow/GeometryCollectionProcessingNodes.cpp
--- a/Engine/Plugins/Experimental/GeometryCollectionPlugin/Source/GeometryCollectionNodes/Private/Dataflow/GeometryCollectionProcessingNodes.cpp
+++ b/Engine/Plugins/Experimental/GeometryCollectionPlugin/Source/GeometryCollectionNodes/Private/Dataflow/GeometryCollectionProcessingNodes.cpp
@@ -6,7 +6,7 @@
 #include "GeometryCollection/GeometryCollection.h"
 
 
-namespace Dataflow
+namespace UE::Dataflow
 {
 	void GeometryCollectionProcessingNodes()
 	{
@@ -14,7 +14,7 @@ namespace Dataflow
 	}
 }
 
-void FCloseGeometryOnCollectionDataflowNode::Evaluate(Dataflow::FContext& Context, const FDataflowOutput* Out) const
+void FCloseGeometryOnCollectionDataflowNode::Evaluate(UE::Dataflow::FContext& Context, const FDataflowOutput* Out) const
 {
 	if (Out->IsA<DataType>(&Collection))
 	{
diff --git a/Engine/Plugins/Experimental/GeometryCollectionPlugin/Source/GeometryCollectionNodes/Private/Dataflow/GeometryCollectionSamplingNodes.cpp b/Engine/Plugins/Experimental/GeometryCollectionPlugin/Source/GeometryCollectionNodes/Private/Dataflow/GeometryCollectionSamplingNodes.cpp
--- a/Engine/Plugins/Experimental/GeometryCollectionPlugin/Source/GeometryCollectionNodes/Private/Dataflow/GeometryCollectionSamplingNodes.cpp
+++ b/Engine/Plugins/Experimental/GeometryCollectionPlugin/Source/GeometryCollectionNodes/Private/Dataflow/GeometryCollectionSamplingNodes.cpp
@@ -6,7 +6,7 @@
 
 #include UE_INLINE_GENERATED_CPP_BY_NAME(GeometryCollectionSamplingNodes)
 
-namespace Dataflow
+namespace UE::Dataflow
 {
 
 	void GeometryCollectionSamplingNodes()
@@ -22,7 +22,7 @@ namespace Dataflow
 	}
 }
 
-void FUniformPointSamplingDataflowNode::Evaluate(Dataflow::FContext& Context, const FDataflowOutput* Out) const
+void FUniformPointSamplingDataflowNode::Evaluate(UE::Dataflow::FContext& Context, const FDataflowOutput* Out) const
 {
 	if (Out->IsA(&SamplePoints) || 
 		Out->IsA(&SampleTriangleIDs) ||
@@ -66,7 +66,7 @@ void FUniformPointSamplingDataflowNode::Evaluate(Dataflow::FContext& Context, co
 	}
 }
 
-void FNonUniformPointSamplingDataflowNode::Evaluate(Dataflow::FContext& Context, const FDataflowOutput* Out) const
+void FNonUniformPointSamplingDataflowNode::Evaluate(UE::Dataflow::FContext& Context, const FDataflowOutput* Out) const
 {
 	if (Out->IsA(&SamplePoints) ||
 		Out->IsA(&SampleRadii) ||
@@ -117,7 +117,7 @@ void FNonUniformPointSamplingDataflowNode::Evaluate(Dataflow::FContext& Context,
 	}
 }
 
-void FVertexWeightedPointSamplingDataflowNode::Evaluate(Dataflow::FContext& Context, const FDataflowOutput* Out) const
+void FVertexWeightedPointSamplingDataflowNode::Evaluate(UE::Dataflow::FContext& Context, const FDataflowOutput* Out) const
 {
 	if (Out->IsA(&SamplePoints) ||
 		Out->IsA(&SampleRadii) ||
