Change Dataflow namespace to UE::Dataflow to be ok with coding standard

diff --git a/Engine/Plugins/Experimental/GeometryCollectionPlugin/Source/GeometryCollectionNodes/Private/Dataflow/GeometryCollectionSelectionNodes.cpp b/Engine/Plugins/Experimental/GeometryCollectionPlugin/Source/GeometryCollectionNodes/Private/Dataflow/GeometryCollectionSelectionNodes.cpp
--- a/Engine/Plugins/Experimental/GeometryCollectionPlugin/Source/GeometryCollectionNodes/Private/Dataflow/GeometryCollectionSelectionNodes.cpp
+++ b/Engine/Plugins/Experimental/GeometryCollectionPlugin/Source/GeometryCollectionNodes/Private/Dataflow/GeometryCollectionSelectionNodes.cpp
@@ -31,7 +31,7 @@
 
 #include UE_INLINE_GENERATED_CPP_BY_NAME(GeometryCollectionSelectionNodes)
 
-namespace Dataflow
+namespace UE::Dataflow
 {
 
 	void GeometryCollectionSelectionNodes()
@@ -78,7 +78,7 @@ namespace Dataflow
 }
 
 
-void FCollectionTransformSelectionAllDataflowNode::Evaluate(Dataflow::FContext& Context, const FDataflowOutput* Out) const
+void FCollectionTransformSelectionAllDataflowNode::Evaluate(UE::Dataflow::FContext& Context, const FDataflowOutput* Out) const
 {
 	if (Out->IsA<FDataflowTransformSelection>(&TransformSelection))
 	{
@@ -100,7 +100,7 @@ void FCollectionTransformSelectionAllDataflowNode::Evaluate(Dataflow::FContext&
 }
 
 
-void FCollectionTransformSelectionSetOperationDataflowNode::Evaluate(Dataflow::FContext& Context, const FDataflowOutput* Out) const
+void FCollectionTransformSelectionSetOperationDataflowNode::Evaluate(UE::Dataflow::FContext& Context, const FDataflowOutput* Out) const
 {
 	if (Out->IsA<FDataflowTransformSelection>(&TransformSelection))
 	{
@@ -229,7 +229,7 @@ static void BuildHierarchicalOutput(const TManagedArray<int32>& Parents,
 }
 
 
-void FCollectionTransformSelectionInfoDataflowNode::Evaluate(Dataflow::FContext& Context, const FDataflowOutput* Out) const
+void FCollectionTransformSelectionInfoDataflowNode::Evaluate(UE::Dataflow::FContext& Context, const FDataflowOutput* Out) const
 {
 	if (Out->IsA<FString>(&String))
 	{
@@ -278,7 +278,7 @@ void FCollectionTransformSelectionInfoDataflowNode::Evaluate(Dataflow::FContext&
 }
 
 
-void FCollectionTransformSelectionNoneDataflowNode::Evaluate(Dataflow::FContext& Context, const FDataflowOutput* Out) const
+void FCollectionTransformSelectionNoneDataflowNode::Evaluate(UE::Dataflow::FContext& Context, const FDataflowOutput* Out) const
 {
 	if (Out->IsA<FDataflowTransformSelection>(&TransformSelection))
 	{
@@ -300,7 +300,7 @@ void FCollectionTransformSelectionNoneDataflowNode::Evaluate(Dataflow::FContext&
 }
 
 
-void FCollectionTransformSelectionInvertDataflowNode::Evaluate(Dataflow::FContext& Context, const FDataflowOutput* Out) const
+void FCollectionTransformSelectionInvertDataflowNode::Evaluate(UE::Dataflow::FContext& Context, const FDataflowOutput* Out) const
 {
 	if (Out->IsA<FDataflowTransformSelection>(&TransformSelection))
 	{
@@ -313,7 +313,7 @@ void FCollectionTransformSelectionInvertDataflowNode::Evaluate(Dataflow::FContex
 }
 
 
-void FCollectionTransformSelectionRandomDataflowNode::Evaluate(Dataflow::FContext& Context, const FDataflowOutput* Out) const
+void FCollectionTransformSelectionRandomDataflowNode::Evaluate(UE::Dataflow::FContext& Context, const FDataflowOutput* Out) const
 {
 	if (Out->IsA<FDataflowTransformSelection>(&TransformSelection))
 	{
@@ -338,7 +338,7 @@ void FCollectionTransformSelectionRandomDataflowNode::Evaluate(Dataflow::FContex
 }
 
 
-void FCollectionTransformSelectionRootDataflowNode::Evaluate(Dataflow::FContext& Context, const FDataflowOutput* Out) const
+void FCollectionTransformSelectionRootDataflowNode::Evaluate(UE::Dataflow::FContext& Context, const FDataflowOutput* Out) const
 {
 	if (Out->IsA<FDataflowTransformSelection>(&TransformSelection))
 	{
@@ -360,7 +360,7 @@ void FCollectionTransformSelectionRootDataflowNode::Evaluate(Dataflow::FContext&
 }
 
 
-void FCollectionTransformSelectionCustomDataflowNode::Evaluate(Dataflow::FContext& Context, const FDataflowOutput* Out) const
+void FCollectionTransformSelectionCustomDataflowNode::Evaluate(UE::Dataflow::FContext& Context, const FDataflowOutput* Out) const
 {
 	if (Out->IsA<FDataflowTransformSelection>(&TransformSelection))
 	{
@@ -410,7 +410,7 @@ void FCollectionTransformSelectionCustomDataflowNode::Evaluate(Dataflow::FContex
 }
 
 
-void FCollectionTransformSelectionFromIndexArrayDataflowNode::Evaluate(Dataflow::FContext& Context, const FDataflowOutput* Out) const
+void FCollectionTransformSelectionFromIndexArrayDataflowNode::Evaluate(UE::Dataflow::FContext& Context, const FDataflowOutput* Out) const
 {
 	if (Out->IsA(&TransformSelection))
 	{
@@ -450,7 +450,7 @@ void FCollectionTransformSelectionFromIndexArrayDataflowNode::Evaluate(Dataflow:
 }
 
 
-void FCollectionTransformSelectionParentDataflowNode::Evaluate(Dataflow::FContext& Context, const FDataflowOutput* Out) const
+void FCollectionTransformSelectionParentDataflowNode::Evaluate(UE::Dataflow::FContext& Context, const FDataflowOutput* Out) const
 {
 	if (Out->IsA<FDataflowTransformSelection>(&TransformSelection))
 	{
@@ -472,7 +472,7 @@ void FCollectionTransformSelectionParentDataflowNode::Evaluate(Dataflow::FContex
 }
 
 
-void FCollectionTransformSelectionByPercentageDataflowNode::Evaluate(Dataflow::FContext& Context, const FDataflowOutput* Out) const
+void FCollectionTransformSelectionByPercentageDataflowNode::Evaluate(UE::Dataflow::FContext& Context, const FDataflowOutput* Out) const
 {
 	if (Out->IsA<FDataflowTransformSelection>(&TransformSelection))
 	{
@@ -491,7 +491,7 @@ void FCollectionTransformSelectionByPercentageDataflowNode::Evaluate(Dataflow::F
 }
 
 
-void FCollectionTransformSelectionChildrenDataflowNode::Evaluate(Dataflow::FContext& Context, const FDataflowOutput* Out) const
+void FCollectionTransformSelectionChildrenDataflowNode::Evaluate(UE::Dataflow::FContext& Context, const FDataflowOutput* Out) const
 {
 	if (Out->IsA<FDataflowTransformSelection>(&TransformSelection))
 	{
@@ -513,7 +513,7 @@ void FCollectionTransformSelectionChildrenDataflowNode::Evaluate(Dataflow::FCont
 }
 
 
-void FCollectionTransformSelectionSiblingsDataflowNode::Evaluate(Dataflow::FContext& Context, const FDataflowOutput* Out) const
+void FCollectionTransformSelectionSiblingsDataflowNode::Evaluate(UE::Dataflow::FContext& Context, const FDataflowOutput* Out) const
 {
 	if (Out->IsA<FDataflowTransformSelection>(&TransformSelection))
 	{
@@ -535,7 +535,7 @@ void FCollectionTransformSelectionSiblingsDataflowNode::Evaluate(Dataflow::FCont
 }
 
 
-void FCollectionTransformSelectionLevelDataflowNode::Evaluate(Dataflow::FContext& Context, const FDataflowOutput* Out) const
+void FCollectionTransformSelectionLevelDataflowNode::Evaluate(UE::Dataflow::FContext& Context, const FDataflowOutput* Out) const
 {
 	if (Out->IsA<FDataflowTransformSelection>(&TransformSelection))
 	{
@@ -557,7 +557,7 @@ void FCollectionTransformSelectionLevelDataflowNode::Evaluate(Dataflow::FContext
 }
 
 
-void FCollectionTransformSelectionTargetLevelDataflowNode::Evaluate(Dataflow::FContext& Context, const FDataflowOutput* Out) const
+void FCollectionTransformSelectionTargetLevelDataflowNode::Evaluate(UE::Dataflow::FContext& Context, const FDataflowOutput* Out) const
 {
 	if (Out->IsA(&TransformSelection))
 	{
@@ -581,7 +581,7 @@ void FCollectionTransformSelectionTargetLevelDataflowNode::Evaluate(Dataflow::FC
 }
 
 
-void FCollectionTransformSelectionContactDataflowNode::Evaluate(Dataflow::FContext& Context, const FDataflowOutput* Out) const
+void FCollectionTransformSelectionContactDataflowNode::Evaluate(UE::Dataflow::FContext& Context, const FDataflowOutput* Out) const
 {
 	if (Out->IsA<FDataflowTransformSelection>(&TransformSelection))
 	{
@@ -603,7 +603,7 @@ void FCollectionTransformSelectionContactDataflowNode::Evaluate(Dataflow::FConte
 }
 
 
-void FCollectionTransformSelectionLeafDataflowNode::Evaluate(Dataflow::FContext& Context, const FDataflowOutput* Out) const
+void FCollectionTransformSelectionLeafDataflowNode::Evaluate(UE::Dataflow::FContext& Context, const FDataflowOutput* Out) const
 {
 	if (Out->IsA<FDataflowTransformSelection>(&TransformSelection))
 	{
@@ -625,7 +625,7 @@ void FCollectionTransformSelectionLeafDataflowNode::Evaluate(Dataflow::FContext&
 }
 
 
-void FCollectionTransformSelectionClusterDataflowNode::Evaluate(Dataflow::FContext& Context, const FDataflowOutput* Out) const
+void FCollectionTransformSelectionClusterDataflowNode::Evaluate(UE::Dataflow::FContext& Context, const FDataflowOutput* Out) const
 {
 	if (Out->IsA(&TransformSelection))
 	{
@@ -649,7 +649,7 @@ void FCollectionTransformSelectionClusterDataflowNode::Evaluate(Dataflow::FConte
 	}
 }
 
-void FCollectionTransformSelectionClusterDataflowNode_v2::Evaluate(Dataflow::FContext& Context, const FDataflowOutput* Out) const
+void FCollectionTransformSelectionClusterDataflowNode_v2::Evaluate(UE::Dataflow::FContext& Context, const FDataflowOutput* Out) const
 {
 	if (Out->IsA(&TransformSelection))
 	{
@@ -670,7 +670,7 @@ void FCollectionTransformSelectionClusterDataflowNode_v2::Evaluate(Dataflow::FCo
 	}
 }
 
-void FCollectionTransformSelectionBySizeDataflowNode::Evaluate(Dataflow::FContext& Context, const FDataflowOutput* Out) const
+void FCollectionTransformSelectionBySizeDataflowNode::Evaluate(UE::Dataflow::FContext& Context, const FDataflowOutput* Out) const
 {
 	if (Out->IsA<FDataflowTransformSelection>(&TransformSelection))
 	{
@@ -695,7 +695,7 @@ void FCollectionTransformSelectionBySizeDataflowNode::Evaluate(Dataflow::FContex
 }
 
 
-void FCollectionTransformSelectionByVolumeDataflowNode::Evaluate(Dataflow::FContext& Context, const FDataflowOutput* Out) const
+void FCollectionTransformSelectionByVolumeDataflowNode::Evaluate(UE::Dataflow::FContext& Context, const FDataflowOutput* Out) const
 {
 	if (Out->IsA<FDataflowTransformSelection>(&TransformSelection))
 	{
@@ -720,7 +720,7 @@ void FCollectionTransformSelectionByVolumeDataflowNode::Evaluate(Dataflow::FCont
 }
 
 
-void FCollectionTransformSelectionInBoxDataflowNode::Evaluate(Dataflow::FContext& Context, const FDataflowOutput* Out) const
+void FCollectionTransformSelectionInBoxDataflowNode::Evaluate(UE::Dataflow::FContext& Context, const FDataflowOutput* Out) const
 {
 	if (Out->IsA<FDataflowTransformSelection>(&TransformSelection))
 	{
@@ -757,7 +757,7 @@ void FCollectionTransformSelectionInBoxDataflowNode::Evaluate(Dataflow::FContext
 }
 
 
-void FCollectionTransformSelectionInSphereDataflowNode::Evaluate(Dataflow::FContext& Context, const FDataflowOutput* Out) const
+void FCollectionTransformSelectionInSphereDataflowNode::Evaluate(UE::Dataflow::FContext& Context, const FDataflowOutput* Out) const
 {
 	if (Out->IsA<FDataflowTransformSelection>(&TransformSelection))
 	{
@@ -794,7 +794,7 @@ void FCollectionTransformSelectionInSphereDataflowNode::Evaluate(Dataflow::FCont
 }
 
 
-void FCollectionTransformSelectionByFloatAttrDataflowNode::Evaluate(Dataflow::FContext& Context, const FDataflowOutput* Out) const
+void FCollectionTransformSelectionByFloatAttrDataflowNode::Evaluate(UE::Dataflow::FContext& Context, const FDataflowOutput* Out) const
 {
 	if (Out->IsA<FDataflowTransformSelection>(&TransformSelection))
 	{
@@ -818,7 +818,7 @@ void FCollectionTransformSelectionByFloatAttrDataflowNode::Evaluate(Dataflow::FC
 	}
 }
 
-void FSelectFloatArrayIndicesInRangeDataflowNode::Evaluate(Dataflow::FContext& Context, const FDataflowOutput* Out) const
+void FSelectFloatArrayIndicesInRangeDataflowNode::Evaluate(UE::Dataflow::FContext& Context, const FDataflowOutput* Out) const
 {
 	if (Out->IsA(&Indices))
 	{
@@ -850,7 +850,7 @@ void FSelectFloatArrayIndicesInRangeDataflowNode::Evaluate(Dataflow::FContext& C
 	}
 }
 
-void FCollectionTransformSelectionByIntAttrDataflowNode::Evaluate(Dataflow::FContext& Context, const FDataflowOutput* Out) const
+void FCollectionTransformSelectionByIntAttrDataflowNode::Evaluate(UE::Dataflow::FContext& Context, const FDataflowOutput* Out) const
 {
 	if (Out->IsA<FDataflowTransformSelection>(&TransformSelection))
 	{
@@ -875,7 +875,7 @@ void FCollectionTransformSelectionByIntAttrDataflowNode::Evaluate(Dataflow::FCon
 }
 
 
-void FCollectionVertexSelectionCustomDataflowNode::Evaluate(Dataflow::FContext& Context, const FDataflowOutput* Out) const
+void FCollectionVertexSelectionCustomDataflowNode::Evaluate(UE::Dataflow::FContext& Context, const FDataflowOutput* Out) const
 {
 	if (Out->IsA<FDataflowVertexSelection>(&VertexSelection))
 	{
@@ -925,7 +925,7 @@ void FCollectionVertexSelectionCustomDataflowNode::Evaluate(Dataflow::FContext&
 }
 
 
-void FCollectionFaceSelectionCustomDataflowNode::Evaluate(Dataflow::FContext& Context, const FDataflowOutput* Out) const
+void FCollectionFaceSelectionCustomDataflowNode::Evaluate(UE::Dataflow::FContext& Context, const FDataflowOutput* Out) const
 {
 	if (Out->IsA<FDataflowFaceSelection>(&FaceSelection))
 	{
@@ -975,7 +975,7 @@ void FCollectionFaceSelectionCustomDataflowNode::Evaluate(Dataflow::FContext& Co
 }
 
 
-void FCollectionSelectionConvertDataflowNode::Evaluate(Dataflow::FContext& Context, const FDataflowOutput* Out) const
+void FCollectionSelectionConvertDataflowNode::Evaluate(UE::Dataflow::FContext& Context, const FDataflowOutput* Out) const
 {
 	if (Out->IsA(&TransformSelection))
 	{
@@ -1092,7 +1092,7 @@ void FCollectionSelectionConvertDataflowNode::Evaluate(Dataflow::FContext& Conte
 }
 
 
-void FCollectionFaceSelectionInvertDataflowNode::Evaluate(Dataflow::FContext& Context, const FDataflowOutput* Out) const
+void FCollectionFaceSelectionInvertDataflowNode::Evaluate(UE::Dataflow::FContext& Context, const FDataflowOutput* Out) const
 {
 	if (Out->IsA<FDataflowFaceSelection>(&FaceSelection))
 	{
@@ -1105,7 +1105,7 @@ void FCollectionFaceSelectionInvertDataflowNode::Evaluate(Dataflow::FContext& Co
 }
 
 
-void FCollectionVertexSelectionByPercentageDataflowNode::Evaluate(Dataflow::FContext& Context, const FDataflowOutput* Out) const
+void FCollectionVertexSelectionByPercentageDataflowNode::Evaluate(UE::Dataflow::FContext& Context, const FDataflowOutput* Out) const
 {
 	if (Out->IsA<FDataflowVertexSelection>(&VertexSelection))
 	{
@@ -1124,7 +1124,7 @@ void FCollectionVertexSelectionByPercentageDataflowNode::Evaluate(Dataflow::FCon
 }
 
 
-void FCollectionVertexSelectionSetOperationDataflowNode::Evaluate(Dataflow::FContext& Context, const FDataflowOutput* Out) const
+void FCollectionVertexSelectionSetOperationDataflowNode::Evaluate(UE::Dataflow::FContext& Context, const FDataflowOutput* Out) const
 {
 	if (Out->IsA<FDataflowVertexSelection>(&VertexSelection))
 	{
@@ -1271,7 +1271,7 @@ static void CreateSelectionFromAttr(const FManagedArrayCollection& InCollection,
 }
 
 
-void FCollectionSelectionByAttrDataflowNode::Evaluate(Dataflow::FContext& Context, const FDataflowOutput* Out) const
+void FCollectionSelectionByAttrDataflowNode::Evaluate(UE::Dataflow::FContext& Context, const FDataflowOutput* Out) const
 {
 	if (Out->IsA<FDataflowVertexSelection>(&VertexSelection) ||
 		Out->IsA<FDataflowFaceSelection>(&FaceSelection) ||
@@ -1281,7 +1281,7 @@ void FCollectionSelectionByAttrDataflowNode::Evaluate(Dataflow::FContext& Contex
 	{
 		const FManagedArrayCollection& InCollection = GetValue<FManagedArrayCollection>(Context, &Collection);
 		FCollectionAttributeKey InAttributeKey = GetValue<FCollectionAttributeKey>(Context, &AttributeKey);
-		FName GroupName = Dataflow::Private::GetAttributeFromEnumAsName(Group);
+		FName GroupName = UE::Dataflow::Private::GetAttributeFromEnumAsName(Group);
 		FName AttributeName = FName(Attribute);
 		if (IsConnected(&AttributeKey))
 		{
@@ -1305,11 +1305,11 @@ void FCollectionSelectionByAttrDataflowNode::Evaluate(Dataflow::FContext& Contex
 					Operation,
 					NewSelection);
 
-				SetValue(Context, GroupName == Dataflow::Private::GetAttributeFromEnumAsName(ESelectionByAttrGroup::Vertices) ? MoveTemp(NewSelection) : FDataflowSelection(), &VertexSelection);
-				SetValue(Context, GroupName == Dataflow::Private::GetAttributeFromEnumAsName(ESelectionByAttrGroup::Faces) ? MoveTemp(NewSelection) : FDataflowSelection(), &FaceSelection);
-				SetValue(Context, GroupName == Dataflow::Private::GetAttributeFromEnumAsName(ESelectionByAttrGroup::Transform) ? MoveTemp(NewSelection) : FDataflowSelection(), &TransformSelection);
-				SetValue(Context, GroupName == Dataflow::Private::GetAttributeFromEnumAsName(ESelectionByAttrGroup::Geometry) ? MoveTemp(NewSelection) : FDataflowSelection(), &GeometrySelection);
-				SetValue(Context, GroupName == Dataflow::Private::GetAttributeFromEnumAsName(ESelectionByAttrGroup::Material) ? MoveTemp(NewSelection) : FDataflowSelection(), &MaterialSelection);
+				SetValue(Context, GroupName == UE::Dataflow::Private::GetAttributeFromEnumAsName(ESelectionByAttrGroup::Vertices) ? MoveTemp(NewSelection) : FDataflowSelection(), &VertexSelection);
+				SetValue(Context, GroupName == UE::Dataflow::Private::GetAttributeFromEnumAsName(ESelectionByAttrGroup::Faces) ? MoveTemp(NewSelection) : FDataflowSelection(), &FaceSelection);
+				SetValue(Context, GroupName == UE::Dataflow::Private::GetAttributeFromEnumAsName(ESelectionByAttrGroup::Transform) ? MoveTemp(NewSelection) : FDataflowSelection(), &TransformSelection);
+				SetValue(Context, GroupName == UE::Dataflow::Private::GetAttributeFromEnumAsName(ESelectionByAttrGroup::Geometry) ? MoveTemp(NewSelection) : FDataflowSelection(), &GeometrySelection);
+				SetValue(Context, GroupName == UE::Dataflow::Private::GetAttributeFromEnumAsName(ESelectionByAttrGroup::Material) ? MoveTemp(NewSelection) : FDataflowSelection(), &MaterialSelection);
 
 				return;
 			}
diff --git a/Engine/Plugins/Experimental/GeometryCollectionPlugin/Source/GeometryCollectionNodes/Private/Dataflow/GeometryCollectionSkeletalMeshToCollectionNode.cpp b/Engine/Plugins/Experimental/GeometryCollectionPlugin/Source/GeometryCollectionNodes/Private/Dataflow/GeometryCollectionSkeletalMeshToCollectionNode.cpp
--- a/Engine/Plugins/Experimental/GeometryCollectionPlugin/Source/GeometryCollectionNodes/Private/Dataflow/GeometryCollectionSkeletalMeshToCollectionNode.cpp
+++ b/Engine/Plugins/Experimental/GeometryCollectionPlugin/Source/GeometryCollectionNodes/Private/Dataflow/GeometryCollectionSkeletalMeshToCollectionNode.cpp
@@ -17,7 +17,7 @@
 #include "IndexTypes.h"
 #include "GeometryCollection/GeometryCollectionAlgo.h"
 
-void FSkeletalMeshToCollectionDataflowNode::Evaluate(Dataflow::FContext& Context, const FDataflowOutput* Out) const
+void FSkeletalMeshToCollectionDataflowNode::Evaluate(UE::Dataflow::FContext& Context, const FDataflowOutput* Out) const
 {
 	if (Out->IsA<DataType>(&Collection))
 	{
diff --git a/Engine/Plugins/Experimental/GeometryCollectionPlugin/Source/GeometryCollectionNodes/Private/Dataflow/GeometryCollectionSkeletonToCollectionNode.cpp b/Engine/Plugins/Experimental/GeometryCollectionPlugin/Source/GeometryCollectionNodes/Private/Dataflow/GeometryCollectionSkeletonToCollectionNode.cpp
--- a/Engine/Plugins/Experimental/GeometryCollectionPlugin/Source/GeometryCollectionNodes/Private/Dataflow/GeometryCollectionSkeletonToCollectionNode.cpp
+++ b/Engine/Plugins/Experimental/GeometryCollectionPlugin/Source/GeometryCollectionNodes/Private/Dataflow/GeometryCollectionSkeletonToCollectionNode.cpp
@@ -8,7 +8,7 @@
 #include "GeometryCollection/GeometryCollection.h"
 #include "GeometryCollection/GeometryCollectionEngineConversion.h"
 
-void FSkeletonToCollectionDataflowNode::Evaluate(Dataflow::FContext& Context, const FDataflowOutput* Out) const
+void FSkeletonToCollectionDataflowNode::Evaluate(UE::Dataflow::FContext& Context, const FDataflowOutput* Out) const
 {
 	if (Out->IsA<DataType>(&Collection))
 	{
diff --git a/Engine/Plugins/Experimental/GeometryCollectionPlugin/Source/GeometryCollectionNodes/Private/Dataflow/GeometryCollectionStaticMeshToCollectionNode.cpp b/Engine/Plugins/Experimental/GeometryCollectionPlugin/Source/GeometryCollectionNodes/Private/Dataflow/GeometryCollectionStaticMeshToCollectionNode.cpp
--- a/Engine/Plugins/Experimental/GeometryCollectionPlugin/Source/GeometryCollectionNodes/Private/Dataflow/GeometryCollectionStaticMeshToCollectionNode.cpp
+++ b/Engine/Plugins/Experimental/GeometryCollectionPlugin/Source/GeometryCollectionNodes/Private/Dataflow/GeometryCollectionStaticMeshToCollectionNode.cpp
@@ -13,7 +13,7 @@
 // ===========================================================================================================================
 
 
-FStaticMeshToCollectionDataflowNode::FStaticMeshToCollectionDataflowNode(const Dataflow::FNodeParameters& InParam, FGuid InGuid)
+FStaticMeshToCollectionDataflowNode::FStaticMeshToCollectionDataflowNode(const UE::Dataflow::FNodeParameters& InParam, FGuid InGuid)
 	: FDataflowNode(InParam, InGuid)
 {
 	RegisterInputConnection(&StaticMesh);
@@ -24,7 +24,7 @@ FStaticMeshToCollectionDataflowNode::FStaticMeshToCollectionDataflowNode(const D
 	RegisterOutputConnection(&InstancedMeshes);
 }
 
-void FStaticMeshToCollectionDataflowNode::Evaluate(Dataflow::FContext& Context, const FDataflowOutput* Out) const
+void FStaticMeshToCollectionDataflowNode::Evaluate(UE::Dataflow::FContext& Context, const FDataflowOutput* Out) const
 {
 	ensure(Out->IsA(&Collection) || Out->IsA(&Materials) || Out->IsA(&InstancedMeshes));
 
diff --git a/Engine/Plugins/Experimental/GeometryCollectionPlugin/Source/GeometryCollectionNodes/Private/Dataflow/GeometryCollectionTransferVertexAttributeNode.cpp b/Engine/Plugins/Experimental/GeometryCollectionPlugin/Source/GeometryCollectionNodes/Private/Dataflow/GeometryCollectionTransferVertexAttributeNode.cpp
--- a/Engine/Plugins/Experimental/GeometryCollectionPlugin/Source/GeometryCollectionNodes/Private/Dataflow/GeometryCollectionTransferVertexAttributeNode.cpp
+++ b/Engine/Plugins/Experimental/GeometryCollectionPlugin/Source/GeometryCollectionNodes/Private/Dataflow/GeometryCollectionTransferVertexAttributeNode.cpp
@@ -85,7 +85,7 @@ namespace UE::Private {
 
 }
 
-void FGeometryCollectionTransferVertexAttributeNode::Evaluate(Dataflow::FContext& Context, const FDataflowOutput* Out) const
+void FGeometryCollectionTransferVertexAttributeNode::Evaluate(UE::Dataflow::FContext& Context, const FDataflowOutput* Out) const
 {
 	FCollectionAttributeKey Key = GetValue(Context, &AttributeKey, AttributeKey);
 
diff --git a/Engine/Plugins/Experimental/GeometryCollectionPlugin/Source/GeometryCollectionNodes/Private/Dataflow/GeometryCollectionTriangleBoundaryIndicesNode.cpp b/Engine/Plugins/Experimental/GeometryCollectionPlugin/Source/GeometryCollectionNodes/Private/Dataflow/GeometryCollectionTriangleBoundaryIndicesNode.cpp
--- a/Engine/Plugins/Experimental/GeometryCollectionPlugin/Source/GeometryCollectionNodes/Private/Dataflow/GeometryCollectionTriangleBoundaryIndicesNode.cpp
+++ b/Engine/Plugins/Experimental/GeometryCollectionPlugin/Source/GeometryCollectionNodes/Private/Dataflow/GeometryCollectionTriangleBoundaryIndicesNode.cpp
@@ -9,7 +9,7 @@
 
 #include UE_INLINE_GENERATED_CPP_BY_NAME(GeometryCollectionTriangleBoundaryIndicesNode)
 
-void FTriangleBoundaryIndicesNode::Evaluate(Dataflow::FContext& Context, const FDataflowOutput* Out) const
+void FTriangleBoundaryIndicesNode::Evaluate(UE::Dataflow::FContext& Context, const FDataflowOutput* Out) const
 {
 	if (Out->IsA<TArray<int32>>(&BoundaryIndicesOut))
 	{
diff --git a/Engine/Plugins/Experimental/GeometryCollectionPlugin/Source/GeometryCollectionNodes/Private/Dataflow/GeometryCollectionUtilityNodes.cpp b/Engine/Plugins/Experimental/GeometryCollectionPlugin/Source/GeometryCollectionNodes/Private/Dataflow/GeometryCollectionUtilityNodes.cpp
--- a/Engine/Plugins/Experimental/GeometryCollectionPlugin/Source/GeometryCollectionNodes/Private/Dataflow/GeometryCollectionUtilityNodes.cpp
+++ b/Engine/Plugins/Experimental/GeometryCollectionPlugin/Source/GeometryCollectionNodes/Private/Dataflow/GeometryCollectionUtilityNodes.cpp
@@ -17,7 +17,7 @@
 
 //#include UE_INLINE_GENERATED_CPP_BY_NAME(GeometryCollectionUtilityNodes)
 
-namespace Dataflow
+namespace UE::Dataflow
 {
 
 	void GeometryCollectionUtilityNodes()
@@ -41,7 +41,7 @@ namespace Dataflow
 	}
 }
 
-FMakeDataflowConvexDecompositionSettingsNode::FMakeDataflowConvexDecompositionSettingsNode(const Dataflow::FNodeParameters& InParam, FGuid InGuid)
+FMakeDataflowConvexDecompositionSettingsNode::FMakeDataflowConvexDecompositionSettingsNode(const UE::Dataflow::FNodeParameters& InParam, FGuid InGuid)
 	: FDataflowNode(InParam, InGuid)
 {
 	RegisterInputConnection(&MinSizeToDecompose);
@@ -53,7 +53,7 @@ FMakeDataflowConvexDecompositionSettingsNode::FMakeDataflowConvexDecompositionSe
 	RegisterOutputConnection(&DecompositionSettings);
 }
 
-void FMakeDataflowConvexDecompositionSettingsNode::Evaluate(Dataflow::FContext& Context, const FDataflowOutput* Out) const
+void FMakeDataflowConvexDecompositionSettingsNode::Evaluate(UE::Dataflow::FContext& Context, const FDataflowOutput* Out) const
 {
 	if (Out->IsA(&DecompositionSettings))
 	{
@@ -68,7 +68,7 @@ void FMakeDataflowConvexDecompositionSettingsNode::Evaluate(Dataflow::FContext&
 	}
 }
 
-FCreateLeafConvexHullsDataflowNode::FCreateLeafConvexHullsDataflowNode(const Dataflow::FNodeParameters& InParam, FGuid InGuid)
+FCreateLeafConvexHullsDataflowNode::FCreateLeafConvexHullsDataflowNode(const UE::Dataflow::FNodeParameters& InParam, FGuid InGuid)
 	: FDataflowNode(InParam, InGuid)
 {
 	RegisterInputConnection(&Collection);
@@ -78,7 +78,7 @@ FCreateLeafConvexHullsDataflowNode::FCreateLeafConvexHullsDataflowNode(const Dat
 	RegisterOutputConnection(&Collection);
 }
 
-void FCreateLeafConvexHullsDataflowNode::Evaluate(Dataflow::FContext& Context, const FDataflowOutput* Out) const
+void FCreateLeafConvexHullsDataflowNode::Evaluate(UE::Dataflow::FContext& Context, const FDataflowOutput* Out) const
 {
 	if (Out->IsA(&Collection))
 	{
@@ -121,7 +121,7 @@ void FCreateLeafConvexHullsDataflowNode::Evaluate(Dataflow::FContext& Context, c
 	}
 }
 
-FSimplifyConvexHullsDataflowNode::FSimplifyConvexHullsDataflowNode(const Dataflow::FNodeParameters& InParam, FGuid InGuid)
+FSimplifyConvexHullsDataflowNode::FSimplifyConvexHullsDataflowNode(const UE::Dataflow::FNodeParameters& InParam, FGuid InGuid)
 	: FDataflowNode(InParam, InGuid)
 {
 	RegisterInputConnection(&Collection);
@@ -132,7 +132,7 @@ FSimplifyConvexHullsDataflowNode::FSimplifyConvexHullsDataflowNode(const Dataflo
 	RegisterOutputConnection(&Collection);
 }
 
-void FSimplifyConvexHullsDataflowNode::Evaluate(Dataflow::FContext& Context, const FDataflowOutput* Out) const
+void FSimplifyConvexHullsDataflowNode::Evaluate(UE::Dataflow::FContext& Context, const FDataflowOutput* Out) const
 {
 	if (Out->IsA(&Collection) && IsConnected(&Collection))
 	{
@@ -165,7 +165,7 @@ void FSimplifyConvexHullsDataflowNode::Evaluate(Dataflow::FContext& Context, con
 	}
 }
 
-FCreateNonOverlappingConvexHullsDataflowNode::FCreateNonOverlappingConvexHullsDataflowNode(const Dataflow::FNodeParameters& InParam, FGuid InGuid)
+FCreateNonOverlappingConvexHullsDataflowNode::FCreateNonOverlappingConvexHullsDataflowNode(const UE::Dataflow::FNodeParameters& InParam, FGuid InGuid)
 	: FDataflowNode(InParam, InGuid)
 {
 	RegisterInputConnection(&Collection);
@@ -176,7 +176,7 @@ FCreateNonOverlappingConvexHullsDataflowNode::FCreateNonOverlappingConvexHullsDa
 	RegisterOutputConnection(&Collection);
 }
 
-void FCreateNonOverlappingConvexHullsDataflowNode::Evaluate(Dataflow::FContext& Context, const FDataflowOutput* Out) const
+void FCreateNonOverlappingConvexHullsDataflowNode::Evaluate(UE::Dataflow::FContext& Context, const FDataflowOutput* Out) const
 {
 	if (Out->IsA<FManagedArrayCollection>(&Collection) && IsConnected(&Collection))
 	{
@@ -214,7 +214,7 @@ static UE::Geometry::FNegativeSpaceSampleSettings::ESampleMethod ConvertNegative
 	return UE::Geometry::FNegativeSpaceSampleSettings::ESampleMethod::Uniform;
 }
 
-FGenerateClusterConvexHullsFromLeafHullsDataflowNode::FGenerateClusterConvexHullsFromLeafHullsDataflowNode(const Dataflow::FNodeParameters& InParam, FGuid InGuid)
+FGenerateClusterConvexHullsFromLeafHullsDataflowNode::FGenerateClusterConvexHullsFromLeafHullsDataflowNode(const UE::Dataflow::FNodeParameters& InParam, FGuid InGuid)
 	: FDataflowNode(InParam, InGuid)
 {
 	RegisterInputConnection(&Collection);
@@ -231,7 +231,7 @@ FGenerateClusterConvexHullsFromLeafHullsDataflowNode::FGenerateClusterConvexHull
 	RegisterOutputConnection(&SphereCovering);
 }
 
-void FGenerateClusterConvexHullsFromLeafHullsDataflowNode::Evaluate(Dataflow::FContext& Context, const FDataflowOutput* Out) const
+void FGenerateClusterConvexHullsFromLeafHullsDataflowNode::Evaluate(UE::Dataflow::FContext& Context, const FDataflowOutput* Out) const
 {
 	if (Out->IsA(&Collection) || Out->IsA(&SphereCovering))
 	{
@@ -302,7 +302,7 @@ void FGenerateClusterConvexHullsFromLeafHullsDataflowNode::Evaluate(Dataflow::FC
 	}
 }
 
-FGenerateClusterConvexHullsFromChildrenHullsDataflowNode::FGenerateClusterConvexHullsFromChildrenHullsDataflowNode(const Dataflow::FNodeParameters& InParam, FGuid InGuid)
+FGenerateClusterConvexHullsFromChildrenHullsDataflowNode::FGenerateClusterConvexHullsFromChildrenHullsDataflowNode(const UE::Dataflow::FNodeParameters& InParam, FGuid InGuid)
 	: FDataflowNode(InParam, InGuid)
 {
 	RegisterInputConnection(&Collection);
@@ -319,7 +319,7 @@ FGenerateClusterConvexHullsFromChildrenHullsDataflowNode::FGenerateClusterConvex
 	RegisterOutputConnection(&SphereCovering);
 }
 
-void FGenerateClusterConvexHullsFromChildrenHullsDataflowNode::Evaluate(Dataflow::FContext& Context, const FDataflowOutput* Out) const
+void FGenerateClusterConvexHullsFromChildrenHullsDataflowNode::Evaluate(UE::Dataflow::FContext& Context, const FDataflowOutput* Out) const
 {
 	if (Out->IsA(&Collection) || Out->IsA(&SphereCovering))
 	{
@@ -390,7 +390,7 @@ void FGenerateClusterConvexHullsFromChildrenHullsDataflowNode::Evaluate(Dataflow
 	}
 }
 
-FMergeConvexHullsDataflowNode::FMergeConvexHullsDataflowNode(const Dataflow::FNodeParameters& InParam, FGuid InGuid)
+FMergeConvexHullsDataflowNode::FMergeConvexHullsDataflowNode(const UE::Dataflow::FNodeParameters& InParam, FGuid InGuid)
 	: FDataflowNode(InParam, InGuid)
 {
 	RegisterInputConnection(&Collection);
@@ -407,7 +407,7 @@ FMergeConvexHullsDataflowNode::FMergeConvexHullsDataflowNode(const Dataflow::FNo
 	RegisterOutputConnection(&SphereCovering);
 }
 
-void FClearConvexHullsDataflowNode::Evaluate(Dataflow::FContext& Context, const FDataflowOutput* Out) const
+void FClearConvexHullsDataflowNode::Evaluate(UE::Dataflow::FContext& Context, const FDataflowOutput* Out) const
 {
 	if (Out->IsA(&Collection))
 	{
@@ -439,7 +439,7 @@ void FClearConvexHullsDataflowNode::Evaluate(Dataflow::FContext& Context, const
 	}
 }
 
-void FMergeConvexHullsDataflowNode::Evaluate(Dataflow::FContext& Context, const FDataflowOutput* Out) const
+void FMergeConvexHullsDataflowNode::Evaluate(UE::Dataflow::FContext& Context, const FDataflowOutput* Out) const
 {
 	if (Out->IsA(&Collection) || Out->IsA(&SphereCovering))
 	{
@@ -495,14 +495,14 @@ void FMergeConvexHullsDataflowNode::Evaluate(Dataflow::FContext& Context, const
 }
 
 
-FUpdateVolumeAttributesDataflowNode::FUpdateVolumeAttributesDataflowNode(const Dataflow::FNodeParameters& InParam, FGuid InGuid)
+FUpdateVolumeAttributesDataflowNode::FUpdateVolumeAttributesDataflowNode(const UE::Dataflow::FNodeParameters& InParam, FGuid InGuid)
 	: FDataflowNode(InParam, InGuid)
 {
 	RegisterInputConnection(&Collection);
 	RegisterOutputConnection(&Collection);
 }
 
-void FUpdateVolumeAttributesDataflowNode::Evaluate(Dataflow::FContext& Context, const FDataflowOutput* Out) const
+void FUpdateVolumeAttributesDataflowNode::Evaluate(UE::Dataflow::FContext& Context, const FDataflowOutput* Out) const
 {
 	if (Out->IsA(&Collection))
 	{
@@ -516,7 +516,7 @@ void FUpdateVolumeAttributesDataflowNode::Evaluate(Dataflow::FContext& Context,
 }
 
 
-FGetConvexHullVolumeDataflowNode::FGetConvexHullVolumeDataflowNode(const Dataflow::FNodeParameters& InParam, FGuid InGuid)
+FGetConvexHullVolumeDataflowNode::FGetConvexHullVolumeDataflowNode(const UE::Dataflow::FNodeParameters& InParam, FGuid InGuid)
 	: FDataflowNode(InParam, InGuid)
 {
 	RegisterInputConnection(&Collection);
@@ -524,7 +524,7 @@ FGetConvexHullVolumeDataflowNode::FGetConvexHullVolumeDataflowNode(const Dataflo
 	RegisterOutputConnection(&Volume);
 }
 
-void FGetConvexHullVolumeDataflowNode::Evaluate(Dataflow::FContext& Context, const FDataflowOutput* Out) const
+void FGetConvexHullVolumeDataflowNode::Evaluate(UE::Dataflow::FContext& Context, const FDataflowOutput* Out) const
 {
 	if (Out->IsA(&Volume))
 	{
@@ -608,7 +608,7 @@ void FGetConvexHullVolumeDataflowNode::Evaluate(Dataflow::FContext& Context, con
 	}
 }
 
-void FFixTinyGeoDataflowNode::Evaluate(Dataflow::FContext& Context, const FDataflowOutput* Out) const
+void FFixTinyGeoDataflowNode::Evaluate(UE::Dataflow::FContext& Context, const FDataflowOutput* Out) const
 {
 	if (Out->IsA(&Collection))
 	{
@@ -656,7 +656,7 @@ void FFixTinyGeoDataflowNode::Evaluate(Dataflow::FContext& Context, const FDataf
 	}
 }
 
-void FRecomputeNormalsInGeometryCollectionDataflowNode::Evaluate(Dataflow::FContext& Context, const FDataflowOutput* Out) const
+void FRecomputeNormalsInGeometryCollectionDataflowNode::Evaluate(UE::Dataflow::FContext& Context, const FDataflowOutput* Out) const
 {
 	if (Out->IsA(&Collection))
 	{
@@ -698,7 +698,7 @@ void FRecomputeNormalsInGeometryCollectionDataflowNode::Evaluate(Dataflow::FCont
 	}
 }
 
-void FResampleGeometryCollectionDataflowNode::Evaluate(Dataflow::FContext& Context, const FDataflowOutput* Out) const
+void FResampleGeometryCollectionDataflowNode::Evaluate(UE::Dataflow::FContext& Context, const FDataflowOutput* Out) const
 {
 	if (Out->IsA(&Collection))
 	{
@@ -737,7 +737,7 @@ void FResampleGeometryCollectionDataflowNode::Evaluate(Dataflow::FContext& Conte
 	}
 }
 
-void FValidateGeometryCollectionDataflowNode::Evaluate(Dataflow::FContext& Context, const FDataflowOutput* Out) const
+void FValidateGeometryCollectionDataflowNode::Evaluate(UE::Dataflow::FContext& Context, const FDataflowOutput* Out) const
 {
 	if (Out->IsA(&Collection))
 	{
diff --git a/Engine/Plugins/Experimental/GeometryCollectionPlugin/Source/GeometryCollectionNodes/Private/Dataflow/GeometryCollectionVertexScalarToVertexIndicesNode.cpp b/Engine/Plugins/Experimental/GeometryCollectionPlugin/Source/GeometryCollectionNodes/Private/Dataflow/GeometryCollectionVertexScalarToVertexIndicesNode.cpp
--- a/Engine/Plugins/Experimental/GeometryCollectionPlugin/Source/GeometryCollectionNodes/Private/Dataflow/GeometryCollectionVertexScalarToVertexIndicesNode.cpp
+++ b/Engine/Plugins/Experimental/GeometryCollectionPlugin/Source/GeometryCollectionNodes/Private/Dataflow/GeometryCollectionVertexScalarToVertexIndicesNode.cpp
@@ -7,7 +7,7 @@
 #include UE_INLINE_GENERATED_CPP_BY_NAME(GeometryCollectionVertexScalarToVertexIndicesNode)
 #define LOCTEXT_NAMESPACE "FGeometryCollectionVertexScalarToVertexIndicesNode"
 
-FGeometryCollectionVertexScalarToVertexIndicesNode::FGeometryCollectionVertexScalarToVertexIndicesNode(const Dataflow::FNodeParameters& InParam, FGuid InGuid)
+FGeometryCollectionVertexScalarToVertexIndicesNode::FGeometryCollectionVertexScalarToVertexIndicesNode(const UE::Dataflow::FNodeParameters& InParam, FGuid InGuid)
 	: FDataflowNode(InParam, InGuid)
 {
 	RegisterInputConnection(&Collection);
@@ -15,7 +15,7 @@ FGeometryCollectionVertexScalarToVertexIndicesNode::FGeometryCollectionVertexSca
 	RegisterOutputConnection(&VertexIndices);
 }
 
-void FGeometryCollectionVertexScalarToVertexIndicesNode::Evaluate(Dataflow::FContext& Context, const FDataflowOutput* Out) const
+void FGeometryCollectionVertexScalarToVertexIndicesNode::Evaluate(UE::Dataflow::FContext& Context, const FDataflowOutput* Out) const
 {
 	if (Out->IsA< TArray<int32> >(&VertexIndices))
 	{
diff --git a/Engine/Plugins/Experimental/GeometryCollectionPlugin/Source/GeometryCollectionNodes/Private/Dataflow/GeometryCollectionVerticesNodes.cpp b/Engine/Plugins/Experimental/GeometryCollectionPlugin/Source/GeometryCollectionNodes/Private/Dataflow/GeometryCollectionVerticesNodes.cpp
--- a/Engine/Plugins/Experimental/GeometryCollectionPlugin/Source/GeometryCollectionNodes/Private/Dataflow/GeometryCollectionVerticesNodes.cpp
+++ b/Engine/Plugins/Experimental/GeometryCollectionPlugin/Source/GeometryCollectionNodes/Private/Dataflow/GeometryCollectionVerticesNodes.cpp
@@ -7,7 +7,7 @@
 
 #include UE_INLINE_GENERATED_CPP_BY_NAME(GeometryCollectionVerticesNodes)
 
-namespace Dataflow
+namespace UE::Dataflow
 {
 	void GeometryCollectionVerticesNodes()
 	{
@@ -16,7 +16,7 @@ namespace Dataflow
 }
 
 
-void FTransformCollectionAttributeDataflowNode::Evaluate(Dataflow::FContext& Context, const FDataflowOutput* Out) const
+void FTransformCollectionAttributeDataflowNode::Evaluate(UE::Dataflow::FContext& Context, const FDataflowOutput* Out) const
 {
 	if (Out->IsA<DataType>(&Collection))
 	{
diff --git a/Engine/Plugins/Experimental/GeometryCollectionPlugin/Source/GeometryCollectionNodes/Private/Dataflow/SetVertexColorFromFloatArrayNode.cpp b/Engine/Plugins/Experimental/GeometryCollectionPlugin/Source/GeometryCollectionNodes/Private/Dataflow/SetVertexColorFromFloatArrayNode.cpp
--- a/Engine/Plugins/Experimental/GeometryCollectionPlugin/Source/GeometryCollectionNodes/Private/Dataflow/SetVertexColorFromFloatArrayNode.cpp
+++ b/Engine/Plugins/Experimental/GeometryCollectionPlugin/Source/GeometryCollectionNodes/Private/Dataflow/SetVertexColorFromFloatArrayNode.cpp
@@ -5,7 +5,7 @@
 
 #include UE_INLINE_GENERATED_CPP_BY_NAME(SetVertexColorFromFloatArrayNode)
 
-void FSetVertexColorFromFloatArrayDataflowNode::Evaluate(Dataflow::FContext& Context, const FDataflowOutput* Out) const
+void FSetVertexColorFromFloatArrayDataflowNode::Evaluate(UE::Dataflow::FContext& Context, const FDataflowOutput* Out) const
 {
 	if (Out->IsA<FManagedArrayCollection>(&Collection))
 	{
diff --git a/Engine/Plugins/Experimental/GeometryCollectionPlugin/Source/GeometryCollectionNodes/Private/Dataflow/SetVertexColorFromVertexIndicesNode.cpp b/Engine/Plugins/Experimental/GeometryCollectionPlugin/Source/GeometryCollectionNodes/Private/Dataflow/SetVertexColorFromVertexIndicesNode.cpp
--- a/Engine/Plugins/Experimental/GeometryCollectionPlugin/Source/GeometryCollectionNodes/Private/Dataflow/SetVertexColorFromVertexIndicesNode.cpp
+++ b/Engine/Plugins/Experimental/GeometryCollectionPlugin/Source/GeometryCollectionNodes/Private/Dataflow/SetVertexColorFromVertexIndicesNode.cpp
@@ -5,7 +5,7 @@
 
 #include UE_INLINE_GENERATED_CPP_BY_NAME(SetVertexColorFromVertexIndicesNode)
 
-void FSetVertexColorFromVertexIndicesDataflowNode::Evaluate(Dataflow::FContext& Context, const FDataflowOutput* Out) const
+void FSetVertexColorFromVertexIndicesDataflowNode::Evaluate(UE::Dataflow::FContext& Context, const FDataflowOutput* Out) const
 {
 	if (Out->IsA<FManagedArrayCollection>(&Collection))
 	{
diff --git a/Engine/Plugins/Experimental/GeometryCollectionPlugin/Source/GeometryCollectionNodes/Private/Dataflow/SetVertexColorFromVertexSelectionNode.cpp b/Engine/Plugins/Experimental/GeometryCollectionPlugin/Source/GeometryCollectionNodes/Private/Dataflow/SetVertexColorFromVertexSelectionNode.cpp
--- a/Engine/Plugins/Experimental/GeometryCollectionPlugin/Source/GeometryCollectionNodes/Private/Dataflow/SetVertexColorFromVertexSelectionNode.cpp
+++ b/Engine/Plugins/Experimental/GeometryCollectionPlugin/Source/GeometryCollectionNodes/Private/Dataflow/SetVertexColorFromVertexSelectionNode.cpp
@@ -5,7 +5,7 @@
 
 #include UE_INLINE_GENERATED_CPP_BY_NAME(SetVertexColorFromVertexSelectionNode)
 
-void FSetVertexColorFromVertexSelectionDataflowNode::Evaluate(Dataflow::FContext& Context, const FDataflowOutput* Out) const
+void FSetVertexColorFromVertexSelectionDataflowNode::Evaluate(UE::Dataflow::FContext& Context, const FDataflowOutput* Out) const
 {
 	if (Out->IsA<FManagedArrayCollection>(&Collection))
 	{
diff --git a/Engine/Plugins/Experimental/GeometryCollectionPlugin/Source/GeometryCollectionNodes/Public/Dataflow/CreateColorArrayFromFloatArrayNode.h b/Engine/Plugins/Experimental/GeometryCollectionPlugin/Source/GeometryCollectionNodes/Public/Dataflow/CreateColorArrayFromFloatArrayNode.h
--- a/Engine/Plugins/Experimental/GeometryCollectionPlugin/Source/GeometryCollectionNodes/Public/Dataflow/CreateColorArrayFromFloatArrayNode.h
+++ b/Engine/Plugins/Experimental/GeometryCollectionPlugin/Source/GeometryCollectionNodes/Public/Dataflow/CreateColorArrayFromFloatArrayNode.h
@@ -35,13 +35,13 @@ public:
 	FLinearColor Color = FLinearColor(FColor::White);
 
 
-	FCreateColorArrayFromFloatArrayDataflowNode(const Dataflow::FNodeParameters& InParam, FGuid InGuid = FGuid::NewGuid())
+	FCreateColorArrayFromFloatArrayDataflowNode(const UE::Dataflow::FNodeParameters& InParam, FGuid InGuid = FGuid::NewGuid())
 		: FDataflowNode(InParam, InGuid)
 	{
 		RegisterInputConnection(&FloatArray);
 		RegisterOutputConnection(&ColorArray);
 	}
 
-	virtual void Evaluate(Dataflow::FContext& Context, const FDataflowOutput* Out) const override;
+	virtual void Evaluate(UE::Dataflow::FContext& Context, const FDataflowOutput* Out) const override;
 
 };
diff --git a/Engine/Plugins/Experimental/GeometryCollectionPlugin/Source/GeometryCollectionNodes/Public/Dataflow/GeometryCollectionArrayNodes.h b/Engine/Plugins/Experimental/GeometryCollectionPlugin/Source/GeometryCollectionNodes/Public/Dataflow/GeometryCollectionArrayNodes.h
--- a/Engine/Plugins/Experimental/GeometryCollectionPlugin/Source/GeometryCollectionNodes/Public/Dataflow/GeometryCollectionArrayNodes.h
+++ b/Engine/Plugins/Experimental/GeometryCollectionPlugin/Source/GeometryCollectionNodes/Public/Dataflow/GeometryCollectionArrayNodes.h
@@ -37,7 +37,7 @@ public:
 	UPROPERTY(meta = (DataflowOutput))
 	float FloatValue = 0;
 
-	FGetFloatArrayElementDataflowNode(const Dataflow::FNodeParameters& InParam, FGuid InGuid = FGuid::NewGuid())
+	FGetFloatArrayElementDataflowNode(const UE::Dataflow::FNodeParameters& InParam, FGuid InGuid = FGuid::NewGuid())
 		: FDataflowNode(InParam, InGuid)
 	{
 		RegisterInputConnection(&FloatArray);
@@ -45,7 +45,7 @@ public:
 		RegisterOutputConnection(&FloatValue);
 	}
 
-	virtual void Evaluate(Dataflow::FContext& Context, const FDataflowOutput* Out) const override;
+	virtual void Evaluate(UE::Dataflow::FContext& Context, const FDataflowOutput* Out) const override;
 
 };
 
@@ -95,14 +95,14 @@ public:
 	UPROPERTY(meta = (DataflowOutput))
 	TArray<int32> IntArray;
 
-	FFloatArrayToIntArrayDataflowNode(const Dataflow::FNodeParameters& InParam, FGuid InGuid = FGuid::NewGuid())
+	FFloatArrayToIntArrayDataflowNode(const UE::Dataflow::FNodeParameters& InParam, FGuid InGuid = FGuid::NewGuid())
 		: FDataflowNode(InParam, InGuid)
 	{
 		RegisterInputConnection(&FloatArray);
 		RegisterOutputConnection(&IntArray);
 	}
 
-	virtual void Evaluate(Dataflow::FContext& Context, const FDataflowOutput* Out) const override;
+	virtual void Evaluate(UE::Dataflow::FContext& Context, const FDataflowOutput* Out) const override;
 
 };
 
@@ -131,7 +131,7 @@ public:
 	UPROPERTY(meta = (DataflowOutput))
 	FVector Point = FVector(0.0);
 
-	FGetArrayElementDataflowNode(const Dataflow::FNodeParameters& InParam, FGuid InGuid = FGuid::NewGuid())
+	FGetArrayElementDataflowNode(const UE::Dataflow::FNodeParameters& InParam, FGuid InGuid = FGuid::NewGuid())
 		: FDataflowNode(InParam, InGuid)
 	{
 		RegisterInputConnection(&Points);
@@ -139,7 +139,7 @@ public:
 		RegisterOutputConnection(&Point);
 	}
 
-	virtual void Evaluate(Dataflow::FContext& Context, const FDataflowOutput* Out) const override;
+	virtual void Evaluate(UE::Dataflow::FContext& Context, const FDataflowOutput* Out) const override;
 
 };
 
@@ -176,7 +176,7 @@ public:
 	UPROPERTY(meta = (DataflowOutput))
 	int32 NumElements = 0;
 
-	FGetNumArrayElementsDataflowNode(const Dataflow::FNodeParameters& InParam, FGuid InGuid = FGuid::NewGuid())
+	FGetNumArrayElementsDataflowNode(const UE::Dataflow::FNodeParameters& InParam, FGuid InGuid = FGuid::NewGuid())
 		: FDataflowNode(InParam, InGuid)
 	{
 		RegisterInputConnection(&FloatArray);
@@ -186,7 +186,7 @@ public:
 		RegisterOutputConnection(&NumElements);
 	}
 
-	virtual void Evaluate(Dataflow::FContext& Context, const FDataflowOutput* Out) const override;
+	virtual void Evaluate(UE::Dataflow::FContext& Context, const FDataflowOutput* Out) const override;
 
 };
 
@@ -210,14 +210,14 @@ public:
 	UPROPERTY(meta = (DataflowOutput, DisplayName = "FaceSelection"))
 	FDataflowFaceSelection FaceSelection;
 
-	FBoolArrayToFaceSelectionDataflowNode(const Dataflow::FNodeParameters& InParam, FGuid InGuid = FGuid::NewGuid())
+	FBoolArrayToFaceSelectionDataflowNode(const UE::Dataflow::FNodeParameters& InParam, FGuid InGuid = FGuid::NewGuid())
 		: FDataflowNode(InParam, InGuid)
 	{
 		RegisterInputConnection(&BoolAttributeData);
 		RegisterOutputConnection(&FaceSelection);
 	}
 
-	virtual void Evaluate(Dataflow::FContext& Context, const FDataflowOutput* Out) const override;
+	virtual void Evaluate(UE::Dataflow::FContext& Context, const FDataflowOutput* Out) const override;
 
 };
 
@@ -262,14 +262,14 @@ public:
 	UPROPERTY(meta = (DataflowOutput, DisplayName = "VertexSelection"))
 	FDataflowVertexSelection VertexSelection;
 
-	FFloatArrayToVertexSelectionDataflowNode(const Dataflow::FNodeParameters& InParam, FGuid InGuid = FGuid::NewGuid())
+	FFloatArrayToVertexSelectionDataflowNode(const UE::Dataflow::FNodeParameters& InParam, FGuid InGuid = FGuid::NewGuid())
 		: FDataflowNode(InParam, InGuid)
 	{
 		RegisterInputConnection(&FloatArray);
 		RegisterOutputConnection(&VertexSelection);
 	}
 
-	virtual void Evaluate(Dataflow::FContext& Context, const FDataflowOutput* Out) const override;
+	virtual void Evaluate(UE::Dataflow::FContext& Context, const FDataflowOutput* Out) const override;
 
 };
 
@@ -306,7 +306,7 @@ public:
 	UPROPERTY(meta = (DataflowOutput))
 	TArray<float> OutFloatArray;
 
-	FFloatArrayNormalizeDataflowNode(const Dataflow::FNodeParameters& InParam, FGuid InGuid = FGuid::NewGuid())
+	FFloatArrayNormalizeDataflowNode(const UE::Dataflow::FNodeParameters& InParam, FGuid InGuid = FGuid::NewGuid())
 		: FDataflowNode(InParam, InGuid)
 	{
 		RegisterInputConnection(&InFloatArray);
@@ -316,7 +316,7 @@ public:
 		RegisterOutputConnection(&OutFloatArray);
 	}
 
-	virtual void Evaluate(Dataflow::FContext& Context, const FDataflowOutput* Out) const override;
+	virtual void Evaluate(UE::Dataflow::FContext& Context, const FDataflowOutput* Out) const override;
 
 };
 
@@ -349,7 +349,7 @@ public:
 	UPROPERTY(meta = (DataflowOutput))
 	TArray<FVector> OutVectorArray;
 
-	FVectorArrayNormalizeDataflowNode(const Dataflow::FNodeParameters& InParam, FGuid InGuid = FGuid::NewGuid())
+	FVectorArrayNormalizeDataflowNode(const UE::Dataflow::FNodeParameters& InParam, FGuid InGuid = FGuid::NewGuid())
 		: FDataflowNode(InParam, InGuid)
 	{
 		RegisterInputConnection(&InVectorArray);
@@ -358,7 +358,7 @@ public:
 		RegisterOutputConnection(&OutVectorArray);
 	}
 
-	virtual void Evaluate(Dataflow::FContext& Context, const FDataflowOutput* Out) const override;
+	virtual void Evaluate(UE::Dataflow::FContext& Context, const FDataflowOutput* Out) const override;
 
 };
 
@@ -379,7 +379,7 @@ public:
 	UPROPERTY(meta = (DataflowOutput, DisplayName = "OutArray", DataflowPassthrough = "InArray1"));
 	TArray<int32> OutArray;
 
-	FUnionIntArraysDataflowNode(const Dataflow::FNodeParameters& InParam, FGuid InGuid = FGuid::NewGuid())
+	FUnionIntArraysDataflowNode(const UE::Dataflow::FNodeParameters& InParam, FGuid InGuid = FGuid::NewGuid())
 		: FDataflowNode(InParam, InGuid)
 	{
 		RegisterInputConnection(&InArray1);
@@ -387,7 +387,7 @@ public:
 		RegisterOutputConnection(&OutArray, &InArray1);
 	}
 
-	virtual void Evaluate(Dataflow::FContext& Context, const FDataflowOutput* Out) const override;
+	virtual void Evaluate(UE::Dataflow::FContext& Context, const FDataflowOutput* Out) const override;
 };
 
 
@@ -415,7 +415,7 @@ public:
 	UPROPERTY(meta = (DataflowInput, DataflowOutput, DataflowPassthrough = "FloatArray", DataflowIntrinsic))
 	TArray<float> FloatArray;
 
-	FRemoveFloatArrayElementDataflowNode(const Dataflow::FNodeParameters& InParam, FGuid InGuid = FGuid::NewGuid())
+	FRemoveFloatArrayElementDataflowNode(const UE::Dataflow::FNodeParameters& InParam, FGuid InGuid = FGuid::NewGuid())
 		: FDataflowNode(InParam, InGuid)
 	{
 		RegisterInputConnection(&FloatArray);
@@ -423,7 +423,7 @@ public:
 		RegisterOutputConnection(&FloatArray, &FloatArray);
 	}
 
-	virtual void Evaluate(Dataflow::FContext& Context, const FDataflowOutput* Out) const override;
+	virtual void Evaluate(UE::Dataflow::FContext& Context, const FDataflowOutput* Out) const override;
 
 };
 
@@ -474,7 +474,7 @@ public:
 	UPROPERTY(meta = (DataflowOutput))
 	TArray<int32> Indices;
 
-	FFloatArrayComputeStatisticsDataflowNode(const Dataflow::FNodeParameters& InParam, FGuid InGuid = FGuid::NewGuid())
+	FFloatArrayComputeStatisticsDataflowNode(const UE::Dataflow::FNodeParameters& InParam, FGuid InGuid = FGuid::NewGuid())
 		: FDataflowNode(InParam, InGuid)
 	{
 		RegisterInputConnection(&FloatArray);
@@ -483,7 +483,7 @@ public:
 		RegisterOutputConnection(&Indices);
 	}
 
-	virtual void Evaluate(Dataflow::FContext& Context, const FDataflowOutput* Out) const override;
+	virtual void Evaluate(UE::Dataflow::FContext& Context, const FDataflowOutput* Out) const override;
 
 };
 
@@ -515,7 +515,7 @@ public:
 	UPROPERTY(EditAnywhere, Category = "Random", meta = (DataflowInput, UIMin = "0"));
 	int32 RandomSeed = 0;
 
-	FRandomizeFloatArrayDataflowNode(const Dataflow::FNodeParameters& InParam, FGuid InGuid = FGuid::NewGuid())
+	FRandomizeFloatArrayDataflowNode(const UE::Dataflow::FNodeParameters& InParam, FGuid InGuid = FGuid::NewGuid())
 		: FDataflowNode(InParam, InGuid)
 	{
 		RegisterInputConnection(&FloatArray);
@@ -525,10 +525,10 @@ public:
 		RegisterOutputConnection(&FloatArray, &FloatArray);
 	}
 
-	virtual void Evaluate(Dataflow::FContext& Context, const FDataflowOutput* Out) const override;
+	virtual void Evaluate(UE::Dataflow::FContext& Context, const FDataflowOutput* Out) const override;
 };
 
-namespace Dataflow
+namespace UE::Dataflow
 {
 	void GeometryCollectionArrayNodes();
 }
diff --git a/Engine/Plugins/Experimental/GeometryCollectionPlugin/Source/GeometryCollectionNodes/Public/Dataflow/GeometryCollectionAssetNodes.h b/Engine/Plugins/Experimental/GeometryCollectionPlugin/Source/GeometryCollectionNodes/Public/Dataflow/GeometryCollectionAssetNodes.h
--- a/Engine/Plugins/Experimental/GeometryCollectionPlugin/Source/GeometryCollectionNodes/Public/Dataflow/GeometryCollectionAssetNodes.h
+++ b/Engine/Plugins/Experimental/GeometryCollectionPlugin/Source/GeometryCollectionNodes/Public/Dataflow/GeometryCollectionAssetNodes.h
@@ -36,10 +36,10 @@ public:
 	UPROPERTY(meta = (DataflowInput, DataflowOutput, DataflowPassthrough = "InstancedMeshes", DisplayName = "InstancedMeshes"))
 	TArray<FGeometryCollectionAutoInstanceMesh> InstancedMeshes;
 
-	FGeometryCollectionTerminalDataflowNode(const Dataflow::FNodeParameters& InParam, FGuid InGuid = FGuid::NewGuid());
+	FGeometryCollectionTerminalDataflowNode(const UE::Dataflow::FNodeParameters& InParam, FGuid InGuid = FGuid::NewGuid());
 
-	virtual void Evaluate(Dataflow::FContext& Context) const override;
-	virtual void SetAssetValue(TObjectPtr<UObject> Asset, Dataflow::FContext& Context) const override;
+	virtual void Evaluate(UE::Dataflow::FContext& Context) const override;
+	virtual void SetAssetValue(TObjectPtr<UObject> Asset, UE::Dataflow::FContext& Context) const override;
 };
 
 
@@ -54,8 +54,8 @@ struct FGetGeometryCollectionAssetDataflowNode : public FDataflowNode
 	DATAFLOW_NODE_DEFINE_INTERNAL(FGetGeometryCollectionAssetDataflowNode, "GetGeometryCollectionAsset", "GeometryCollection|Asset", "")
 
 public:
-	FGetGeometryCollectionAssetDataflowNode(const Dataflow::FNodeParameters& InParam, FGuid InGuid = FGuid::NewGuid());
-	virtual void Evaluate(Dataflow::FContext& Context, const FDataflowOutput* Out) const override;
+	FGetGeometryCollectionAssetDataflowNode(const UE::Dataflow::FNodeParameters& InParam, FGuid InGuid = FGuid::NewGuid());
+	virtual void Evaluate(UE::Dataflow::FContext& Context, const FDataflowOutput* Out) const override;
 
 	/** Asset this data flow graph instance is assigned to */
 	UPROPERTY(meta = (DataflowOutput, DisplayName = "Asset"))
@@ -74,8 +74,8 @@ struct FGetGeometryCollectionSourcesDataflowNode : public FDataflowNode
 	DATAFLOW_NODE_DEFINE_INTERNAL(FGetGeometryCollectionSourcesDataflowNode, "GetGeometryCollectionSources", "GeometryCollection|Asset", "")
 
 public:
-	FGetGeometryCollectionSourcesDataflowNode(const Dataflow::FNodeParameters& InParam, FGuid InGuid = FGuid::NewGuid());
-	virtual void Evaluate(Dataflow::FContext& Context, const FDataflowOutput* Out) const override;
+	FGetGeometryCollectionSourcesDataflowNode(const UE::Dataflow::FNodeParameters& InParam, FGuid InGuid = FGuid::NewGuid());
+	virtual void Evaluate(UE::Dataflow::FContext& Context, const FDataflowOutput* Out) const override;
 	
 	/** Asset to get geometry sources from */
 	UPROPERTY(meta = (DataflowInput, DisplayName = "Asset"))
@@ -97,8 +97,8 @@ struct FCreateGeometryCollectionFromSourcesDataflowNode : public FDataflowNode
 	DATAFLOW_NODE_DEFINE_INTERNAL(FCreateGeometryCollectionFromSourcesDataflowNode, "CreateGeometryCollectionFromSources", "GeometryCollection|Asset", "")
 
 public:
-	FCreateGeometryCollectionFromSourcesDataflowNode(const Dataflow::FNodeParameters& InParam, FGuid InGuid = FGuid::NewGuid());
-	virtual void Evaluate(Dataflow::FContext& Context, const FDataflowOutput* Out) const override;
+	FCreateGeometryCollectionFromSourcesDataflowNode(const UE::Dataflow::FNodeParameters& InParam, FGuid InGuid = FGuid::NewGuid());
+	virtual void Evaluate(UE::Dataflow::FContext& Context, const FDataflowOutput* Out) const override;
 	
 	/** array of geometry sources */
 	UPROPERTY(meta = (DataflowInput, DisplayName = "Sources"))
@@ -153,8 +153,8 @@ public:
 	UPROPERTY(meta = (DataflowOutput, DisplayName = "InstancedMeshes"))
 	TArray<FGeometryCollectionAutoInstanceMesh> InstancedMeshes;
 
-	FGeometryCollectionToCollectionDataflowNode(const Dataflow::FNodeParameters& InParam, FGuid InGuid = FGuid::NewGuid());
-	virtual void Evaluate(Dataflow::FContext& Context, const FDataflowOutput* Out) const override;
+	FGeometryCollectionToCollectionDataflowNode(const UE::Dataflow::FNodeParameters& InParam, FGuid InGuid = FGuid::NewGuid());
+	virtual void Evaluate(UE::Dataflow::FContext& Context, const FDataflowOutput* Out) const override;
 };
 
 
@@ -193,12 +193,12 @@ public:
  	UPROPERTY(meta = (DataflowOutput, DisplayName = "InstancedMeshes"))
  	TArray<FGeometryCollectionAutoInstanceMesh> InstancedMeshes;
  
-	FBlueprintToCollectionDataflowNode(const Dataflow::FNodeParameters& InParam, FGuid InGuid = FGuid::NewGuid());
- 	virtual void Evaluate(Dataflow::FContext& Context, const FDataflowOutput* Out) const override;
+	FBlueprintToCollectionDataflowNode(const UE::Dataflow::FNodeParameters& InParam, FGuid InGuid = FGuid::NewGuid());
+ 	virtual void Evaluate(UE::Dataflow::FContext& Context, const FDataflowOutput* Out) const override;
  };
 
  
-namespace Dataflow
+namespace UE::Dataflow
 {
 	void GeometryCollectionEngineAssetNodes();
 }
\ No newline at end of file
diff --git a/Engine/Plugins/Experimental/GeometryCollectionPlugin/Source/GeometryCollectionNodes/Public/Dataflow/GeometryCollectionClusteringNodes.h b/Engine/Plugins/Experimental/GeometryCollectionPlugin/Source/GeometryCollectionNodes/Public/Dataflow/GeometryCollectionClusteringNodes.h
--- a/Engine/Plugins/Experimental/GeometryCollectionPlugin/Source/GeometryCollectionNodes/Public/Dataflow/GeometryCollectionClusteringNodes.h
+++ b/Engine/Plugins/Experimental/GeometryCollectionPlugin/Source/GeometryCollectionNodes/Public/Dataflow/GeometryCollectionClusteringNodes.h
@@ -100,7 +100,7 @@ public:
 	UPROPERTY(meta = (DataflowInput, DisplayName = "TransformSelection"))
 	FDataflowTransformSelection TransformSelection;
 
-	FAutoClusterDataflowNode(const Dataflow::FNodeParameters& InParam, FGuid InGuid = FGuid::NewGuid())
+	FAutoClusterDataflowNode(const UE::Dataflow::FNodeParameters& InParam, FGuid InGuid = FGuid::NewGuid())
 		: FDataflowNode(InParam, InGuid)
 	{
 		RegisterInputConnection(&Collection);
@@ -118,7 +118,7 @@ public:
 		RegisterOutputConnection(&Collection, &Collection);
 	}
 
-	virtual void Evaluate(Dataflow::FContext& Context, const FDataflowOutput* Out) const override;
+	virtual void Evaluate(UE::Dataflow::FContext& Context, const FDataflowOutput* Out) const override;
 
 };
 
@@ -144,7 +144,7 @@ public:
 	UPROPERTY(meta = (DataflowInput))
 	FDataflowTransformSelection OptionalTransformSelection;
 
-	FClusterFlattenDataflowNode(const Dataflow::FNodeParameters& InParam, FGuid InGuid = FGuid::NewGuid())
+	FClusterFlattenDataflowNode(const UE::Dataflow::FNodeParameters& InParam, FGuid InGuid = FGuid::NewGuid())
 		: FDataflowNode(InParam, InGuid)
 	{
 		RegisterInputConnection(&Collection);
@@ -153,7 +153,7 @@ public:
 		RegisterOutputConnection(&Collection, &Collection);
 	}
 
-	virtual void Evaluate(Dataflow::FContext& Context, const FDataflowOutput* Out) const override;
+	virtual void Evaluate(UE::Dataflow::FContext& Context, const FDataflowOutput* Out) const override;
 
 };
 
@@ -176,7 +176,7 @@ public:
 	UPROPERTY(meta = (DataflowInput, DisplayName = "TransformSelection"))
 	FDataflowTransformSelection TransformSelection;
 
-	FClusterUnclusterDataflowNode(const Dataflow::FNodeParameters& InParam, FGuid InGuid = FGuid::NewGuid())
+	FClusterUnclusterDataflowNode(const UE::Dataflow::FNodeParameters& InParam, FGuid InGuid = FGuid::NewGuid())
 		: FDataflowNode(InParam, InGuid)
 	{
 		RegisterInputConnection(&Collection);
@@ -184,7 +184,7 @@ public:
 		RegisterOutputConnection(&Collection, &Collection);
 	}
 
-	virtual void Evaluate(Dataflow::FContext& Context, const FDataflowOutput* Out) const override;
+	virtual void Evaluate(UE::Dataflow::FContext& Context, const FDataflowOutput* Out) const override;
 
 };
 
@@ -207,7 +207,7 @@ public:
 	UPROPERTY(meta = (DataflowInput, DisplayName = "TransformSelection"))
 	FDataflowTransformSelection TransformSelection;
 
-	FClusterDataflowNode(const Dataflow::FNodeParameters& InParam, FGuid InGuid = FGuid::NewGuid())
+	FClusterDataflowNode(const UE::Dataflow::FNodeParameters& InParam, FGuid InGuid = FGuid::NewGuid())
 		: FDataflowNode(InParam, InGuid)
 	{
 		RegisterInputConnection(&Collection);
@@ -215,7 +215,7 @@ public:
 		RegisterOutputConnection(&Collection, &Collection);
 	}
 
-	virtual void Evaluate(Dataflow::FContext& Context, const FDataflowOutput* Out) const override;
+	virtual void Evaluate(UE::Dataflow::FContext& Context, const FDataflowOutput* Out) const override;
 
 };
 
@@ -263,7 +263,7 @@ public:
 	UPROPERTY(EditAnywhere, Category = Options, meta = (DataflowInput))
 	bool bOnlySameParent = true;
 
-	FClusterMergeToNeighborsDataflowNode(const Dataflow::FNodeParameters& InParam, FGuid InGuid = FGuid::NewGuid())
+	FClusterMergeToNeighborsDataflowNode(const UE::Dataflow::FNodeParameters& InParam, FGuid InGuid = FGuid::NewGuid())
 		: FDataflowNode(InParam, InGuid)
 	{
 		RegisterInputConnection(&Collection);
@@ -275,7 +275,7 @@ public:
 		RegisterOutputConnection(&Collection, &Collection);
 	}
 
-	virtual void Evaluate(Dataflow::FContext& Context, const FDataflowOutput* Out) const override;
+	virtual void Evaluate(UE::Dataflow::FContext& Context, const FDataflowOutput* Out) const override;
 
 };
 
@@ -299,7 +299,7 @@ public:
 	UPROPERTY(meta = (DataflowInput, DisplayName = "TransformSelection"))
 	FDataflowTransformSelection TransformSelection;
 
-	FClusterMergeDataflowNode(const Dataflow::FNodeParameters& InParam, FGuid InGuid = FGuid::NewGuid())
+	FClusterMergeDataflowNode(const UE::Dataflow::FNodeParameters& InParam, FGuid InGuid = FGuid::NewGuid())
 		: FDataflowNode(InParam, InGuid)
 	{
 		RegisterInputConnection(&Collection);
@@ -307,7 +307,7 @@ public:
 		RegisterOutputConnection(&Collection, &Collection);
 	}
 
-	virtual void Evaluate(Dataflow::FContext& Context, const FDataflowOutput* Out) const override;
+	virtual void Evaluate(UE::Dataflow::FContext& Context, const FDataflowOutput* Out) const override;
 
 };
 
@@ -327,14 +327,14 @@ public:
 	UPROPERTY(meta = (DataflowInput, DataflowOutput, DataflowPassthrough = "Collection", DataflowIntrinsic))
 	FManagedArrayCollection Collection;
 
-	FClusterIsolatedRootsDataflowNode(const Dataflow::FNodeParameters& InParam, FGuid InGuid = FGuid::NewGuid())
+	FClusterIsolatedRootsDataflowNode(const UE::Dataflow::FNodeParameters& InParam, FGuid InGuid = FGuid::NewGuid())
 		: FDataflowNode(InParam, InGuid)
 	{
 		RegisterInputConnection(&Collection);
 		RegisterOutputConnection(&Collection, &Collection);
 	}
 
-	virtual void Evaluate(Dataflow::FContext& Context, const FDataflowOutput* Out) const override;
+	virtual void Evaluate(UE::Dataflow::FContext& Context, const FDataflowOutput* Out) const override;
 
 };
 
@@ -362,7 +362,7 @@ public:
 	UPROPERTY(EditAnywhere, Category = "Cluster Magnet", meta = (ClampMin = "1", DataflowInput, DisplayName = "Iterations"))
 	int32 Iterations = 1;
 
-	FClusterMagnetDataflowNode(const Dataflow::FNodeParameters& InParam, FGuid InGuid = FGuid::NewGuid())
+	FClusterMagnetDataflowNode(const UE::Dataflow::FNodeParameters& InParam, FGuid InGuid = FGuid::NewGuid())
 		: FDataflowNode(InParam, InGuid)
 	{
 		RegisterInputConnection(&Collection);
@@ -371,12 +371,12 @@ public:
 		RegisterOutputConnection(&Collection, &Collection);
 	}
 
-	virtual void Evaluate(Dataflow::FContext& Context, const FDataflowOutput* Out) const override;
+	virtual void Evaluate(UE::Dataflow::FContext& Context, const FDataflowOutput* Out) const override;
 
 };
 
 
-namespace Dataflow
+namespace UE::Dataflow
 {
 	void GeometryCollectionClusteringNodes();
 }
diff --git a/Engine/Plugins/Experimental/GeometryCollectionPlugin/Source/GeometryCollectionNodes/Public/Dataflow/GeometryCollectionConversionNodes.h b/Engine/Plugins/Experimental/GeometryCollectionPlugin/Source/GeometryCollectionNodes/Public/Dataflow/GeometryCollectionConversionNodes.h
--- a/Engine/Plugins/Experimental/GeometryCollectionPlugin/Source/GeometryCollectionNodes/Public/Dataflow/GeometryCollectionConversionNodes.h
+++ b/Engine/Plugins/Experimental/GeometryCollectionPlugin/Source/GeometryCollectionNodes/Public/Dataflow/GeometryCollectionConversionNodes.h
@@ -25,14 +25,14 @@ public:
 	UPROPERTY(meta = (DataflowOutput))
 	FString String = FString("");
 
-	FVectorToStringDataflowNode(const Dataflow::FNodeParameters& InParam, FGuid InGuid = FGuid::NewGuid())
+	FVectorToStringDataflowNode(const UE::Dataflow::FNodeParameters& InParam, FGuid InGuid = FGuid::NewGuid())
 		: FDataflowNode(InParam, InGuid)
 	{
 		RegisterInputConnection(&Vector);
 		RegisterOutputConnection(&String);
 	}
 
-	virtual void Evaluate(Dataflow::FContext& Context, const FDataflowOutput* Out) const override;
+	virtual void Evaluate(UE::Dataflow::FContext& Context, const FDataflowOutput* Out) const override;
 
 };
 
@@ -54,14 +54,14 @@ public:
 	UPROPERTY(meta = (DataflowOutput))
 	FString String = FString("");
 
-	FFloatToStringDataflowNode(const Dataflow::FNodeParameters& InParam, FGuid InGuid = FGuid::NewGuid())
+	FFloatToStringDataflowNode(const UE::Dataflow::FNodeParameters& InParam, FGuid InGuid = FGuid::NewGuid())
 		: FDataflowNode(InParam, InGuid)
 	{
 		RegisterInputConnection(&Float);
 		RegisterOutputConnection(&String);
 	}
 
-	virtual void Evaluate(Dataflow::FContext& Context, const FDataflowOutput* Out) const override;
+	virtual void Evaluate(UE::Dataflow::FContext& Context, const FDataflowOutput* Out) const override;
 
 };
 
@@ -83,14 +83,14 @@ public:
 	UPROPERTY(meta = (DataflowOutput))
 	FString String = FString("");
 
-	FIntToStringDataflowNode(const Dataflow::FNodeParameters& InParam, FGuid InGuid = FGuid::NewGuid())
+	FIntToStringDataflowNode(const UE::Dataflow::FNodeParameters& InParam, FGuid InGuid = FGuid::NewGuid())
 		: FDataflowNode(InParam, InGuid)
 	{
 		RegisterInputConnection(&Int);
 		RegisterOutputConnection(&String);
 	}
 
-	virtual void Evaluate(Dataflow::FContext& Context, const FDataflowOutput* Out) const override;
+	virtual void Evaluate(UE::Dataflow::FContext& Context, const FDataflowOutput* Out) const override;
 
 };
 
@@ -112,14 +112,14 @@ public:
 	UPROPERTY(meta = (DataflowOutput))
 	FString String = FString("");
 
-	FBoolToStringDataflowNode(const Dataflow::FNodeParameters& InParam, FGuid InGuid = FGuid::NewGuid())
+	FBoolToStringDataflowNode(const UE::Dataflow::FNodeParameters& InParam, FGuid InGuid = FGuid::NewGuid())
 		: FDataflowNode(InParam, InGuid)
 	{
 		RegisterInputConnection(&Bool);
 		RegisterOutputConnection(&String);
 	}
 
-	virtual void Evaluate(Dataflow::FContext& Context, const FDataflowOutput* Out) const override;
+	virtual void Evaluate(UE::Dataflow::FContext& Context, const FDataflowOutput* Out) const override;
 
 };
 
@@ -141,14 +141,14 @@ public:
 	UPROPERTY(meta = (DataflowOutput))
 	float Float = 0.f;
 
-	FIntToFloatDataflowNode(const Dataflow::FNodeParameters& InParam, FGuid InGuid = FGuid::NewGuid())
+	FIntToFloatDataflowNode(const UE::Dataflow::FNodeParameters& InParam, FGuid InGuid = FGuid::NewGuid())
 		: FDataflowNode(InParam, InGuid)
 	{
 		RegisterInputConnection(&Int);
 		RegisterOutputConnection(&Float);
 	}
 
-	virtual void Evaluate(Dataflow::FContext& Context, const FDataflowOutput* Out) const override;
+	virtual void Evaluate(UE::Dataflow::FContext& Context, const FDataflowOutput* Out) const override;
 };
 
 /**
@@ -169,14 +169,14 @@ public:
 	UPROPERTY(meta = (DataflowOutput))
 	double Double = 0.0;
 
-	FIntToDoubleDataflowNode(const Dataflow::FNodeParameters& InParam, FGuid InGuid = FGuid::NewGuid())
+	FIntToDoubleDataflowNode(const UE::Dataflow::FNodeParameters& InParam, FGuid InGuid = FGuid::NewGuid())
 		: FDataflowNode(InParam, InGuid)
 	{
 		RegisterInputConnection(&Int);
 		RegisterOutputConnection(&Double);
 	}
 
-	virtual void Evaluate(Dataflow::FContext& Context, const FDataflowOutput* Out) const override;
+	virtual void Evaluate(UE::Dataflow::FContext& Context, const FDataflowOutput* Out) const override;
 };
 
 /**
@@ -197,14 +197,14 @@ public:
 	UPROPERTY(meta = (DataflowOutput))
 	double Double = 0.0;
 
-	FFloatToDoubleDataflowNode(const Dataflow::FNodeParameters& InParam, FGuid InGuid = FGuid::NewGuid())
+	FFloatToDoubleDataflowNode(const UE::Dataflow::FNodeParameters& InParam, FGuid InGuid = FGuid::NewGuid())
 		: FDataflowNode(InParam, InGuid)
 	{
 		RegisterInputConnection(&Float);
 		RegisterOutputConnection(&Double);
 	}
 
-	virtual void Evaluate(Dataflow::FContext& Context, const FDataflowOutput* Out) const override;
+	virtual void Evaluate(UE::Dataflow::FContext& Context, const FDataflowOutput* Out) const override;
 };
 
 
@@ -244,14 +244,14 @@ public:
 	UPROPERTY(meta = (DataflowOutput))
 	int32 Int = 0;
 
-	FFloatToIntDataflowNode(const Dataflow::FNodeParameters& InParam, FGuid InGuid = FGuid::NewGuid())
+	FFloatToIntDataflowNode(const UE::Dataflow::FNodeParameters& InParam, FGuid InGuid = FGuid::NewGuid())
 		: FDataflowNode(InParam, InGuid)
 	{
 		RegisterInputConnection(&Float);
 		RegisterOutputConnection(&Int);
 	}
 
-	virtual void Evaluate(Dataflow::FContext& Context, const FDataflowOutput* Out) const override;
+	virtual void Evaluate(UE::Dataflow::FContext& Context, const FDataflowOutput* Out) const override;
 
 };
 
@@ -273,14 +273,14 @@ public:
 	UPROPERTY(meta = (DataflowOutput))
 	bool Bool = false;
 
-	FIntToBoolDataflowNode(const Dataflow::FNodeParameters& InParam, FGuid InGuid = FGuid::NewGuid())
+	FIntToBoolDataflowNode(const UE::Dataflow::FNodeParameters& InParam, FGuid InGuid = FGuid::NewGuid())
 		: FDataflowNode(InParam, InGuid)
 	{
 		RegisterInputConnection(&Int);
 		RegisterOutputConnection(&Bool);
 	}
 
-	virtual void Evaluate(Dataflow::FContext& Context, const FDataflowOutput* Out) const override;
+	virtual void Evaluate(UE::Dataflow::FContext& Context, const FDataflowOutput* Out) const override;
 
 };
 
@@ -302,18 +302,18 @@ public:
 	UPROPERTY(meta = (DataflowOutput))
 	int32 Int = 0;
 
-	FBoolToIntDataflowNode(const Dataflow::FNodeParameters& InParam, FGuid InGuid = FGuid::NewGuid())
+	FBoolToIntDataflowNode(const UE::Dataflow::FNodeParameters& InParam, FGuid InGuid = FGuid::NewGuid())
 		: FDataflowNode(InParam, InGuid)
 	{
 		RegisterInputConnection(&Bool);
 		RegisterOutputConnection(&Int);
 	}
 
-	virtual void Evaluate(Dataflow::FContext& Context, const FDataflowOutput* Out) const override;
+	virtual void Evaluate(UE::Dataflow::FContext& Context, const FDataflowOutput* Out) const override;
 
 };
 
-namespace Dataflow
+namespace UE::Dataflow
 {
 	void GeometryCollectionConversionNodes();
 }
diff --git a/Engine/Plugins/Experimental/GeometryCollectionPlugin/Source/GeometryCollectionNodes/Public/Dataflow/GeometryCollectionDebugNodes.h b/Engine/Plugins/Experimental/GeometryCollectionPlugin/Source/GeometryCollectionNodes/Public/Dataflow/GeometryCollectionDebugNodes.h
--- a/Engine/Plugins/Experimental/GeometryCollectionPlugin/Source/GeometryCollectionNodes/Public/Dataflow/GeometryCollectionDebugNodes.h
+++ b/Engine/Plugins/Experimental/GeometryCollectionPlugin/Source/GeometryCollectionNodes/Public/Dataflow/GeometryCollectionDebugNodes.h
@@ -32,7 +32,7 @@ public:
 	UPROPERTY(meta = (DataflowOutput))
 	TObjectPtr<UDynamicMesh> Mesh;
 
-	FConvexHullToMeshDataflowNode(const Dataflow::FNodeParameters& InParam, FGuid InGuid = FGuid::NewGuid())
+	FConvexHullToMeshDataflowNode(const UE::Dataflow::FNodeParameters& InParam, FGuid InGuid = FGuid::NewGuid())
 		: FDataflowNode(InParam, InGuid)
 	{
 		RegisterInputConnection(&Collection);
@@ -40,7 +40,7 @@ public:
 		RegisterOutputConnection(&Mesh);
 	}
 
-	virtual void Evaluate(Dataflow::FContext& Context, const FDataflowOutput* Out) const override;
+	virtual void Evaluate(UE::Dataflow::FContext& Context, const FDataflowOutput* Out) const override;
 };
 
 // Convert a sphere covering, as generated by the 'protect negative space' option on the "Generate Cluster Convex Hull" nodes, to a dynamic mesh
@@ -62,14 +62,14 @@ public:
 	UPROPERTY(meta = (DataflowOutput))
 	TObjectPtr<UDynamicMesh> Mesh;
 
-	FSphereCoveringToMeshDataflowNode(const Dataflow::FNodeParameters& InParam, FGuid InGuid = FGuid::NewGuid())
+	FSphereCoveringToMeshDataflowNode(const UE::Dataflow::FNodeParameters& InParam, FGuid InGuid = FGuid::NewGuid())
 		: FDataflowNode(InParam, InGuid)
 	{
 		RegisterInputConnection(&SphereCovering);
 		RegisterOutputConnection(&Mesh);
 	}
 
-	virtual void Evaluate(Dataflow::FContext& Context, const FDataflowOutput* Out) const override;
+	virtual void Evaluate(UE::Dataflow::FContext& Context, const FDataflowOutput* Out) const override;
 };
 
 // Report the number of spheres in a sphere covering
@@ -88,14 +88,14 @@ public:
 	UPROPERTY(meta = (DataflowOutput))
 	int32 NumSpheres = 0;
 
-	FSphereCoveringCountSpheresNode(const Dataflow::FNodeParameters& InParam, FGuid InGuid = FGuid::NewGuid())
+	FSphereCoveringCountSpheresNode(const UE::Dataflow::FNodeParameters& InParam, FGuid InGuid = FGuid::NewGuid())
 		: FDataflowNode(InParam, InGuid)
 	{
 		RegisterInputConnection(&SphereCovering);
 		RegisterOutputConnection(&NumSpheres);
 	}
 
-	virtual void Evaluate(Dataflow::FContext& Context, const FDataflowOutput* Out) const override
+	virtual void Evaluate(UE::Dataflow::FContext& Context, const FDataflowOutput* Out) const override
 	{
 		if (Out->IsA(&NumSpheres))
 		{
@@ -124,7 +124,7 @@ public:
 	UPROPERTY(meta = (DataflowOutput))
 	FString StringOBJ;
 
-	FMeshToOBJStringDebugDataflowNode(const Dataflow::FNodeParameters& InParam, FGuid InGuid = FGuid::NewGuid())
+	FMeshToOBJStringDebugDataflowNode(const UE::Dataflow::FNodeParameters& InParam, FGuid InGuid = FGuid::NewGuid())
 		: FDataflowNode(InParam, InGuid)
 	{
 		RegisterInputConnection(&Mesh);
@@ -132,7 +132,7 @@ public:
 		RegisterOutputConnection(&StringOBJ);
 	}
 
-	virtual void Evaluate(Dataflow::FContext& Context, const FDataflowOutput* Out) const override;
+	virtual void Evaluate(UE::Dataflow::FContext& Context, const FDataflowOutput* Out) const override;
 };
 
 // Write a string to a file
@@ -152,17 +152,17 @@ public:
 	UPROPERTY(EditAnywhere, Category = Options, meta = (DataflowInput))
 	FString FileContents;
 
-	FWriteStringToFile(const Dataflow::FNodeParameters& InParam, FGuid InGuid = FGuid::NewGuid())
+	FWriteStringToFile(const UE::Dataflow::FNodeParameters& InParam, FGuid InGuid = FGuid::NewGuid())
 		: FDataflowNode(InParam, InGuid)
 	{
 		RegisterInputConnection(&FileContents);
 		RegisterInputConnection(&FilePath);
 	}
 
-	virtual void Evaluate(Dataflow::FContext& Context, const FDataflowOutput* Out) const override;
+	virtual void Evaluate(UE::Dataflow::FContext& Context, const FDataflowOutput* Out) const override;
 };
 
-namespace Dataflow
+namespace UE::Dataflow
 {
 	void GeometryCollectionDebugNodes();
 }
diff --git a/Engine/Plugins/Experimental/GeometryCollectionPlugin/Source/GeometryCollectionNodes/Public/Dataflow/GeometryCollectionEditNodes.h b/Engine/Plugins/Experimental/GeometryCollectionPlugin/Source/GeometryCollectionNodes/Public/Dataflow/GeometryCollectionEditNodes.h
--- a/Engine/Plugins/Experimental/GeometryCollectionPlugin/Source/GeometryCollectionNodes/Public/Dataflow/GeometryCollectionEditNodes.h
+++ b/Engine/Plugins/Experimental/GeometryCollectionPlugin/Source/GeometryCollectionNodes/Public/Dataflow/GeometryCollectionEditNodes.h
@@ -32,7 +32,7 @@ public:
 	UPROPERTY(meta = (DataflowInput, DisplayName = "TransformSelection", DataflowIntrinsic))
 	FDataflowTransformSelection TransformSelection;
 
-	FPruneInCollectionDataflowNode(const Dataflow::FNodeParameters& InParam, FGuid InGuid = FGuid::NewGuid())
+	FPruneInCollectionDataflowNode(const UE::Dataflow::FNodeParameters& InParam, FGuid InGuid = FGuid::NewGuid())
 		: FDataflowNode(InParam, InGuid)
 	{
 		RegisterInputConnection(&Collection);
@@ -40,7 +40,7 @@ public:
 		RegisterOutputConnection(&Collection, &Collection);
 	}
 
-	virtual void Evaluate(Dataflow::FContext& Context, const FDataflowOutput* Out) const override;
+	virtual void Evaluate(UE::Dataflow::FContext& Context, const FDataflowOutput* Out) const override;
 
 };
 
@@ -83,7 +83,7 @@ public:
 	UPROPERTY(meta = (DataflowInput, DisplayName = "TransformSelection", DataflowIntrinsic))
 	FDataflowFaceSelection FaceSelection;
 
-	FSetVisibilityInCollectionDataflowNode(const Dataflow::FNodeParameters& InParam, FGuid InGuid = FGuid::NewGuid())
+	FSetVisibilityInCollectionDataflowNode(const UE::Dataflow::FNodeParameters& InParam, FGuid InGuid = FGuid::NewGuid())
 		: FDataflowNode(InParam, InGuid)
 	{
 		RegisterInputConnection(&Collection);
@@ -92,7 +92,7 @@ public:
 		RegisterOutputConnection(&Collection, &Collection);
 	}
 
-	virtual void Evaluate(Dataflow::FContext& Context, const FDataflowOutput* Out) const override;
+	virtual void Evaluate(UE::Dataflow::FContext& Context, const FDataflowOutput* Out) const override;
 
 };
 
@@ -117,7 +117,7 @@ public:
 	UPROPERTY(meta = (DataflowInput, DisplayName = "TransformSelection", DataflowIntrinsic))
 	FDataflowTransformSelection TransformSelection;
 
-	FMergeInCollectionDataflowNode(const Dataflow::FNodeParameters& InParam, FGuid InGuid = FGuid::NewGuid())
+	FMergeInCollectionDataflowNode(const UE::Dataflow::FNodeParameters& InParam, FGuid InGuid = FGuid::NewGuid())
 		: FDataflowNode(InParam, InGuid)
 	{
 		RegisterInputConnection(&Collection);
@@ -125,12 +125,12 @@ public:
 		RegisterOutputConnection(&Collection, &Collection);
 	}
 
-	virtual void Evaluate(Dataflow::FContext& Context, const FDataflowOutput* Out) const override;
+	virtual void Evaluate(UE::Dataflow::FContext& Context, const FDataflowOutput* Out) const override;
 
 };
 
 
-namespace Dataflow
+namespace UE::Dataflow
 {
 	void GeometryCollectionEditNodes();
 }
diff --git a/Engine/Plugins/Experimental/GeometryCollectionPlugin/Source/GeometryCollectionNodes/Public/Dataflow/GeometryCollectionFieldNodes.h b/Engine/Plugins/Experimental/GeometryCollectionPlugin/Source/GeometryCollectionNodes/Public/Dataflow/GeometryCollectionFieldNodes.h
--- a/Engine/Plugins/Experimental/GeometryCollectionPlugin/Source/GeometryCollectionNodes/Public/Dataflow/GeometryCollectionFieldNodes.h
+++ b/Engine/Plugins/Experimental/GeometryCollectionPlugin/Source/GeometryCollectionNodes/Public/Dataflow/GeometryCollectionFieldNodes.h
@@ -93,7 +93,7 @@ public:
 	UPROPERTY(meta = (DataflowOutput))
 	int32 NumSamplePositions = 0;
 
-	FRadialFalloffFieldDataflowNode(const Dataflow::FNodeParameters& InParam, FGuid InGuid = FGuid::NewGuid())
+	FRadialFalloffFieldDataflowNode(const UE::Dataflow::FNodeParameters& InParam, FGuid InGuid = FGuid::NewGuid())
 		: FDataflowNode(InParam, InGuid)
 	{
 		RegisterInputConnection(&SamplePositions);
@@ -110,7 +110,7 @@ public:
 		RegisterOutputConnection(&NumSamplePositions);
 	}
 
-	virtual void Evaluate(Dataflow::FContext& Context, const FDataflowOutput* Out) const override;
+	virtual void Evaluate(UE::Dataflow::FContext& Context, const FDataflowOutput* Out) const override;
 
 };
 
@@ -178,7 +178,7 @@ public:
 	UPROPERTY(meta = (DataflowOutput))
 	int32 NumSamplePositions = 0;
 
-	FBoxFalloffFieldDataflowNode(const Dataflow::FNodeParameters& InParam, FGuid InGuid = FGuid::NewGuid())
+	FBoxFalloffFieldDataflowNode(const UE::Dataflow::FNodeParameters& InParam, FGuid InGuid = FGuid::NewGuid())
 		: FDataflowNode(InParam, InGuid)
 	{
 		RegisterInputConnection(&SamplePositions);
@@ -195,7 +195,7 @@ public:
 		RegisterOutputConnection(&NumSamplePositions);
 	}
 
-	virtual void Evaluate(Dataflow::FContext& Context, const FDataflowOutput* Out) const override;
+	virtual void Evaluate(UE::Dataflow::FContext& Context, const FDataflowOutput* Out) const override;
 
 };
 
@@ -272,7 +272,7 @@ public:
 	UPROPERTY(meta = (DataflowOutput))
 	int32 NumSamplePositions = 0;
 
-	FPlaneFalloffFieldDataflowNode(const Dataflow::FNodeParameters& InParam, FGuid InGuid = FGuid::NewGuid())
+	FPlaneFalloffFieldDataflowNode(const UE::Dataflow::FNodeParameters& InParam, FGuid InGuid = FGuid::NewGuid())
 		: FDataflowNode(InParam, InGuid)
 	{
 		RegisterInputConnection(&SamplePositions);
@@ -291,7 +291,7 @@ public:
 		RegisterOutputConnection(&NumSamplePositions);
 	}
 
-	virtual void Evaluate(Dataflow::FContext& Context, const FDataflowOutput* Out) const override;
+	virtual void Evaluate(UE::Dataflow::FContext& Context, const FDataflowOutput* Out) const override;
 
 };
 
@@ -364,7 +364,7 @@ public:
 	UPROPERTY(meta = (DataflowOutput))
 	int32 NumSamplePositions = 0;
 
-	FRadialIntMaskFieldDataflowNode(const Dataflow::FNodeParameters& InParam, FGuid InGuid = FGuid::NewGuid())
+	FRadialIntMaskFieldDataflowNode(const UE::Dataflow::FNodeParameters& InParam, FGuid InGuid = FGuid::NewGuid())
 		: FDataflowNode(InParam, InGuid)
 	{
 		RegisterInputConnection(&SamplePositions);
@@ -378,7 +378,7 @@ public:
 		RegisterOutputConnection(&NumSamplePositions);
 	}
 
-	virtual void Evaluate(Dataflow::FContext& Context, const FDataflowOutput* Out) const override;
+	virtual void Evaluate(UE::Dataflow::FContext& Context, const FDataflowOutput* Out) const override;
 
 };
 
@@ -419,7 +419,7 @@ public:
 	UPROPERTY(meta = (DataflowOutput))
 	int32 NumSamplePositions = 0;
 
-	FUniformScalarFieldDataflowNode(const Dataflow::FNodeParameters& InParam, FGuid InGuid = FGuid::NewGuid())
+	FUniformScalarFieldDataflowNode(const UE::Dataflow::FNodeParameters& InParam, FGuid InGuid = FGuid::NewGuid())
 		: FDataflowNode(InParam, InGuid)
 	{
 		RegisterInputConnection(&SamplePositions);
@@ -430,7 +430,7 @@ public:
 		RegisterOutputConnection(&NumSamplePositions);
 	}
 
-	virtual void Evaluate(Dataflow::FContext& Context, const FDataflowOutput* Out) const override;
+	virtual void Evaluate(UE::Dataflow::FContext& Context, const FDataflowOutput* Out) const override;
 
 };
 
@@ -475,7 +475,7 @@ public:
 	UPROPERTY(meta = (DataflowOutput))
 	int32 NumSamplePositions = 0;
 
-	FUniformVectorFieldDataflowNode(const Dataflow::FNodeParameters& InParam, FGuid InGuid = FGuid::NewGuid())
+	FUniformVectorFieldDataflowNode(const UE::Dataflow::FNodeParameters& InParam, FGuid InGuid = FGuid::NewGuid())
 		: FDataflowNode(InParam, InGuid)
 	{
 		RegisterInputConnection(&SamplePositions);
@@ -487,7 +487,7 @@ public:
 		RegisterOutputConnection(&NumSamplePositions);
 	}
 
-	virtual void Evaluate(Dataflow::FContext& Context, const FDataflowOutput* Out) const override;
+	virtual void Evaluate(UE::Dataflow::FContext& Context, const FDataflowOutput* Out) const override;
 
 };
 
@@ -532,7 +532,7 @@ public:
 	UPROPERTY(meta = (DataflowOutput))
 	int32 NumSamplePositions = 0;
 
-	FRadialVectorFieldDataflowNode(const Dataflow::FNodeParameters& InParam, FGuid InGuid = FGuid::NewGuid())
+	FRadialVectorFieldDataflowNode(const UE::Dataflow::FNodeParameters& InParam, FGuid InGuid = FGuid::NewGuid())
 		: FDataflowNode(InParam, InGuid)
 	{
 		RegisterInputConnection(&SamplePositions);
@@ -544,7 +544,7 @@ public:
 		RegisterOutputConnection(&NumSamplePositions);
 	}
 
-	virtual void Evaluate(Dataflow::FContext& Context, const FDataflowOutput* Out) const override;
+	virtual void Evaluate(UE::Dataflow::FContext& Context, const FDataflowOutput* Out) const override;
 
 };
 
@@ -585,7 +585,7 @@ public:
 	UPROPERTY(meta = (DataflowOutput))
 	int32 NumSamplePositions = 0;
 
-	FRandomVectorFieldDataflowNode(const Dataflow::FNodeParameters& InParam, FGuid InGuid = FGuid::NewGuid())
+	FRandomVectorFieldDataflowNode(const UE::Dataflow::FNodeParameters& InParam, FGuid InGuid = FGuid::NewGuid())
 		: FDataflowNode(InParam, InGuid)
 	{
 		RegisterInputConnection(&SamplePositions);
@@ -596,7 +596,7 @@ public:
 		RegisterOutputConnection(&NumSamplePositions);
 	}
 
-	virtual void Evaluate(Dataflow::FContext& Context, const FDataflowOutput* Out) const override;
+	virtual void Evaluate(UE::Dataflow::FContext& Context, const FDataflowOutput* Out) const override;
 
 };
 
@@ -645,7 +645,7 @@ public:
 	UPROPERTY(meta = (DataflowOutput))
 	int32 NumSamplePositions = 0;
 
-	FNoiseFieldDataflowNode(const Dataflow::FNodeParameters& InParam, FGuid InGuid = FGuid::NewGuid())
+	FNoiseFieldDataflowNode(const UE::Dataflow::FNodeParameters& InParam, FGuid InGuid = FGuid::NewGuid())
 		: FDataflowNode(InParam, InGuid)
 	{
 		RegisterInputConnection(&SamplePositions);
@@ -658,7 +658,7 @@ public:
 		RegisterOutputConnection(&NumSamplePositions);
 	}
 
-	virtual void Evaluate(Dataflow::FContext& Context, const FDataflowOutput* Out) const override;
+	virtual void Evaluate(UE::Dataflow::FContext& Context, const FDataflowOutput* Out) const override;
 
 };
 
@@ -699,7 +699,7 @@ public:
 	UPROPERTY(meta = (DataflowOutput))
 	int32 NumSamplePositions = 0;
 
-	FUniformIntegerFieldDataflowNode(const Dataflow::FNodeParameters& InParam, FGuid InGuid = FGuid::NewGuid())
+	FUniformIntegerFieldDataflowNode(const UE::Dataflow::FNodeParameters& InParam, FGuid InGuid = FGuid::NewGuid())
 		: FDataflowNode(InParam, InGuid)
 	{
 		RegisterInputConnection(&SamplePositions);
@@ -710,7 +710,7 @@ public:
 		RegisterOutputConnection(&NumSamplePositions);
 	}
 
-	virtual void Evaluate(Dataflow::FContext& Context, const FDataflowOutput* Out) const override;
+	virtual void Evaluate(UE::Dataflow::FContext& Context, const FDataflowOutput* Out) const override;
 
 };
 
@@ -792,7 +792,7 @@ public:
 	UPROPERTY(meta = (DataflowOutput))
 	int32 NumSamplePositions = 0;
 
-	FWaveScalarFieldDataflowNode(const Dataflow::FNodeParameters& InParam, FGuid InGuid = FGuid::NewGuid())
+	FWaveScalarFieldDataflowNode(const UE::Dataflow::FNodeParameters& InParam, FGuid InGuid = FGuid::NewGuid())
 		: FDataflowNode(InParam, InGuid)
 	{
 		RegisterInputConnection(&SamplePositions);
@@ -807,7 +807,7 @@ public:
 		RegisterOutputConnection(&NumSamplePositions);
 	}
 
-	virtual void Evaluate(Dataflow::FContext& Context, const FDataflowOutput* Out) const override;
+	virtual void Evaluate(UE::Dataflow::FContext& Context, const FDataflowOutput* Out) const override;
 
 };
 
@@ -878,7 +878,7 @@ public:
 	UPROPERTY(meta = (DataflowOutput))
 	TArray<int32> FieldRemap;
 
-	FSumScalarFieldDataflowNode(const Dataflow::FNodeParameters& InParam, FGuid InGuid = FGuid::NewGuid())
+	FSumScalarFieldDataflowNode(const UE::Dataflow::FNodeParameters& InParam, FGuid InGuid = FGuid::NewGuid())
 		: FDataflowNode(InParam, InGuid)
 	{
 		RegisterInputConnection(&FieldFloatLeft);
@@ -889,7 +889,7 @@ public:
 		RegisterOutputConnection(&FieldRemap);
 	}
 
-	virtual void Evaluate(Dataflow::FContext& Context, const FDataflowOutput* Out) const override;
+	virtual void Evaluate(UE::Dataflow::FContext& Context, const FDataflowOutput* Out) const override;
 
 };
 
@@ -964,7 +964,7 @@ public:
 	UPROPERTY(meta = (DataflowOutput))
 	TArray<int32> FieldRemap;
 
-	FSumVectorFieldDataflowNode(const Dataflow::FNodeParameters& InParam, FGuid InGuid = FGuid::NewGuid())
+	FSumVectorFieldDataflowNode(const UE::Dataflow::FNodeParameters& InParam, FGuid InGuid = FGuid::NewGuid())
 		: FDataflowNode(InParam, InGuid)
 	{
 		RegisterInputConnection(&FieldFloat);
@@ -978,7 +978,7 @@ public:
 		RegisterOutputConnection(&FieldRemap);
 	}
 
-	virtual void Evaluate(Dataflow::FContext& Context, const FDataflowOutput* Out) const override;
+	virtual void Evaluate(UE::Dataflow::FContext& Context, const FDataflowOutput* Out) const override;
 
 };
 
@@ -1017,7 +1017,7 @@ public:
 	UPROPERTY(meta = (DataflowOutput))
 	TArray<float> FieldFloatResult;
 
-	FFieldMakeDenseFloatArrayDataflowNode(const Dataflow::FNodeParameters& InParam, FGuid InGuid = FGuid::NewGuid())
+	FFieldMakeDenseFloatArrayDataflowNode(const UE::Dataflow::FNodeParameters& InParam, FGuid InGuid = FGuid::NewGuid())
 		: FDataflowNode(InParam, InGuid)
 	{
 		RegisterInputConnection(&FieldFloatInput);
@@ -1026,11 +1026,11 @@ public:
 		RegisterOutputConnection(&FieldFloatResult);
 	}
 
-	virtual void Evaluate(Dataflow::FContext& Context, const FDataflowOutput* Out) const override;
+	virtual void Evaluate(UE::Dataflow::FContext& Context, const FDataflowOutput* Out) const override;
 
 };
 
-namespace Dataflow
+namespace UE::Dataflow
 {
 	void GeometryCollectionFieldNodes();
 }
diff --git a/Engine/Plugins/Experimental/GeometryCollectionPlugin/Source/GeometryCollectionNodes/Public/Dataflow/GeometryCollectionFracturingNodes.h b/Engine/Plugins/Experimental/GeometryCollectionPlugin/Source/GeometryCollectionNodes/Public/Dataflow/GeometryCollectionFracturingNodes.h
--- a/Engine/Plugins/Experimental/GeometryCollectionPlugin/Source/GeometryCollectionNodes/Public/Dataflow/GeometryCollectionFracturingNodes.h
+++ b/Engine/Plugins/Experimental/GeometryCollectionPlugin/Source/GeometryCollectionNodes/Public/Dataflow/GeometryCollectionFracturingNodes.h
@@ -35,7 +35,7 @@ public:
 	UPROPERTY(meta = (DataflowOutput))
 	TArray<FVector> Points;
 
-	FUniformScatterPointsDataflowNode(const Dataflow::FNodeParameters& InParam, FGuid InGuid = FGuid::NewGuid())
+	FUniformScatterPointsDataflowNode(const UE::Dataflow::FNodeParameters& InParam, FGuid InGuid = FGuid::NewGuid())
 		: FDataflowNode(InParam, InGuid)
 	{
 		RegisterInputConnection(&BoundingBox);
@@ -45,7 +45,7 @@ public:
 		RegisterOutputConnection(&Points);
 	}
 
-	virtual void Evaluate(Dataflow::FContext& Context, const FDataflowOutput* Out) const override;
+	virtual void Evaluate(UE::Dataflow::FContext& Context, const FDataflowOutput* Out) const override;
 
 };
 
@@ -76,7 +76,7 @@ public:
 	UPROPERTY(meta = (DataflowOutput))
 	TArray<FVector> Points;
 
-	FUniformScatterPointsDataflowNode_v2(const Dataflow::FNodeParameters& InParam, FGuid InGuid = FGuid::NewGuid())
+	FUniformScatterPointsDataflowNode_v2(const UE::Dataflow::FNodeParameters& InParam, FGuid InGuid = FGuid::NewGuid())
 		: FDataflowNode(InParam, InGuid)
 	{
 		RegisterInputConnection(&BoundingBox);
@@ -86,7 +86,7 @@ public:
 		RegisterOutputConnection(&Points);
 	}
 
-	virtual void Evaluate(Dataflow::FContext& Context, const FDataflowOutput* Out) const override;
+	virtual void Evaluate(UE::Dataflow::FContext& Context, const FDataflowOutput* Out) const override;
 
 };
 
@@ -124,7 +124,7 @@ public:
 	UPROPERTY(meta = (DataflowOutput))
 	TArray<FVector> Points;
 
-	FRadialScatterPointsDataflowNode(const Dataflow::FNodeParameters& InParam, FGuid InGuid = FGuid::NewGuid())
+	FRadialScatterPointsDataflowNode(const UE::Dataflow::FNodeParameters& InParam, FGuid InGuid = FGuid::NewGuid())
 		: FDataflowNode(InParam, InGuid)
 	{
 		RegisterInputConnection(&Center);
@@ -138,7 +138,7 @@ public:
 		RegisterOutputConnection(&Points);
 	}
 
-	virtual void Evaluate(Dataflow::FContext& Context, const FDataflowOutput* Out) const override;
+	virtual void Evaluate(UE::Dataflow::FContext& Context, const FDataflowOutput* Out) const override;
 
 };
 
@@ -213,7 +213,7 @@ public:
 	UPROPERTY(meta = (DataflowOutput))
 	TArray<FVector> Points;
 
-	FRadialScatterPointsDataflowNode_v2(const Dataflow::FNodeParameters& InParam, FGuid InGuid = FGuid::NewGuid())
+	FRadialScatterPointsDataflowNode_v2(const UE::Dataflow::FNodeParameters& InParam, FGuid InGuid = FGuid::NewGuid())
 		: FDataflowNode(InParam, InGuid)
 	{
 		RegisterInputConnection(&BoundingBox);
@@ -234,7 +234,7 @@ public:
 		RegisterOutputConnection(&Points);
 	}
 
-	virtual void Evaluate(Dataflow::FContext& Context, const FDataflowOutput* Out) const override;
+	virtual void Evaluate(UE::Dataflow::FContext& Context, const FDataflowOutput* Out) const override;
 };
 
 USTRUCT(meta = (DataflowGeometryCollection))
@@ -280,7 +280,7 @@ public:
 	UPROPERTY(meta = (DataflowOutput))
 	TArray<FVector> Points;
 
-	FGridScatterPointsDataflowNode(const Dataflow::FNodeParameters& InParam, FGuid InGuid = FGuid::NewGuid())
+	FGridScatterPointsDataflowNode(const UE::Dataflow::FNodeParameters& InParam, FGuid InGuid = FGuid::NewGuid())
 		: FDataflowNode(InParam, InGuid)
 	{
 		RegisterInputConnection(&BoundingBox);
@@ -294,7 +294,7 @@ public:
 		RegisterOutputConnection(&Points);
 	}
 
-	virtual void Evaluate(Dataflow::FContext& Context, const FDataflowOutput* Out) const override;
+	virtual void Evaluate(UE::Dataflow::FContext& Context, const FDataflowOutput* Out) const override;
 };
 
 /**
@@ -356,7 +356,7 @@ public:
 	UPROPERTY(EditAnywhere, Category = "Collision", meta = (DataflowInput, UIMin = 0.f));
 	float CollisionSampleSpacing = 50.f;
 
-	FVoronoiFractureDataflowNode(const Dataflow::FNodeParameters& InParam, FGuid InGuid = FGuid::NewGuid())
+	FVoronoiFractureDataflowNode(const UE::Dataflow::FNodeParameters& InParam, FGuid InGuid = FGuid::NewGuid())
 		: FDataflowNode(InParam, InGuid)
 	{
 		RegisterInputConnection(&Collection);
@@ -375,7 +375,7 @@ public:
 		RegisterOutputConnection(&Collection, &Collection);
 	}
 
-	virtual void Evaluate(Dataflow::FContext& Context, const FDataflowOutput* Out) const override;
+	virtual void Evaluate(UE::Dataflow::FContext& Context, const FDataflowOutput* Out) const override;
 
 };
 
@@ -472,7 +472,7 @@ public:
 	UPROPERTY(meta = (DataflowOutput, DisplayName = "FracturedTransformSelection "))
 	FDataflowTransformSelection NewGeometryTransformSelection;
 
-	FVoronoiFractureDataflowNode_v2(const Dataflow::FNodeParameters& InParam, FGuid InGuid = FGuid::NewGuid())
+	FVoronoiFractureDataflowNode_v2(const UE::Dataflow::FNodeParameters& InParam, FGuid InGuid = FGuid::NewGuid())
 		: FDataflowNode(InParam, InGuid)
 	{
 		RegisterInputConnection(&Collection);
@@ -493,7 +493,7 @@ public:
 		RegisterOutputConnection(&NewGeometryTransformSelection);
 	}
 
-	virtual void Evaluate(Dataflow::FContext& Context, const FDataflowOutput* Out) const override;
+	virtual void Evaluate(UE::Dataflow::FContext& Context, const FDataflowOutput* Out) const override;
 };
 
 /**
@@ -552,7 +552,7 @@ public:
 	UPROPERTY(EditAnywhere, Category = "Collision", meta = (DataflowInput, UIMin = 0.f));
 	float CollisionSampleSpacing = 50.f;
 
-	FPlaneCutterDataflowNode(const Dataflow::FNodeParameters& InParam, FGuid InGuid = FGuid::NewGuid())
+	FPlaneCutterDataflowNode(const UE::Dataflow::FNodeParameters& InParam, FGuid InGuid = FGuid::NewGuid())
 		: FDataflowNode(InParam, InGuid)
 	{
 		RegisterInputConnection(&Collection);
@@ -572,7 +572,7 @@ public:
 		RegisterOutputConnection(&Collection, &Collection);
 	}
 
-	virtual void Evaluate(Dataflow::FContext& Context, const FDataflowOutput* Out) const override;
+	virtual void Evaluate(UE::Dataflow::FContext& Context, const FDataflowOutput* Out) const override;
 };
 
 /**
@@ -671,7 +671,7 @@ public:
 	UPROPERTY(meta = (DataflowOutput, DisplayName = "FracturedTransformSelection "))
 	FDataflowTransformSelection NewGeometryTransformSelection;
 
-	FPlaneCutterDataflowNode_v2(const Dataflow::FNodeParameters& InParam, FGuid InGuid = FGuid::NewGuid())
+	FPlaneCutterDataflowNode_v2(const UE::Dataflow::FNodeParameters& InParam, FGuid InGuid = FGuid::NewGuid())
 		: FDataflowNode(InParam, InGuid)
 	{
 		RegisterInputConnection(&Collection);
@@ -694,7 +694,7 @@ public:
 		RegisterOutputConnection(&NewGeometryTransformSelection);
 	}
 
-	virtual void Evaluate(Dataflow::FContext& Context, const FDataflowOutput* Out) const override;
+	virtual void Evaluate(UE::Dataflow::FContext& Context, const FDataflowOutput* Out) const override;
 };
 
 /**
@@ -722,7 +722,7 @@ public:
 	UPROPERTY(EditAnywhere, Category = "Scale", meta = (DataflowInput))
 	FVector Scale = FVector(1.0);
 
-	FExplodedViewDataflowNode(const Dataflow::FNodeParameters& InParam, FGuid InGuid = FGuid::NewGuid())
+	FExplodedViewDataflowNode(const UE::Dataflow::FNodeParameters& InParam, FGuid InGuid = FGuid::NewGuid())
 		: FDataflowNode(InParam, InGuid)
 	{
 		RegisterInputConnection(&Collection);
@@ -731,7 +731,7 @@ public:
 		RegisterOutputConnection(&Collection, &Collection);
 	}
 
-	virtual void Evaluate(Dataflow::FContext& Context, const FDataflowOutput* Out) const override;
+	virtual void Evaluate(UE::Dataflow::FContext& Context, const FDataflowOutput* Out) const override;
 
 private:
 	// todo(chaos) this is a copy of a function in FractureEditorModeToolkit, we should move this to a common place  
@@ -845,7 +845,7 @@ public:
 	UPROPERTY(meta = (DataflowOutput, DisplayName = "FracturedTransformSelection "))
 	FDataflowTransformSelection NewGeometryTransformSelection;
 
-	FSliceCutterDataflowNode(const Dataflow::FNodeParameters& InParam, FGuid InGuid = FGuid::NewGuid())
+	FSliceCutterDataflowNode(const UE::Dataflow::FNodeParameters& InParam, FGuid InGuid = FGuid::NewGuid())
 		: FDataflowNode(InParam, InGuid)
 	{
 		RegisterInputConnection(&Collection);
@@ -871,7 +871,7 @@ public:
 		RegisterOutputConnection(&NewGeometryTransformSelection);
 	}
 
-	virtual void Evaluate(Dataflow::FContext& Context, const FDataflowOutput* Out) const override;
+	virtual void Evaluate(UE::Dataflow::FContext& Context, const FDataflowOutput* Out) const override;
 };
 
 /**
@@ -983,7 +983,7 @@ public:
 	UPROPERTY(meta = (DataflowOutput, DisplayName = "FracturedTransformSelection "))
 	FDataflowTransformSelection NewGeometryTransformSelection;
 
-	FBrickCutterDataflowNode(const Dataflow::FNodeParameters& InParam, FGuid InGuid = FGuid::NewGuid())
+	FBrickCutterDataflowNode(const UE::Dataflow::FNodeParameters& InParam, FGuid InGuid = FGuid::NewGuid())
 		: FDataflowNode(InParam, InGuid)
 	{
 		RegisterInputConnection(&Collection);
@@ -1008,7 +1008,7 @@ public:
 		RegisterOutputConnection(&NewGeometryTransformSelection);
 	}
 
-	virtual void Evaluate(Dataflow::FContext& Context, const FDataflowOutput* Out) const override;
+	virtual void Evaluate(UE::Dataflow::FContext& Context, const FDataflowOutput* Out) const override;
 };
 
 /**
@@ -1132,7 +1132,7 @@ public:
 	UPROPERTY(meta = (DataflowOutput, DisplayName = "FracturedTransformSelection "))
 	FDataflowTransformSelection NewGeometryTransformSelection;
 
-	FMeshCutterDataflowNode(const Dataflow::FNodeParameters& InParam, FGuid InGuid = FGuid::NewGuid())
+	FMeshCutterDataflowNode(const UE::Dataflow::FNodeParameters& InParam, FGuid InGuid = FGuid::NewGuid())
 		: FDataflowNode(InParam, InGuid)
 	{
 		RegisterInputConnection(&Collection);
@@ -1158,10 +1158,10 @@ public:
 		RegisterOutputConnection(&NewGeometryTransformSelection);
 	}
 
-	virtual void Evaluate(Dataflow::FContext& Context, const FDataflowOutput* Out) const override;
+	virtual void Evaluate(UE::Dataflow::FContext& Context, const FDataflowOutput* Out) const override;
 };
 
-namespace Dataflow
+namespace UE::Dataflow
 {
 	void GeometryCollectionFracturingNodes();
 }
