Change Dataflow namespace to UE::Dataflow to be ok with coding standard

diff --git a/Engine/Plugins/Experimental/GeometryCollectionPlugin/Source/GeometryCollectionNodes/Public/Dataflow/GeometryCollectionMakeNodes.h b/Engine/Plugins/Experimental/GeometryCollectionPlugin/Source/GeometryCollectionNodes/Public/Dataflow/GeometryCollectionMakeNodes.h
--- a/Engine/Plugins/Experimental/GeometryCollectionPlugin/Source/GeometryCollectionNodes/Public/Dataflow/GeometryCollectionMakeNodes.h
+++ b/Engine/Plugins/Experimental/GeometryCollectionPlugin/Source/GeometryCollectionNodes/Public/Dataflow/GeometryCollectionMakeNodes.h
@@ -31,13 +31,13 @@ public:
 	UPROPERTY(meta = (DataflowOutput))
 	FString String;
 
-	FMakeLiteralStringDataflowNode(const Dataflow::FNodeParameters& InParam, FGuid InGuid = FGuid::NewGuid())
+	FMakeLiteralStringDataflowNode(const UE::Dataflow::FNodeParameters& InParam, FGuid InGuid = FGuid::NewGuid())
 		: FDataflowNode(InParam, InGuid)
 	{
 		RegisterOutputConnection(&String);
 	}
 
-	virtual void Evaluate(Dataflow::FContext& Context, const FDataflowOutput* Out) const override;
+	virtual void Evaluate(UE::Dataflow::FContext& Context, const FDataflowOutput* Out) const override;
 
 };
 
@@ -59,13 +59,13 @@ public:
 	UPROPERTY(meta = (DataflowOutput))
 	TArray<FVector> Points;
 
-	FMakePointsDataflowNode(const Dataflow::FNodeParameters& InParam, FGuid InGuid = FGuid::NewGuid())
+	FMakePointsDataflowNode(const UE::Dataflow::FNodeParameters& InParam, FGuid InGuid = FGuid::NewGuid())
 		: FDataflowNode(InParam, InGuid)
 	{
 		RegisterOutputConnection(&Points);
 	}
 
-	virtual void Evaluate(Dataflow::FContext& Context, const FDataflowOutput* Out) const override;
+	virtual void Evaluate(UE::Dataflow::FContext& Context, const FDataflowOutput* Out) const override;
 
 };
 
@@ -111,7 +111,7 @@ public:
 	UPROPERTY(meta = (DataflowOutput));
 	FBox Box = FBox(ForceInit);
 
-	FMakeBoxDataflowNode(const Dataflow::FNodeParameters& InParam, FGuid InGuid = FGuid::NewGuid())
+	FMakeBoxDataflowNode(const UE::Dataflow::FNodeParameters& InParam, FGuid InGuid = FGuid::NewGuid())
 		: FDataflowNode(InParam, InGuid)
 	{
 		RegisterInputConnection(&Min);
@@ -121,7 +121,7 @@ public:
 		RegisterOutputConnection(&Box);
 	}
 
-	virtual void Evaluate(Dataflow::FContext& Context, const FDataflowOutput* Out) const override;
+	virtual void Evaluate(UE::Dataflow::FContext& Context, const FDataflowOutput* Out) const override;
 
 };
 
@@ -147,7 +147,7 @@ public:
 	UPROPERTY(meta = (DataflowOutput));
 	FSphere Sphere = FSphere(ForceInit);
 
-	FMakeSphereDataflowNode(const Dataflow::FNodeParameters& InParam, FGuid InGuid = FGuid::NewGuid())
+	FMakeSphereDataflowNode(const UE::Dataflow::FNodeParameters& InParam, FGuid InGuid = FGuid::NewGuid())
 		: FDataflowNode(InParam, InGuid)
 	{
 		RegisterInputConnection(&Center);
@@ -155,7 +155,7 @@ public:
 		RegisterOutputConnection(&Sphere);
 	}
 
-	virtual void Evaluate(Dataflow::FContext& Context, const FDataflowOutput* Out) const override;
+	virtual void Evaluate(UE::Dataflow::FContext& Context, const FDataflowOutput* Out) const override;
 
 };
 
@@ -178,13 +178,13 @@ public:
 	UPROPERTY(meta = (DataflowOutput))
 	float Float = 0.f;
 
-	FMakeLiteralFloatDataflowNode(const Dataflow::FNodeParameters& InParam, FGuid InGuid = FGuid::NewGuid())
+	FMakeLiteralFloatDataflowNode(const UE::Dataflow::FNodeParameters& InParam, FGuid InGuid = FGuid::NewGuid())
 		: FDataflowNode(InParam, InGuid)
 	{
 		RegisterOutputConnection(&Float);
 	}
 
-	virtual void Evaluate(Dataflow::FContext& Context, const FDataflowOutput* Out) const override;
+	virtual void Evaluate(UE::Dataflow::FContext& Context, const FDataflowOutput* Out) const override;
 
 };
 
@@ -207,13 +207,13 @@ public:
 	UPROPERTY(meta = (DataflowOutput, DisplayName = "Int"))
 	int32 Int = 0;
 
-	FMakeLiteralIntDataflowNode(const Dataflow::FNodeParameters& InParam, FGuid InGuid = FGuid::NewGuid())
+	FMakeLiteralIntDataflowNode(const UE::Dataflow::FNodeParameters& InParam, FGuid InGuid = FGuid::NewGuid())
 		: FDataflowNode(InParam, InGuid)
 	{
 		RegisterOutputConnection(&Int);
 	}
 
-	virtual void Evaluate(Dataflow::FContext& Context, const FDataflowOutput* Out) const override;
+	virtual void Evaluate(UE::Dataflow::FContext& Context, const FDataflowOutput* Out) const override;
 
 };
 
@@ -235,13 +235,13 @@ public:
 	UPROPERTY(meta = (DataflowOutput))
 	bool Bool = false;
 
-	FMakeLiteralBoolDataflowNode(const Dataflow::FNodeParameters& InParam, FGuid InGuid = FGuid::NewGuid())
+	FMakeLiteralBoolDataflowNode(const UE::Dataflow::FNodeParameters& InParam, FGuid InGuid = FGuid::NewGuid())
 		: FDataflowNode(InParam, InGuid)
 	{
 		RegisterOutputConnection(&Bool);
 	}
 
-	virtual void Evaluate(Dataflow::FContext& Context, const FDataflowOutput* Out) const override;
+	virtual void Evaluate(UE::Dataflow::FContext& Context, const FDataflowOutput* Out) const override;
 
 };
 
@@ -270,7 +270,7 @@ public:
 	UPROPERTY(meta = (DataflowOutput, DisplayName = "Vector"))
 	FVector Vector = FVector(0.0);
 
-	FMakeLiteralVectorDataflowNode(const Dataflow::FNodeParameters& InParam, FGuid InGuid = FGuid::NewGuid())
+	FMakeLiteralVectorDataflowNode(const UE::Dataflow::FNodeParameters& InParam, FGuid InGuid = FGuid::NewGuid())
 		: FDataflowNode(InParam, InGuid)
 	{
 		RegisterInputConnection(&X);
@@ -279,7 +279,7 @@ public:
 		RegisterOutputConnection(&Vector);
 	}
 
-	virtual void Evaluate(Dataflow::FContext& Context, const FDataflowOutput* Out) const override;
+	virtual void Evaluate(UE::Dataflow::FContext& Context, const FDataflowOutput* Out) const override;
 
 };
 
@@ -308,7 +308,7 @@ public:
 	UPROPERTY(meta = (DataflowOutput, DisplayName = "Transform"));
 	FTransform OutTransform = FTransform::Identity;
 
-	FMakeTransformDataflowNode(const Dataflow::FNodeParameters& InParam, FGuid InGuid = FGuid::NewGuid())
+	FMakeTransformDataflowNode(const UE::Dataflow::FNodeParameters& InParam, FGuid InGuid = FGuid::NewGuid())
 		: FDataflowNode(InParam, InGuid)
 	{
 		RegisterInputConnection(&InTranslation);
@@ -317,7 +317,7 @@ public:
 		RegisterOutputConnection(&OutTransform);
 	}
 
-	virtual void Evaluate(Dataflow::FContext& Context, const FDataflowOutput* Out) const override;
+	virtual void Evaluate(UE::Dataflow::FContext& Context, const FDataflowOutput* Out) const override;
 
 };
 
@@ -348,7 +348,7 @@ public:
 	UPROPERTY(meta = (DataflowOutput, DisplayName = "Quaternion"))
 	FQuat Quaternion = FQuat(ForceInitToZero);
 
-	FMakeQuaternionDataflowNode(const Dataflow::FNodeParameters& InParam, FGuid InGuid = FGuid::NewGuid())
+	FMakeQuaternionDataflowNode(const UE::Dataflow::FNodeParameters& InParam, FGuid InGuid = FGuid::NewGuid())
 		: FDataflowNode(InParam, InGuid)
 	{
 		RegisterInputConnection(&X);
@@ -358,7 +358,7 @@ public:
 		RegisterOutputConnection(&Quaternion);
 	}
 
-	virtual void Evaluate(Dataflow::FContext& Context, const FDataflowOutput* Out) const override;
+	virtual void Evaluate(UE::Dataflow::FContext& Context, const FDataflowOutput* Out) const override;
 
 };
 
@@ -386,7 +386,7 @@ public:
 	UPROPERTY(meta = (DataflowOutput))
 	TArray<float> FloatArray;
 
-	FMakeFloatArrayDataflowNode(const Dataflow::FNodeParameters& InParam, FGuid InGuid = FGuid::NewGuid())
+	FMakeFloatArrayDataflowNode(const UE::Dataflow::FNodeParameters& InParam, FGuid InGuid = FGuid::NewGuid())
 		: FDataflowNode(InParam, InGuid)
 	{
 		RegisterInputConnection(&NumElements);
@@ -394,11 +394,11 @@ public:
 		RegisterOutputConnection(&FloatArray);
 	}
 
-	virtual void Evaluate(Dataflow::FContext& Context, const FDataflowOutput* Out) const override;
+	virtual void Evaluate(UE::Dataflow::FContext& Context, const FDataflowOutput* Out) const override;
 
 };
 
-namespace Dataflow
+namespace UE::Dataflow
 {
 	void GeometryCollectionMakeNodes();
 }
diff --git a/Engine/Plugins/Experimental/GeometryCollectionPlugin/Source/GeometryCollectionNodes/Public/Dataflow/GeometryCollectionMaterialNodes.h b/Engine/Plugins/Experimental/GeometryCollectionPlugin/Source/GeometryCollectionNodes/Public/Dataflow/GeometryCollectionMaterialNodes.h
--- a/Engine/Plugins/Experimental/GeometryCollectionPlugin/Source/GeometryCollectionNodes/Public/Dataflow/GeometryCollectionMaterialNodes.h
+++ b/Engine/Plugins/Experimental/GeometryCollectionPlugin/Source/GeometryCollectionNodes/Public/Dataflow/GeometryCollectionMaterialNodes.h
@@ -59,7 +59,7 @@ public:
 	UPROPERTY(EditAnywhere, Category = "Material", meta = (DisplayName = "Assign Inside Material to Inside Faces"))
 	bool bAssignInsideMaterial = false;
 
-	FAddMaterialToCollectionDataflowNode(const Dataflow::FNodeParameters& InParam, FGuid InGuid = FGuid::NewGuid())
+	FAddMaterialToCollectionDataflowNode(const UE::Dataflow::FNodeParameters& InParam, FGuid InGuid = FGuid::NewGuid())
 		: FDataflowNode(InParam, InGuid)
 	{
 		RegisterInputConnection(&Collection);
@@ -71,7 +71,7 @@ public:
 		RegisterOutputConnection(&Materials, &Materials);
 	}
 
-	virtual void Evaluate(Dataflow::FContext& Context, const FDataflowOutput* Out) const override;
+	virtual void Evaluate(UE::Dataflow::FContext& Context, const FDataflowOutput* Out) const override;
 
 };
 
@@ -116,7 +116,7 @@ public:
 	UPROPERTY(EditAnywhere, Category = "Material", meta = (DisplayName = "Assign Inside Material to Inside Faces"))
 	bool bAssignInsideMaterial = false;
 
-	FReAssignMaterialInCollectionDataflowNode(const Dataflow::FNodeParameters& InParam, FGuid InGuid = FGuid::NewGuid())
+	FReAssignMaterialInCollectionDataflowNode(const UE::Dataflow::FNodeParameters& InParam, FGuid InGuid = FGuid::NewGuid())
 		: FDataflowNode(InParam, InGuid)
 	{
 		RegisterInputConnection(&Collection);
@@ -128,7 +128,7 @@ public:
 		RegisterOutputConnection(&Materials, &Materials);
 	}
 
-	virtual void Evaluate(Dataflow::FContext& Context, const FDataflowOutput* Out) const override;
+	virtual void Evaluate(UE::Dataflow::FContext& Context, const FDataflowOutput* Out) const override;
 
 };
 
@@ -153,14 +153,14 @@ public:
 	UPROPERTY(meta = (DataflowOutput))
 	FString String;
 
-	FMaterialsInfoDataflowNode(const Dataflow::FNodeParameters& InParam, FGuid InGuid = FGuid::NewGuid())
+	FMaterialsInfoDataflowNode(const UE::Dataflow::FNodeParameters& InParam, FGuid InGuid = FGuid::NewGuid())
 		: FDataflowNode(InParam, InGuid)
 	{
 		RegisterInputConnection(&Materials);
 		RegisterOutputConnection(&String);
 	}
 
-	virtual void Evaluate(Dataflow::FContext& Context, const FDataflowOutput* Out) const override;
+	virtual void Evaluate(UE::Dataflow::FContext& Context, const FDataflowOutput* Out) const override;
 
 };
 
@@ -189,7 +189,7 @@ public:
 	UPROPERTY(EditAnywhere, Category = "Material", meta = (DataflowInput, DisplayName = "Material Index"))
 	int32 MaterialIdx = 0;
 
-	FGetMaterialFromMaterialsArrayDataflowNode(const Dataflow::FNodeParameters& InParam, FGuid InGuid = FGuid::NewGuid())
+	FGetMaterialFromMaterialsArrayDataflowNode(const UE::Dataflow::FNodeParameters& InParam, FGuid InGuid = FGuid::NewGuid())
 		: FDataflowNode(InParam, InGuid)
 	{
 		RegisterInputConnection(&Materials);
@@ -197,7 +197,7 @@ public:
 		RegisterOutputConnection(&Material);
 	}
 
-	virtual void Evaluate(Dataflow::FContext& Context, const FDataflowOutput* Out) const override;
+	virtual void Evaluate(UE::Dataflow::FContext& Context, const FDataflowOutput* Out) const override;
 
 };
 
@@ -241,7 +241,7 @@ public:
 	UPROPERTY(EditAnywhere, Category = "Material", meta = (DataflowInput, DisplayName = "Material Index", EditCondition = "Operation == ESetMaterialOperationTypeEnum::Dataflow_SetMaterialOperationType_Insert", EditConditionHides))
 	int32 MaterialIdx = 0;
 
-	FSetMaterialInMaterialsArrayDataflowNode(const Dataflow::FNodeParameters& InParam, FGuid InGuid = FGuid::NewGuid())
+	FSetMaterialInMaterialsArrayDataflowNode(const UE::Dataflow::FNodeParameters& InParam, FGuid InGuid = FGuid::NewGuid())
 		: FDataflowNode(InParam, InGuid)
 	{
 		RegisterInputConnection(&Materials);
@@ -250,7 +250,7 @@ public:
 		RegisterOutputConnection(&Materials, &Materials);
 	}
 
-	virtual void Evaluate(Dataflow::FContext& Context, const FDataflowOutput* Out) const override;
+	virtual void Evaluate(UE::Dataflow::FContext& Context, const FDataflowOutput* Out) const override;
 
 };
 
@@ -275,13 +275,13 @@ public:
 	UPROPERTY(meta = (DataflowOutput));
 	TObjectPtr<UMaterial> Material;
 
-	FMakeMaterialDataflowNode(const Dataflow::FNodeParameters& InParam, FGuid InGuid = FGuid::NewGuid())
+	FMakeMaterialDataflowNode(const UE::Dataflow::FNodeParameters& InParam, FGuid InGuid = FGuid::NewGuid())
 		: FDataflowNode(InParam, InGuid)
 	{
 		RegisterOutputConnection(&Material);
 	}
 
-	virtual void Evaluate(Dataflow::FContext& Context, const FDataflowOutput* Out) const override;
+	virtual void Evaluate(UE::Dataflow::FContext& Context, const FDataflowOutput* Out) const override;
 
 };
 
@@ -301,18 +301,18 @@ public:
 	UPROPERTY(meta = (DataflowOutput));
 	TArray<TObjectPtr<UMaterial>> Materials;
 
-	FMakeMaterialsArrayDataflowNode(const Dataflow::FNodeParameters& InParam, FGuid InGuid = FGuid::NewGuid())
+	FMakeMaterialsArrayDataflowNode(const UE::Dataflow::FNodeParameters& InParam, FGuid InGuid = FGuid::NewGuid())
 		: FDataflowNode(InParam, InGuid)
 	{
 		RegisterOutputConnection(&Materials);
 	}
 
-	virtual void Evaluate(Dataflow::FContext& Context, const FDataflowOutput* Out) const override;
+	virtual void Evaluate(UE::Dataflow::FContext& Context, const FDataflowOutput* Out) const override;
 
 };
 
 
-namespace Dataflow
+namespace UE::Dataflow
 {
 	void GeometryCollectionMaterialNodes();
 }
diff --git a/Engine/Plugins/Experimental/GeometryCollectionPlugin/Source/GeometryCollectionNodes/Public/Dataflow/GeometryCollectionMathNodes.h b/Engine/Plugins/Experimental/GeometryCollectionPlugin/Source/GeometryCollectionNodes/Public/Dataflow/GeometryCollectionMathNodes.h
--- a/Engine/Plugins/Experimental/GeometryCollectionPlugin/Source/GeometryCollectionNodes/Public/Dataflow/GeometryCollectionMathNodes.h
+++ b/Engine/Plugins/Experimental/GeometryCollectionPlugin/Source/GeometryCollectionNodes/Public/Dataflow/GeometryCollectionMathNodes.h
@@ -28,7 +28,7 @@ public:
 	UPROPERTY(meta = (DataflowOutput))
 	float ReturnValue = 0.f;
 
-	FAddDataflowNode(const Dataflow::FNodeParameters& InParam, FGuid InGuid = FGuid::NewGuid())
+	FAddDataflowNode(const UE::Dataflow::FNodeParameters& InParam, FGuid InGuid = FGuid::NewGuid())
 		: FDataflowNode(InParam, InGuid)
 	{
 		RegisterInputConnection(&FloatA);
@@ -36,7 +36,7 @@ public:
 		RegisterOutputConnection(&ReturnValue);
 	}
 
-	virtual void Evaluate(Dataflow::FContext& Context, const FDataflowOutput* Out) const override;
+	virtual void Evaluate(UE::Dataflow::FContext& Context, const FDataflowOutput* Out) const override;
 };
 
 /**
@@ -60,7 +60,7 @@ public:
 	UPROPERTY(meta = (DataflowOutput))
 	float ReturnValue = 0.f;
 
-	FSubtractDataflowNode(const Dataflow::FNodeParameters& InParam, FGuid InGuid = FGuid::NewGuid())
+	FSubtractDataflowNode(const UE::Dataflow::FNodeParameters& InParam, FGuid InGuid = FGuid::NewGuid())
 		: FDataflowNode(InParam, InGuid)
 	{
 		RegisterInputConnection(&FloatA);
@@ -68,7 +68,7 @@ public:
 		RegisterOutputConnection(&ReturnValue);
 	}
 
-	virtual void Evaluate(Dataflow::FContext& Context, const FDataflowOutput* Out) const override;
+	virtual void Evaluate(UE::Dataflow::FContext& Context, const FDataflowOutput* Out) const override;
 };
 
 /**
@@ -92,7 +92,7 @@ public:
 	UPROPERTY(meta = (DataflowOutput))
 	float ReturnValue = 0.f;
 
-	FMultiplyDataflowNode(const Dataflow::FNodeParameters& InParam, FGuid InGuid = FGuid::NewGuid())
+	FMultiplyDataflowNode(const UE::Dataflow::FNodeParameters& InParam, FGuid InGuid = FGuid::NewGuid())
 		: FDataflowNode(InParam, InGuid)
 	{
 		RegisterInputConnection(&FloatA);
@@ -100,7 +100,7 @@ public:
 		RegisterOutputConnection(&ReturnValue);
 	}
 
-	virtual void Evaluate(Dataflow::FContext& Context, const FDataflowOutput* Out) const override;
+	virtual void Evaluate(UE::Dataflow::FContext& Context, const FDataflowOutput* Out) const override;
 };
 
 /**
@@ -124,7 +124,7 @@ public:
 	UPROPERTY(meta = (DataflowOutput))
 	float ReturnValue = 0.f;
 
-	FSafeDivideDataflowNode(const Dataflow::FNodeParameters& InParam, FGuid InGuid = FGuid::NewGuid())
+	FSafeDivideDataflowNode(const UE::Dataflow::FNodeParameters& InParam, FGuid InGuid = FGuid::NewGuid())
 		: FDataflowNode(InParam, InGuid)
 	{
 		RegisterInputConnection(&FloatA);
@@ -132,7 +132,7 @@ public:
 		RegisterOutputConnection(&ReturnValue);
 	}
 
-	virtual void Evaluate(Dataflow::FContext& Context, const FDataflowOutput* Out) const override;
+	virtual void Evaluate(UE::Dataflow::FContext& Context, const FDataflowOutput* Out) const override;
 };
 
 USTRUCT(meta = (DataflowGeometryCollection))
@@ -142,7 +142,7 @@ struct FDivideDataflowNode : public FSafeDivideDataflowNode
 	DATAFLOW_NODE_DEFINE_INTERNAL(FDivideDataflowNode, "Divide", "Math|Float", "")
 
 public:
-	FDivideDataflowNode(const Dataflow::FNodeParameters& InParam, FGuid InGuid = FGuid::NewGuid())
+	FDivideDataflowNode(const UE::Dataflow::FNodeParameters& InParam, FGuid InGuid = FGuid::NewGuid())
 		: FSafeDivideDataflowNode(InParam, InGuid)
 	{}
 };
@@ -171,7 +171,7 @@ public:
 	UPROPERTY(meta = (DataflowOutput))
 	int32 ReturnValue = 0.f;
 
-	FDivisionDataflowNode(const Dataflow::FNodeParameters& InParam, FGuid InGuid = FGuid::NewGuid())
+	FDivisionDataflowNode(const UE::Dataflow::FNodeParameters& InParam, FGuid InGuid = FGuid::NewGuid())
 		: FDataflowNode(InParam, InGuid)
 	{
 		RegisterInputConnection(&Dividend);
@@ -180,7 +180,7 @@ public:
 		RegisterOutputConnection(&ReturnValue);
 	}
 
-	virtual void Evaluate(Dataflow::FContext& Context, const FDataflowOutput* Out) const override;
+	virtual void Evaluate(UE::Dataflow::FContext& Context, const FDataflowOutput* Out) const override;
 };
 
 /**
@@ -201,14 +201,14 @@ public:
 	UPROPERTY(meta = (DataflowOutput))
 	float ReturnValue = 0.f;
 
-	FSafeReciprocalDataflowNode(const Dataflow::FNodeParameters& InParam, FGuid InGuid = FGuid::NewGuid())
+	FSafeReciprocalDataflowNode(const UE::Dataflow::FNodeParameters& InParam, FGuid InGuid = FGuid::NewGuid())
 		: FDataflowNode(InParam, InGuid)
 	{
 		RegisterInputConnection(&Float);
 		RegisterOutputConnection(&ReturnValue);
 	}
 
-	virtual void Evaluate(Dataflow::FContext& Context, const FDataflowOutput* Out) const override;
+	virtual void Evaluate(UE::Dataflow::FContext& Context, const FDataflowOutput* Out) const override;
 };
 
 /**
@@ -229,14 +229,14 @@ public:
 	UPROPERTY(meta = (DataflowOutput))
 	float ReturnValue = 0.f;
 
-	FSquareDataflowNode(const Dataflow::FNodeParameters& InParam, FGuid InGuid = FGuid::NewGuid())
+	FSquareDataflowNode(const UE::Dataflow::FNodeParameters& InParam, FGuid InGuid = FGuid::NewGuid())
 		: FDataflowNode(InParam, InGuid)
 	{
 		RegisterInputConnection(&Float);
 		RegisterOutputConnection(&ReturnValue);
 	}
 
-	virtual void Evaluate(Dataflow::FContext& Context, const FDataflowOutput* Out) const override;
+	virtual void Evaluate(UE::Dataflow::FContext& Context, const FDataflowOutput* Out) const override;
 };
 
 /**
@@ -257,14 +257,14 @@ public:
 	UPROPERTY(meta = (DataflowOutput))
 	float ReturnValue = 0.f;
 
-	FSquareRootDataflowNode(const Dataflow::FNodeParameters& InParam, FGuid InGuid = FGuid::NewGuid())
+	FSquareRootDataflowNode(const UE::Dataflow::FNodeParameters& InParam, FGuid InGuid = FGuid::NewGuid())
 		: FDataflowNode(InParam, InGuid)
 	{
 		RegisterInputConnection(&Float);
 		RegisterOutputConnection(&ReturnValue);
 	}
 
-	virtual void Evaluate(Dataflow::FContext& Context, const FDataflowOutput* Out) const override;
+	virtual void Evaluate(UE::Dataflow::FContext& Context, const FDataflowOutput* Out) const override;
 };
 
 /**
@@ -285,14 +285,14 @@ public:
 	UPROPERTY(meta = (DataflowOutput))
 	float ReturnValue = 0.f;
 
-	FInverseSqrtDataflowNode(const Dataflow::FNodeParameters& InParam, FGuid InGuid = FGuid::NewGuid())
+	FInverseSqrtDataflowNode(const UE::Dataflow::FNodeParameters& InParam, FGuid InGuid = FGuid::NewGuid())
 		: FDataflowNode(InParam, InGuid)
 	{
 		RegisterInputConnection(&Float);
 		RegisterOutputConnection(&ReturnValue);
 	}
 
-	virtual void Evaluate(Dataflow::FContext& Context, const FDataflowOutput* Out) const override;
+	virtual void Evaluate(UE::Dataflow::FContext& Context, const FDataflowOutput* Out) const override;
 };
 
 /**
@@ -313,14 +313,14 @@ public:
 	UPROPERTY(meta = (DataflowOutput))
 	float ReturnValue = 0.f;
 
-	FCubeDataflowNode(const Dataflow::FNodeParameters& InParam, FGuid InGuid = FGuid::NewGuid())
+	FCubeDataflowNode(const UE::Dataflow::FNodeParameters& InParam, FGuid InGuid = FGuid::NewGuid())
 		: FDataflowNode(InParam, InGuid)
 	{
 		RegisterInputConnection(&Float);
 		RegisterOutputConnection(&ReturnValue);
 	}
 
-	virtual void Evaluate(Dataflow::FContext& Context, const FDataflowOutput* Out) const override;
+	virtual void Evaluate(UE::Dataflow::FContext& Context, const FDataflowOutput* Out) const override;
 };
 
 /**
@@ -341,14 +341,14 @@ public:
 	UPROPERTY(meta = (DataflowOutput))
 	float ReturnValue = 0.f;
 
-	FNegateDataflowNode(const Dataflow::FNodeParameters& InParam, FGuid InGuid = FGuid::NewGuid())
+	FNegateDataflowNode(const UE::Dataflow::FNodeParameters& InParam, FGuid InGuid = FGuid::NewGuid())
 		: FDataflowNode(InParam, InGuid)
 	{
 		RegisterInputConnection(&Float);
 		RegisterOutputConnection(&ReturnValue);
 	}
 
-	virtual void Evaluate(Dataflow::FContext& Context, const FDataflowOutput* Out) const override;
+	virtual void Evaluate(UE::Dataflow::FContext& Context, const FDataflowOutput* Out) const override;
 };
 
 /**
@@ -369,14 +369,14 @@ public:
 	UPROPERTY(meta = (DataflowOutput))
 	float ReturnValue = 0.f;
 
-	FAbsDataflowNode(const Dataflow::FNodeParameters& InParam, FGuid InGuid = FGuid::NewGuid())
+	FAbsDataflowNode(const UE::Dataflow::FNodeParameters& InParam, FGuid InGuid = FGuid::NewGuid())
 		: FDataflowNode(InParam, InGuid)
 	{
 		RegisterInputConnection(&Float);
 		RegisterOutputConnection(&ReturnValue);
 	}
 
-	virtual void Evaluate(Dataflow::FContext& Context, const FDataflowOutput* Out) const override;
+	virtual void Evaluate(UE::Dataflow::FContext& Context, const FDataflowOutput* Out) const override;
 };
 
 /**
@@ -397,14 +397,14 @@ public:
 	UPROPERTY(meta = (DataflowOutput))
 	float ReturnValue = 0.f;
 
-	FFloorDataflowNode(const Dataflow::FNodeParameters& InParam, FGuid InGuid = FGuid::NewGuid())
+	FFloorDataflowNode(const UE::Dataflow::FNodeParameters& InParam, FGuid InGuid = FGuid::NewGuid())
 		: FDataflowNode(InParam, InGuid)
 	{
 		RegisterInputConnection(&Float);
 		RegisterOutputConnection(&ReturnValue);
 	}
 
-	virtual void Evaluate(Dataflow::FContext& Context, const FDataflowOutput* Out) const override;
+	virtual void Evaluate(UE::Dataflow::FContext& Context, const FDataflowOutput* Out) const override;
 };
 
 /**
@@ -425,14 +425,14 @@ public:
 	UPROPERTY(meta = (DataflowOutput))
 	float ReturnValue = 0.f;
 
-	FCeilDataflowNode(const Dataflow::FNodeParameters& InParam, FGuid InGuid = FGuid::NewGuid())
+	FCeilDataflowNode(const UE::Dataflow::FNodeParameters& InParam, FGuid InGuid = FGuid::NewGuid())
 		: FDataflowNode(InParam, InGuid)
 	{
 		RegisterInputConnection(&Float);
 		RegisterOutputConnection(&ReturnValue);
 	}
 
-	virtual void Evaluate(Dataflow::FContext& Context, const FDataflowOutput* Out) const override;
+	virtual void Evaluate(UE::Dataflow::FContext& Context, const FDataflowOutput* Out) const override;
 };
 
 /**
@@ -453,14 +453,14 @@ public:
 	UPROPERTY(meta = (DataflowOutput))
 	float ReturnValue = 0.f;
 
-	FRoundDataflowNode(const Dataflow::FNodeParameters& InParam, FGuid InGuid = FGuid::NewGuid())
+	FRoundDataflowNode(const UE::Dataflow::FNodeParameters& InParam, FGuid InGuid = FGuid::NewGuid())
 		: FDataflowNode(InParam, InGuid)
 	{
 		RegisterInputConnection(&Float);
 		RegisterOutputConnection(&ReturnValue);
 	}
 
-	virtual void Evaluate(Dataflow::FContext& Context, const FDataflowOutput* Out) const override;
+	virtual void Evaluate(UE::Dataflow::FContext& Context, const FDataflowOutput* Out) const override;
 };
 
 /**
@@ -481,14 +481,14 @@ public:
 	UPROPERTY(meta = (DataflowOutput))
 	float ReturnValue = 0.f;
 
-	FTruncDataflowNode(const Dataflow::FNodeParameters& InParam, FGuid InGuid = FGuid::NewGuid())
+	FTruncDataflowNode(const UE::Dataflow::FNodeParameters& InParam, FGuid InGuid = FGuid::NewGuid())
 		: FDataflowNode(InParam, InGuid)
 	{
 		RegisterInputConnection(&Float);
 		RegisterOutputConnection(&ReturnValue);
 	}
 
-	virtual void Evaluate(Dataflow::FContext& Context, const FDataflowOutput* Out) const override;
+	virtual void Evaluate(UE::Dataflow::FContext& Context, const FDataflowOutput* Out) const override;
 };
 
 /**
@@ -509,14 +509,14 @@ public:
 	UPROPERTY(meta = (DataflowOutput))
 	float ReturnValue = 0.f;
 
-	FFracDataflowNode(const Dataflow::FNodeParameters& InParam, FGuid InGuid = FGuid::NewGuid())
+	FFracDataflowNode(const UE::Dataflow::FNodeParameters& InParam, FGuid InGuid = FGuid::NewGuid())
 		: FDataflowNode(InParam, InGuid)
 	{
 		RegisterInputConnection(&Float);
 		RegisterOutputConnection(&ReturnValue);
 	}
 
-	virtual void Evaluate(Dataflow::FContext& Context, const FDataflowOutput* Out) const override;
+	virtual void Evaluate(UE::Dataflow::FContext& Context, const FDataflowOutput* Out) const override;
 };
 
 /**
@@ -540,7 +540,7 @@ public:
 	UPROPERTY(meta = (DataflowOutput))
 	float ReturnValue = 0.f;
 
-	FMinDataflowNode(const Dataflow::FNodeParameters& InParam, FGuid InGuid = FGuid::NewGuid())
+	FMinDataflowNode(const UE::Dataflow::FNodeParameters& InParam, FGuid InGuid = FGuid::NewGuid())
 		: FDataflowNode(InParam, InGuid)
 	{
 		RegisterInputConnection(&FloatA);
@@ -548,7 +548,7 @@ public:
 		RegisterOutputConnection(&ReturnValue);
 	}
 
-	virtual void Evaluate(Dataflow::FContext& Context, const FDataflowOutput* Out) const override;
+	virtual void Evaluate(UE::Dataflow::FContext& Context, const FDataflowOutput* Out) const override;
 };
 
 /**
@@ -572,7 +572,7 @@ public:
 	UPROPERTY(meta = (DataflowOutput))
 		float ReturnValue = 0.f;
 
-	FMaxDataflowNode(const Dataflow::FNodeParameters& InParam, FGuid InGuid = FGuid::NewGuid())
+	FMaxDataflowNode(const UE::Dataflow::FNodeParameters& InParam, FGuid InGuid = FGuid::NewGuid())
 		: FDataflowNode(InParam, InGuid)
 	{
 		RegisterInputConnection(&FloatA);
@@ -580,7 +580,7 @@ public:
 		RegisterOutputConnection(&ReturnValue);
 	}
 
-	virtual void Evaluate(Dataflow::FContext& Context, const FDataflowOutput* Out) const override;
+	virtual void Evaluate(UE::Dataflow::FContext& Context, const FDataflowOutput* Out) const override;
 };
 
 /**
@@ -607,7 +607,7 @@ public:
 	UPROPERTY(meta = (DataflowOutput))
 		float ReturnValue = 0.f;
 
-	FMin3DataflowNode(const Dataflow::FNodeParameters& InParam, FGuid InGuid = FGuid::NewGuid())
+	FMin3DataflowNode(const UE::Dataflow::FNodeParameters& InParam, FGuid InGuid = FGuid::NewGuid())
 		: FDataflowNode(InParam, InGuid)
 	{
 		RegisterInputConnection(&FloatA);
@@ -616,7 +616,7 @@ public:
 		RegisterOutputConnection(&ReturnValue);
 	}
 
-	virtual void Evaluate(Dataflow::FContext& Context, const FDataflowOutput* Out) const override;
+	virtual void Evaluate(UE::Dataflow::FContext& Context, const FDataflowOutput* Out) const override;
 };
 
 /**
@@ -643,7 +643,7 @@ public:
 	UPROPERTY(meta = (DataflowOutput))
 	float ReturnValue = 0.f;
 
-	FMax3DataflowNode(const Dataflow::FNodeParameters& InParam, FGuid InGuid = FGuid::NewGuid())
+	FMax3DataflowNode(const UE::Dataflow::FNodeParameters& InParam, FGuid InGuid = FGuid::NewGuid())
 		: FDataflowNode(InParam, InGuid)
 	{
 		RegisterInputConnection(&FloatA);
@@ -652,7 +652,7 @@ public:
 		RegisterOutputConnection(&ReturnValue);
 	}
 
-	virtual void Evaluate(Dataflow::FContext& Context, const FDataflowOutput* Out) const override;
+	virtual void Evaluate(UE::Dataflow::FContext& Context, const FDataflowOutput* Out) const override;
 };
 
 /**
@@ -673,14 +673,14 @@ public:
 	UPROPERTY(meta = (DataflowOutput))
 	float ReturnValue = 0.f;
 
-	FSignDataflowNode(const Dataflow::FNodeParameters& InParam, FGuid InGuid = FGuid::NewGuid())
+	FSignDataflowNode(const UE::Dataflow::FNodeParameters& InParam, FGuid InGuid = FGuid::NewGuid())
 		: FDataflowNode(InParam, InGuid)
 	{
 		RegisterInputConnection(&Float);
 		RegisterOutputConnection(&ReturnValue);
 	}
 
-	virtual void Evaluate(Dataflow::FContext& Context, const FDataflowOutput* Out) const override;
+	virtual void Evaluate(UE::Dataflow::FContext& Context, const FDataflowOutput* Out) const override;
 };
 
 /**
@@ -707,7 +707,7 @@ public:
 	UPROPERTY(meta = (DataflowOutput))
 	float ReturnValue = 0.f;
 
-	FClampDataflowNode(const Dataflow::FNodeParameters& InParam, FGuid InGuid = FGuid::NewGuid())
+	FClampDataflowNode(const UE::Dataflow::FNodeParameters& InParam, FGuid InGuid = FGuid::NewGuid())
 		: FDataflowNode(InParam, InGuid)
 	{
 		RegisterInputConnection(&Float);
@@ -716,7 +716,7 @@ public:
 		RegisterOutputConnection(&ReturnValue);
 	}
 
-	virtual void Evaluate(Dataflow::FContext& Context, const FDataflowOutput* Out) const override;
+	virtual void Evaluate(UE::Dataflow::FContext& Context, const FDataflowOutput* Out) const override;
 };
 
 /**
@@ -750,7 +750,7 @@ public:
 	UPROPERTY(meta = (DataflowOutput))
 	float ReturnValue = 0.f;
 
-	FFitDataflowNode(const Dataflow::FNodeParameters& InParam, FGuid InGuid = FGuid::NewGuid())
+	FFitDataflowNode(const UE::Dataflow::FNodeParameters& InParam, FGuid InGuid = FGuid::NewGuid())
 		: FDataflowNode(InParam, InGuid)
 	{
 		RegisterInputConnection(&Float);
@@ -761,7 +761,7 @@ public:
 		RegisterOutputConnection(&ReturnValue);
 	}
 
-	virtual void Evaluate(Dataflow::FContext& Context, const FDataflowOutput* Out) const override;
+	virtual void Evaluate(UE::Dataflow::FContext& Context, const FDataflowOutput* Out) const override;
 };
 
 /**
@@ -795,7 +795,7 @@ public:
 	UPROPERTY(meta = (DataflowOutput))
 	float ReturnValue = 0.f;
 
-	FEFitDataflowNode(const Dataflow::FNodeParameters& InParam, FGuid InGuid = FGuid::NewGuid())
+	FEFitDataflowNode(const UE::Dataflow::FNodeParameters& InParam, FGuid InGuid = FGuid::NewGuid())
 		: FDataflowNode(InParam, InGuid)
 	{
 		RegisterInputConnection(&Float);
@@ -806,7 +806,7 @@ public:
 		RegisterOutputConnection(&ReturnValue);
 	}
 
-	virtual void Evaluate(Dataflow::FContext& Context, const FDataflowOutput* Out) const override;
+	virtual void Evaluate(UE::Dataflow::FContext& Context, const FDataflowOutput* Out) const override;
 };
 
 /**
@@ -830,7 +830,7 @@ public:
 	UPROPERTY(meta = (DataflowOutput))
 	float ReturnValue = 0.f;
 
-	FPowDataflowNode(const Dataflow::FNodeParameters& InParam, FGuid InGuid = FGuid::NewGuid())
+	FPowDataflowNode(const UE::Dataflow::FNodeParameters& InParam, FGuid InGuid = FGuid::NewGuid())
 		: FDataflowNode(InParam, InGuid)
 	{
 		RegisterInputConnection(&Base);
@@ -838,7 +838,7 @@ public:
 		RegisterOutputConnection(&ReturnValue);
 	}
 
-	virtual void Evaluate(Dataflow::FContext& Context, const FDataflowOutput* Out) const override;
+	virtual void Evaluate(UE::Dataflow::FContext& Context, const FDataflowOutput* Out) const override;
 };
 
 /**
@@ -862,7 +862,7 @@ public:
 	UPROPERTY(meta = (DataflowOutput))
 	float ReturnValue = 0.f;
 
-	FLogDataflowNode(const Dataflow::FNodeParameters& InParam, FGuid InGuid = FGuid::NewGuid())
+	FLogDataflowNode(const UE::Dataflow::FNodeParameters& InParam, FGuid InGuid = FGuid::NewGuid())
 		: FDataflowNode(InParam, InGuid)
 	{
 		RegisterInputConnection(&A);
@@ -870,7 +870,7 @@ public:
 		RegisterOutputConnection(&ReturnValue);
 	}
 
-	virtual void Evaluate(Dataflow::FContext& Context, const FDataflowOutput* Out) const override;
+	virtual void Evaluate(UE::Dataflow::FContext& Context, const FDataflowOutput* Out) const override;
 };
 
 /**
@@ -891,14 +891,14 @@ public:
 	UPROPERTY(meta = (DataflowOutput))
 	float ReturnValue = 0.f;
 
-	FLogeDataflowNode(const Dataflow::FNodeParameters& InParam, FGuid InGuid = FGuid::NewGuid())
+	FLogeDataflowNode(const UE::Dataflow::FNodeParameters& InParam, FGuid InGuid = FGuid::NewGuid())
 		: FDataflowNode(InParam, InGuid)
 	{
 		RegisterInputConnection(&A);
 		RegisterOutputConnection(&ReturnValue);
 	}
 
-	virtual void Evaluate(Dataflow::FContext& Context, const FDataflowOutput* Out) const override;
+	virtual void Evaluate(UE::Dataflow::FContext& Context, const FDataflowOutput* Out) const override;
 };
 
 /**
@@ -925,7 +925,7 @@ public:
 	UPROPERTY(meta = (DataflowOutput))
 	float ReturnValue = 0.f;
 
-	FLerpDataflowNode(const Dataflow::FNodeParameters& InParam, FGuid InGuid = FGuid::NewGuid())
+	FLerpDataflowNode(const UE::Dataflow::FNodeParameters& InParam, FGuid InGuid = FGuid::NewGuid())
 		: FDataflowNode(InParam, InGuid)
 	{
 		RegisterInputConnection(&A);
@@ -934,7 +934,7 @@ public:
 		RegisterOutputConnection(&ReturnValue);
 	}
 
-	virtual void Evaluate(Dataflow::FContext& Context, const FDataflowOutput* Out) const override;
+	virtual void Evaluate(UE::Dataflow::FContext& Context, const FDataflowOutput* Out) const override;
 };
 
 /**
@@ -961,7 +961,7 @@ public:
 	UPROPERTY(meta = (DataflowOutput))
 	float ReturnValue = 0.f;
 
-	FWrapDataflowNode(const Dataflow::FNodeParameters& InParam, FGuid InGuid = FGuid::NewGuid())
+	FWrapDataflowNode(const UE::Dataflow::FNodeParameters& InParam, FGuid InGuid = FGuid::NewGuid())
 		: FDataflowNode(InParam, InGuid)
 	{
 		RegisterInputConnection(&Float);
@@ -970,7 +970,7 @@ public:
 		RegisterOutputConnection(&ReturnValue);
 	}
 
-	virtual void Evaluate(Dataflow::FContext& Context, const FDataflowOutput* Out) const override;
+	virtual void Evaluate(UE::Dataflow::FContext& Context, const FDataflowOutput* Out) const override;
 };
 
 /**
@@ -991,14 +991,14 @@ public:
 	UPROPERTY(meta = (DataflowOutput))
 	float ReturnValue = 0.f;
 
-	FExpDataflowNode(const Dataflow::FNodeParameters& InParam, FGuid InGuid = FGuid::NewGuid())
+	FExpDataflowNode(const UE::Dataflow::FNodeParameters& InParam, FGuid InGuid = FGuid::NewGuid())
 		: FDataflowNode(InParam, InGuid)
 	{
 		RegisterInputConnection(&Float);
 		RegisterOutputConnection(&ReturnValue);
 	}
 
-	virtual void Evaluate(Dataflow::FContext& Context, const FDataflowOutput* Out) const override;
+	virtual void Evaluate(UE::Dataflow::FContext& Context, const FDataflowOutput* Out) const override;
 };
 
 /**
@@ -1019,14 +1019,14 @@ public:
 	UPROPERTY(meta = (DataflowOutput))
 	float ReturnValue = 0.f;
 
-	FSinDataflowNode(const Dataflow::FNodeParameters& InParam, FGuid InGuid = FGuid::NewGuid())
+	FSinDataflowNode(const UE::Dataflow::FNodeParameters& InParam, FGuid InGuid = FGuid::NewGuid())
 		: FDataflowNode(InParam, InGuid)
 	{
 		RegisterInputConnection(&Float);
 		RegisterOutputConnection(&ReturnValue);
 	}
 
-	virtual void Evaluate(Dataflow::FContext& Context, const FDataflowOutput* Out) const override;
+	virtual void Evaluate(UE::Dataflow::FContext& Context, const FDataflowOutput* Out) const override;
 };
 
 /**
@@ -1047,14 +1047,14 @@ public:
 	UPROPERTY(meta = (DataflowOutput))
 	float ReturnValue = 0.f;
 
-	FArcSinDataflowNode(const Dataflow::FNodeParameters& InParam, FGuid InGuid = FGuid::NewGuid())
+	FArcSinDataflowNode(const UE::Dataflow::FNodeParameters& InParam, FGuid InGuid = FGuid::NewGuid())
 		: FDataflowNode(InParam, InGuid)
 	{
 		RegisterInputConnection(&Float);
 		RegisterOutputConnection(&ReturnValue);
 	}
 
-	virtual void Evaluate(Dataflow::FContext& Context, const FDataflowOutput* Out) const override;
+	virtual void Evaluate(UE::Dataflow::FContext& Context, const FDataflowOutput* Out) const override;
 };
 
 /**
@@ -1075,14 +1075,14 @@ public:
 	UPROPERTY(meta = (DataflowOutput))
 	float ReturnValue = 0.f;
 
-	FCosDataflowNode(const Dataflow::FNodeParameters& InParam, FGuid InGuid = FGuid::NewGuid())
+	FCosDataflowNode(const UE::Dataflow::FNodeParameters& InParam, FGuid InGuid = FGuid::NewGuid())
 		: FDataflowNode(InParam, InGuid)
 	{
 		RegisterInputConnection(&Float);
 		RegisterOutputConnection(&ReturnValue);
 	}
 
-	virtual void Evaluate(Dataflow::FContext& Context, const FDataflowOutput* Out) const override;
+	virtual void Evaluate(UE::Dataflow::FContext& Context, const FDataflowOutput* Out) const override;
 };
 
 /**
@@ -1103,14 +1103,14 @@ public:
 	UPROPERTY(meta = (DataflowOutput))
 	float ReturnValue = 0.f;
 
-	FArcCosDataflowNode(const Dataflow::FNodeParameters& InParam, FGuid InGuid = FGuid::NewGuid())
+	FArcCosDataflowNode(const UE::Dataflow::FNodeParameters& InParam, FGuid InGuid = FGuid::NewGuid())
 		: FDataflowNode(InParam, InGuid)
 	{
 		RegisterInputConnection(&Float);
 		RegisterOutputConnection(&ReturnValue);
 	}
 
-	virtual void Evaluate(Dataflow::FContext& Context, const FDataflowOutput* Out) const override;
+	virtual void Evaluate(UE::Dataflow::FContext& Context, const FDataflowOutput* Out) const override;
 };
 
 /**
@@ -1131,14 +1131,14 @@ public:
 	UPROPERTY(meta = (DataflowOutput))
 	float ReturnValue = 0.f;
 
-	FTanDataflowNode(const Dataflow::FNodeParameters& InParam, FGuid InGuid = FGuid::NewGuid())
+	FTanDataflowNode(const UE::Dataflow::FNodeParameters& InParam, FGuid InGuid = FGuid::NewGuid())
 		: FDataflowNode(InParam, InGuid)
 	{
 		RegisterInputConnection(&Float);
 		RegisterOutputConnection(&ReturnValue);
 	}
 
-	virtual void Evaluate(Dataflow::FContext& Context, const FDataflowOutput* Out) const override;
+	virtual void Evaluate(UE::Dataflow::FContext& Context, const FDataflowOutput* Out) const override;
 };
 
 /**
@@ -1159,14 +1159,14 @@ public:
 	UPROPERTY(meta = (DataflowOutput))
 	float ReturnValue = 0.f;
 
-	FArcTanDataflowNode(const Dataflow::FNodeParameters& InParam, FGuid InGuid = FGuid::NewGuid())
+	FArcTanDataflowNode(const UE::Dataflow::FNodeParameters& InParam, FGuid InGuid = FGuid::NewGuid())
 		: FDataflowNode(InParam, InGuid)
 	{
 		RegisterInputConnection(&Float);
 		RegisterOutputConnection(&ReturnValue);
 	}
 
-	virtual void Evaluate(Dataflow::FContext& Context, const FDataflowOutput* Out) const override;
+	virtual void Evaluate(UE::Dataflow::FContext& Context, const FDataflowOutput* Out) const override;
 };
 
 /**
@@ -1190,7 +1190,7 @@ public:
 	UPROPERTY(meta = (DataflowOutput))
 	float ReturnValue = 0.f;
 
-	FArcTan2DataflowNode(const Dataflow::FNodeParameters& InParam, FGuid InGuid = FGuid::NewGuid())
+	FArcTan2DataflowNode(const UE::Dataflow::FNodeParameters& InParam, FGuid InGuid = FGuid::NewGuid())
 		: FDataflowNode(InParam, InGuid)
 	{
 		RegisterInputConnection(&Y);
@@ -1198,7 +1198,7 @@ public:
 		RegisterOutputConnection(&ReturnValue);
 	}
 
-	virtual void Evaluate(Dataflow::FContext& Context, const FDataflowOutput* Out) const override;
+	virtual void Evaluate(UE::Dataflow::FContext& Context, const FDataflowOutput* Out) const override;
 };
 
 /**
@@ -1225,7 +1225,7 @@ public:
 	UPROPERTY(meta = (DataflowOutput))
 	float ReturnValue = 0.f;
 
-	FNormalizeToRangeDataflowNode(const Dataflow::FNodeParameters& InParam, FGuid InGuid = FGuid::NewGuid())
+	FNormalizeToRangeDataflowNode(const UE::Dataflow::FNodeParameters& InParam, FGuid InGuid = FGuid::NewGuid())
 		: FDataflowNode(InParam, InGuid)
 	{
 		RegisterInputConnection(&Float);
@@ -1234,7 +1234,7 @@ public:
 		RegisterOutputConnection(&ReturnValue);
 	}
 
-	virtual void Evaluate(Dataflow::FContext& Context, const FDataflowOutput* Out) const override;
+	virtual void Evaluate(UE::Dataflow::FContext& Context, const FDataflowOutput* Out) const override;
 };
 
 /**
@@ -1258,7 +1258,7 @@ public:
 	UPROPERTY(meta = (DataflowOutput))
 	FVector ScaledVector = FVector(0.f);
 
-	FScaleVectorDataflowNode(const Dataflow::FNodeParameters& InParam, FGuid InGuid = FGuid::NewGuid())
+	FScaleVectorDataflowNode(const UE::Dataflow::FNodeParameters& InParam, FGuid InGuid = FGuid::NewGuid())
 		: FDataflowNode(InParam, InGuid)
 	{
 		RegisterInputConnection(&Vector);
@@ -1266,7 +1266,7 @@ public:
 		RegisterOutputConnection(&ScaledVector);
 	}
 
-	virtual void Evaluate(Dataflow::FContext& Context, const FDataflowOutput* Out) const override;
+	virtual void Evaluate(UE::Dataflow::FContext& Context, const FDataflowOutput* Out) const override;
 };
 
 /**
@@ -1290,7 +1290,7 @@ public:
 	UPROPERTY(meta = (DataflowOutput))
 	float ReturnValue = 0.f;
 
-	FDotProductDataflowNode(const Dataflow::FNodeParameters& InParam, FGuid InGuid = FGuid::NewGuid())
+	FDotProductDataflowNode(const UE::Dataflow::FNodeParameters& InParam, FGuid InGuid = FGuid::NewGuid())
 		: FDataflowNode(InParam, InGuid)
 	{
 		RegisterInputConnection(&VectorA);
@@ -1298,7 +1298,7 @@ public:
 		RegisterOutputConnection(&ReturnValue);
 	}
 
-	virtual void Evaluate(Dataflow::FContext& Context, const FDataflowOutput* Out) const override;
+	virtual void Evaluate(UE::Dataflow::FContext& Context, const FDataflowOutput* Out) const override;
 };
 
 /**
@@ -1322,7 +1322,7 @@ public:
 	UPROPERTY(meta = (DataflowOutput))
 	FVector ReturnValue = FVector(0.f);
 
-	FCrossProductDataflowNode(const Dataflow::FNodeParameters& InParam, FGuid InGuid = FGuid::NewGuid())
+	FCrossProductDataflowNode(const UE::Dataflow::FNodeParameters& InParam, FGuid InGuid = FGuid::NewGuid())
 		: FDataflowNode(InParam, InGuid)
 	{
 		RegisterInputConnection(&VectorA);
@@ -1330,7 +1330,7 @@ public:
 		RegisterOutputConnection(&ReturnValue);
 	}
 
-	virtual void Evaluate(Dataflow::FContext& Context, const FDataflowOutput* Out) const override;
+	virtual void Evaluate(UE::Dataflow::FContext& Context, const FDataflowOutput* Out) const override;
 };
 
 /**
@@ -1354,7 +1354,7 @@ public:
 	UPROPERTY(meta = (DataflowOutput))
 	FVector ReturnValue = FVector(0.f);
 
-	FNormalizeDataflowNode(const Dataflow::FNodeParameters& InParam, FGuid InGuid = FGuid::NewGuid())
+	FNormalizeDataflowNode(const UE::Dataflow::FNodeParameters& InParam, FGuid InGuid = FGuid::NewGuid())
 		: FDataflowNode(InParam, InGuid)
 	{
 		RegisterInputConnection(&VectorA);
@@ -1362,7 +1362,7 @@ public:
 		RegisterOutputConnection(&ReturnValue);
 	}
 
-	virtual void Evaluate(Dataflow::FContext& Context, const FDataflowOutput* Out) const override;
+	virtual void Evaluate(UE::Dataflow::FContext& Context, const FDataflowOutput* Out) const override;
 };
 
 /**
@@ -1383,14 +1383,14 @@ public:
 	UPROPERTY(meta = (DataflowOutput))
 	float ReturnValue = 0.f;
 
-	FLengthDataflowNode(const Dataflow::FNodeParameters& InParam, FGuid InGuid = FGuid::NewGuid())
+	FLengthDataflowNode(const UE::Dataflow::FNodeParameters& InParam, FGuid InGuid = FGuid::NewGuid())
 		: FDataflowNode(InParam, InGuid)
 	{
 		RegisterInputConnection(&Vector);
 		RegisterOutputConnection(&ReturnValue);
 	}
 
-	virtual void Evaluate(Dataflow::FContext& Context, const FDataflowOutput* Out) const override;
+	virtual void Evaluate(UE::Dataflow::FContext& Context, const FDataflowOutput* Out) const override;
 };
 
 /**
@@ -1414,7 +1414,7 @@ public:
 	UPROPERTY(meta = (DataflowOutput))
 	float ReturnValue = 0.f;
 
-	FDistanceDataflowNode(const Dataflow::FNodeParameters& InParam, FGuid InGuid = FGuid::NewGuid())
+	FDistanceDataflowNode(const UE::Dataflow::FNodeParameters& InParam, FGuid InGuid = FGuid::NewGuid())
 		: FDataflowNode(InParam, InGuid)
 	{
 		RegisterInputConnection(&PointA);
@@ -1422,7 +1422,7 @@ public:
 		RegisterOutputConnection(&ReturnValue);
 	}
 
-	virtual void Evaluate(Dataflow::FContext& Context, const FDataflowOutput* Out) const override;
+	virtual void Evaluate(UE::Dataflow::FContext& Context, const FDataflowOutput* Out) const override;
 };
 
 /**
@@ -1443,14 +1443,14 @@ public:
 	UPROPERTY(meta = (DataflowOutput))
 	bool ReturnValue = false;
 
-	FIsNearlyZeroDataflowNode(const Dataflow::FNodeParameters& InParam, FGuid InGuid = FGuid::NewGuid())
+	FIsNearlyZeroDataflowNode(const UE::Dataflow::FNodeParameters& InParam, FGuid InGuid = FGuid::NewGuid())
 		: FDataflowNode(InParam, InGuid)
 	{
 		RegisterInputConnection(&Float);
 		RegisterOutputConnection(&ReturnValue);
 	}
 
-	virtual void Evaluate(Dataflow::FContext& Context, const FDataflowOutput* Out) const override;
+	virtual void Evaluate(UE::Dataflow::FContext& Context, const FDataflowOutput* Out) const override;
 };
 
 /**
@@ -1474,7 +1474,7 @@ public:
 	UPROPERTY(meta = (DataflowOutput))
 	float ReturnValue = 0.f;
 
-	FRandomFloatDataflowNode(const Dataflow::FNodeParameters& InParam, FGuid InGuid = FGuid::NewGuid())
+	FRandomFloatDataflowNode(const UE::Dataflow::FNodeParameters& InParam, FGuid InGuid = FGuid::NewGuid())
 		: FDataflowNode(InParam, InGuid)
 	{
 		RandomSeed = FMath::FRandRange(-1e5, 1e5);
@@ -1482,7 +1482,7 @@ public:
 		RegisterOutputConnection(&ReturnValue);
 	}
 
-	virtual void Evaluate(Dataflow::FContext& Context, const FDataflowOutput* Out) const override;
+	virtual void Evaluate(UE::Dataflow::FContext& Context, const FDataflowOutput* Out) const override;
 
 };
 
@@ -1513,7 +1513,7 @@ public:
 	UPROPERTY(meta = (DataflowOutput))
 	float ReturnValue = 0.f;
 
-	FRandomFloatInRangeDataflowNode(const Dataflow::FNodeParameters& InParam, FGuid InGuid = FGuid::NewGuid())
+	FRandomFloatInRangeDataflowNode(const UE::Dataflow::FNodeParameters& InParam, FGuid InGuid = FGuid::NewGuid())
 		: FDataflowNode(InParam, InGuid)
 	{
 		RandomSeed = FMath::FRandRange(-1e5, 1e5);
@@ -1523,7 +1523,7 @@ public:
 		RegisterOutputConnection(&ReturnValue);
 	}
 
-	virtual void Evaluate(Dataflow::FContext& Context, const FDataflowOutput* Out) const override;
+	virtual void Evaluate(UE::Dataflow::FContext& Context, const FDataflowOutput* Out) const override;
 
 };
 
@@ -1549,7 +1549,7 @@ public:
 	UPROPERTY(meta = (DataflowOutput))
 	FVector ReturnValue = FVector(0.0);
 
-	FRandomUnitVectorDataflowNode(const Dataflow::FNodeParameters& InParam, FGuid InGuid = FGuid::NewGuid())
+	FRandomUnitVectorDataflowNode(const UE::Dataflow::FNodeParameters& InParam, FGuid InGuid = FGuid::NewGuid())
 		: FDataflowNode(InParam, InGuid)
 	{
 		RandomSeed = FMath::FRandRange(-1e5, 1e5);
@@ -1557,7 +1557,7 @@ public:
 		RegisterOutputConnection(&ReturnValue);
 	}
 
-	virtual void Evaluate(Dataflow::FContext& Context, const FDataflowOutput* Out) const override;
+	virtual void Evaluate(UE::Dataflow::FContext& Context, const FDataflowOutput* Out) const override;
 
 };
 
@@ -1591,7 +1591,7 @@ public:
 	UPROPERTY(meta = (DataflowOutput))
 	FVector ReturnValue = FVector(0.0);
 
-	FRandomUnitVectorInConeDataflowNode(const Dataflow::FNodeParameters& InParam, FGuid InGuid = FGuid::NewGuid())
+	FRandomUnitVectorInConeDataflowNode(const UE::Dataflow::FNodeParameters& InParam, FGuid InGuid = FGuid::NewGuid())
 		: FDataflowNode(InParam, InGuid)
 	{
 		RandomSeed = FMath::FRandRange(-1e5, 1e5);
@@ -1601,7 +1601,7 @@ public:
 		RegisterOutputConnection(&ReturnValue);
 	}
 
-	virtual void Evaluate(Dataflow::FContext& Context, const FDataflowOutput* Out) const override;
+	virtual void Evaluate(UE::Dataflow::FContext& Context, const FDataflowOutput* Out) const override;
 
 };
 
@@ -1624,14 +1624,14 @@ public:
 	UPROPERTY(meta = (DataflowOutput))
 	float Degrees = 0.f;
 
-	FRadiansToDegreesDataflowNode(const Dataflow::FNodeParameters& InParam, FGuid InGuid = FGuid::NewGuid())
+	FRadiansToDegreesDataflowNode(const UE::Dataflow::FNodeParameters& InParam, FGuid InGuid = FGuid::NewGuid())
 		: FDataflowNode(InParam, InGuid)
 	{
 		RegisterInputConnection(&Radians);
 		RegisterOutputConnection(&Degrees);
 	}
 
-	virtual void Evaluate(Dataflow::FContext& Context, const FDataflowOutput* Out) const override;
+	virtual void Evaluate(UE::Dataflow::FContext& Context, const FDataflowOutput* Out) const override;
 
 };
 
@@ -1654,14 +1654,14 @@ public:
 	UPROPERTY(meta = (DataflowOutput))
 	float Radians = 0.f;
 
-	FDegreesToRadiansDataflowNode(const Dataflow::FNodeParameters& InParam, FGuid InGuid = FGuid::NewGuid())
+	FDegreesToRadiansDataflowNode(const UE::Dataflow::FNodeParameters& InParam, FGuid InGuid = FGuid::NewGuid())
 		: FDataflowNode(InParam, InGuid)
 	{
 		RegisterInputConnection(&Degrees);
 		RegisterOutputConnection(&Radians);
 	}
 
-	virtual void Evaluate(Dataflow::FContext& Context, const FDataflowOutput* Out) const override;
+	virtual void Evaluate(UE::Dataflow::FContext& Context, const FDataflowOutput* Out) const override;
 
 };
 
@@ -1707,13 +1707,13 @@ public:
 	UPROPERTY(meta = (DataflowOutput))
 	float ReturnValue = 0;
 
-	FMathConstantsDataflowNode(const Dataflow::FNodeParameters& InParam, FGuid InGuid = FGuid::NewGuid())
+	FMathConstantsDataflowNode(const UE::Dataflow::FNodeParameters& InParam, FGuid InGuid = FGuid::NewGuid())
 		: FDataflowNode(InParam, InGuid)
 	{
 		RegisterOutputConnection(&ReturnValue);
 	}
 
-	virtual void Evaluate(Dataflow::FContext& Context, const FDataflowOutput* Out) const override;
+	virtual void Evaluate(UE::Dataflow::FContext& Context, const FDataflowOutput* Out) const override;
 
 };
 
@@ -1735,14 +1735,14 @@ public:
 	UPROPERTY(meta = (DataflowOutput))
 	float ReturnValue = 0.f;
 
-	FOneMinusDataflowNode(const Dataflow::FNodeParameters& InParam, FGuid InGuid = FGuid::NewGuid())
+	FOneMinusDataflowNode(const UE::Dataflow::FNodeParameters& InParam, FGuid InGuid = FGuid::NewGuid())
 		: FDataflowNode(InParam, InGuid)
 	{
 		RegisterInputConnection(&A);
 		RegisterOutputConnection(&ReturnValue);
 	}
 
-	virtual void Evaluate(Dataflow::FContext& Context, const FDataflowOutput* Out) const override;
+	virtual void Evaluate(UE::Dataflow::FContext& Context, const FDataflowOutput* Out) const override;
 };
 
 /**
@@ -1775,9 +1775,9 @@ public:
 	UPROPERTY(meta = (DataflowOutput))
 	float ReturnValue = 0.f;
 
-	FFloatMathExpressionDataflowNode(const Dataflow::FNodeParameters& InParam, FGuid InGuid = FGuid::NewGuid());
+	FFloatMathExpressionDataflowNode(const UE::Dataflow::FNodeParameters& InParam, FGuid InGuid = FGuid::NewGuid());
 
-	virtual void Evaluate(Dataflow::FContext& Context, const FDataflowOutput* Out) const override;
+	virtual void Evaluate(UE::Dataflow::FContext& Context, const FDataflowOutput* Out) const override;
 };
 
 /**
@@ -1810,12 +1810,12 @@ public:
 	UPROPERTY(meta = (DataflowOutput))
 	FDataflowNumericTypes ReturnValue;
 
-	FMathExpressionDataflowNode(const Dataflow::FNodeParameters& InParam, FGuid InGuid = FGuid::NewGuid());
+	FMathExpressionDataflowNode(const UE::Dataflow::FNodeParameters& InParam, FGuid InGuid = FGuid::NewGuid());
 
-	virtual void Evaluate(Dataflow::FContext& Context, const FDataflowOutput* Out) const override;
+	virtual void Evaluate(UE::Dataflow::FContext& Context, const FDataflowOutput* Out) const override;
 };
 
-namespace Dataflow
+namespace UE::Dataflow
 {
 	void GeometryCollectionMathNodes();
 }
diff --git a/Engine/Plugins/Experimental/GeometryCollectionPlugin/Source/GeometryCollectionNodes/Public/Dataflow/GeometryCollectionMeshNodes.h b/Engine/Plugins/Experimental/GeometryCollectionPlugin/Source/GeometryCollectionNodes/Public/Dataflow/GeometryCollectionMeshNodes.h
--- a/Engine/Plugins/Experimental/GeometryCollectionPlugin/Source/GeometryCollectionNodes/Public/Dataflow/GeometryCollectionMeshNodes.h
+++ b/Engine/Plugins/Experimental/GeometryCollectionPlugin/Source/GeometryCollectionNodes/Public/Dataflow/GeometryCollectionMeshNodes.h
@@ -38,7 +38,7 @@ public:
 	UPROPERTY(meta = (DataflowOutput))
 	int32 TriangleCount = 0;
 
-	FPointsToMeshDataflowNode(const Dataflow::FNodeParameters& InParam, FGuid InGuid = FGuid::NewGuid())
+	FPointsToMeshDataflowNode(const UE::Dataflow::FNodeParameters& InParam, FGuid InGuid = FGuid::NewGuid())
 		: FDataflowNode(InParam, InGuid)
 	{
 		RegisterInputConnection(&Points);
@@ -46,7 +46,7 @@ public:
 		RegisterOutputConnection(&TriangleCount);
 	}
 
-	virtual void Evaluate(Dataflow::FContext& Context, const FDataflowOutput* Out) const override;
+	virtual void Evaluate(UE::Dataflow::FContext& Context, const FDataflowOutput* Out) const override;
 
 };
 
@@ -76,7 +76,7 @@ public:
 	UPROPERTY(meta = (DataflowOutput))
 	int32 TriangleCount = 0;
 
-	FBoxToMeshDataflowNode(const Dataflow::FNodeParameters& InParam, FGuid InGuid = FGuid::NewGuid())
+	FBoxToMeshDataflowNode(const UE::Dataflow::FNodeParameters& InParam, FGuid InGuid = FGuid::NewGuid())
 		: FDataflowNode(InParam, InGuid)
 	{
 		RegisterInputConnection(&Box);
@@ -84,7 +84,7 @@ public:
 		RegisterOutputConnection(&TriangleCount);
 	}
 
-	virtual void Evaluate(Dataflow::FContext& Context, const FDataflowOutput* Out) const override;
+	virtual void Evaluate(UE::Dataflow::FContext& Context, const FDataflowOutput* Out) const override;
 
 };
 
@@ -109,14 +109,14 @@ public:
 	UPROPERTY(meta = (DataflowOutput))
 	FString InfoString = FString("");
 
-	FMeshInfoDataflowNode(const Dataflow::FNodeParameters& InParam, FGuid InGuid = FGuid::NewGuid())
+	FMeshInfoDataflowNode(const UE::Dataflow::FNodeParameters& InParam, FGuid InGuid = FGuid::NewGuid())
 		: FDataflowNode(InParam, InGuid)
 	{
 		RegisterInputConnection(&Mesh);
 		RegisterOutputConnection(&InfoString);
 	}
 
-	virtual void Evaluate(Dataflow::FContext& Context, const FDataflowOutput* Out) const override;
+	virtual void Evaluate(UE::Dataflow::FContext& Context, const FDataflowOutput* Out) const override;
 
 };
 
@@ -142,14 +142,14 @@ public:
 	UPROPERTY(meta = (DataflowOutput))
 	FManagedArrayCollection Collection;
 
-	FMeshToCollectionDataflowNode(const Dataflow::FNodeParameters& InParam, FGuid InGuid = FGuid::NewGuid())
+	FMeshToCollectionDataflowNode(const UE::Dataflow::FNodeParameters& InParam, FGuid InGuid = FGuid::NewGuid())
 		: FDataflowNode(InParam, InGuid)
 	{
 		RegisterInputConnection(&Mesh);
 		RegisterOutputConnection(&Collection);
 	}
 
-	virtual void Evaluate(Dataflow::FContext& Context, const FDataflowOutput* Out) const override;
+	virtual void Evaluate(UE::Dataflow::FContext& Context, const FDataflowOutput* Out) const override;
 
 };
 
@@ -178,14 +178,14 @@ public:
 	UPROPERTY(meta = (DataflowOutput))
 	TObjectPtr<UDynamicMesh> Mesh;
 
-	FCollectionToMeshDataflowNode(const Dataflow::FNodeParameters& InParam, FGuid InGuid = FGuid::NewGuid())
+	FCollectionToMeshDataflowNode(const UE::Dataflow::FNodeParameters& InParam, FGuid InGuid = FGuid::NewGuid())
 		: FDataflowNode(InParam, InGuid)
 	{
 		RegisterInputConnection(&Collection);
 		RegisterOutputConnection(&Mesh);
 	}
 
-	virtual void Evaluate(Dataflow::FContext& Context, const FDataflowOutput* Out) const override;
+	virtual void Evaluate(UE::Dataflow::FContext& Context, const FDataflowOutput* Out) const override;
 
 };
 
@@ -219,13 +219,13 @@ public:
 	UPROPERTY(meta = (DataflowOutput))
 	TObjectPtr<UDynamicMesh> Mesh;
 
-	FStaticMeshToMeshDataflowNode(const Dataflow::FNodeParameters& InParam, FGuid InGuid = FGuid::NewGuid())
+	FStaticMeshToMeshDataflowNode(const UE::Dataflow::FNodeParameters& InParam, FGuid InGuid = FGuid::NewGuid())
 		: FDataflowNode(InParam, InGuid)
 	{
 		RegisterOutputConnection(&Mesh);
 	}
 
-	virtual void Evaluate(Dataflow::FContext& Context, const FDataflowOutput* Out) const override;
+	virtual void Evaluate(UE::Dataflow::FContext& Context, const FDataflowOutput* Out) const override;
 
 };
 
@@ -255,7 +255,7 @@ public:
 	UPROPERTY(meta = (DataflowOutput))
 	TObjectPtr<UDynamicMesh> Mesh;
 
-	FMeshAppendDataflowNode(const Dataflow::FNodeParameters& InParam, FGuid InGuid = FGuid::NewGuid())
+	FMeshAppendDataflowNode(const UE::Dataflow::FNodeParameters& InParam, FGuid InGuid = FGuid::NewGuid())
 		: FDataflowNode(InParam, InGuid)
 	{
 		RegisterInputConnection(&Mesh1);
@@ -263,7 +263,7 @@ public:
 		RegisterOutputConnection(&Mesh);
 	}
 
-	virtual void Evaluate(Dataflow::FContext& Context, const FDataflowOutput* Out) const override;
+	virtual void Evaluate(UE::Dataflow::FContext& Context, const FDataflowOutput* Out) const override;
 
 };
 
@@ -308,7 +308,7 @@ public:
 	UPROPERTY(meta = (DataflowOutput))
 	TObjectPtr<UDynamicMesh> Mesh;
 
-	FMeshBooleanDataflowNode(const Dataflow::FNodeParameters& InParam, FGuid InGuid = FGuid::NewGuid())
+	FMeshBooleanDataflowNode(const UE::Dataflow::FNodeParameters& InParam, FGuid InGuid = FGuid::NewGuid())
 		: FDataflowNode(InParam, InGuid)
 	{
 		RegisterInputConnection(&Mesh1);
@@ -316,7 +316,7 @@ public:
 		RegisterOutputConnection(&Mesh);
 	}
 
-	virtual void Evaluate(Dataflow::FContext& Context, const FDataflowOutput* Out) const override;
+	virtual void Evaluate(UE::Dataflow::FContext& Context, const FDataflowOutput* Out) const override;
 
 };
 
@@ -350,7 +350,7 @@ public:
 	UPROPERTY(meta = (DataflowOutput))
 	TObjectPtr<UDynamicMesh> Mesh;
 
-	FMeshCopyToPointsDataflowNode(const Dataflow::FNodeParameters& InParam, FGuid InGuid = FGuid::NewGuid())
+	FMeshCopyToPointsDataflowNode(const UE::Dataflow::FNodeParameters& InParam, FGuid InGuid = FGuid::NewGuid())
 		: FDataflowNode(InParam, InGuid)
 	{
 		RegisterInputConnection(&Points);
@@ -358,7 +358,7 @@ public:
 		RegisterOutputConnection(&Mesh);
 	}
 
-	virtual void Evaluate(Dataflow::FContext& Context, const FDataflowOutput* Out) const override;
+	virtual void Evaluate(UE::Dataflow::FContext& Context, const FDataflowOutput* Out) const override;
 
 };
 
@@ -391,7 +391,7 @@ public:
 	UPROPERTY(meta = (DataflowOutput))
 	int32 TriangleCount = 0;
 
-	FGetMeshDataDataflowNode(const Dataflow::FNodeParameters& InParam, FGuid InGuid = FGuid::NewGuid())
+	FGetMeshDataDataflowNode(const UE::Dataflow::FNodeParameters& InParam, FGuid InGuid = FGuid::NewGuid())
 		: FDataflowNode(InParam, InGuid)
 	{
 		RegisterInputConnection(&Mesh);
@@ -400,12 +400,12 @@ public:
 		RegisterOutputConnection(&TriangleCount);
 	}
 
-	virtual void Evaluate(Dataflow::FContext& Context, const FDataflowOutput* Out) const override;
+	virtual void Evaluate(UE::Dataflow::FContext& Context, const FDataflowOutput* Out) const override;
 
 };
 
 
-namespace Dataflow
+namespace UE::Dataflow
 {
 	void GeometryCollectionMeshNodes();
 }
diff --git a/Engine/Plugins/Experimental/GeometryCollectionPlugin/Source/GeometryCollectionNodes/Public/Dataflow/GeometryCollectionNodes.h b/Engine/Plugins/Experimental/GeometryCollectionPlugin/Source/GeometryCollectionNodes/Public/Dataflow/GeometryCollectionNodes.h
--- a/Engine/Plugins/Experimental/GeometryCollectionPlugin/Source/GeometryCollectionNodes/Public/Dataflow/GeometryCollectionNodes.h
+++ b/Engine/Plugins/Experimental/GeometryCollectionPlugin/Source/GeometryCollectionNodes/Public/Dataflow/GeometryCollectionNodes.h
@@ -39,14 +39,14 @@ public:
 	UPROPERTY(meta = (DataflowOutput))
 	FManagedArrayCollection Collection;
 
-	FGetCollectionFromAssetDataflowNode(const Dataflow::FNodeParameters& InParam, FGuid InGuid = FGuid::NewGuid())
+	FGetCollectionFromAssetDataflowNode(const UE::Dataflow::FNodeParameters& InParam, FGuid InGuid = FGuid::NewGuid())
 		: FDataflowNode(InParam, InGuid)
 	{
 		RegisterInputConnection(&CollectionAsset);
 		RegisterOutputConnection(&Collection);
 	}
 
-	virtual void Evaluate(Dataflow::FContext& Context, const FDataflowOutput* Out) const override;
+	virtual void Evaluate(UE::Dataflow::FContext& Context, const FDataflowOutput* Out) const override;
 
 };
 
@@ -77,7 +77,7 @@ public:
 	UPROPERTY(meta = (DataflowOutput, DisplayName = "GeometryGroupIndicesOut2"))
 		TArray<FString> GeometryGroupGuidsOut2;
 
-	FAppendCollectionAssetsDataflowNode(const Dataflow::FNodeParameters& InParam, FGuid InGuid = FGuid::NewGuid())
+	FAppendCollectionAssetsDataflowNode(const UE::Dataflow::FNodeParameters& InParam, FGuid InGuid = FGuid::NewGuid())
 		: FDataflowNode(InParam, InGuid)
 	{
 		RegisterInputConnection(&Collection1);
@@ -87,7 +87,7 @@ public:
 		RegisterOutputConnection(&GeometryGroupGuidsOut2);
 	}
 
-	virtual void Evaluate(Dataflow::FContext& Context, const FDataflowOutput* Out) const override;
+	virtual void Evaluate(UE::Dataflow::FContext& Context, const FDataflowOutput* Out) const override;
 
 };
 
@@ -119,13 +119,13 @@ public:
 	UPROPERTY(EditAnywhere, Category = "Print", meta = (DataflowInput));
 	FString String = FString("");
 
-	FPrintStringDataflowNode(const Dataflow::FNodeParameters& InParam, FGuid InGuid = FGuid::NewGuid())
+	FPrintStringDataflowNode(const UE::Dataflow::FNodeParameters& InParam, FGuid InGuid = FGuid::NewGuid())
 		: FDataflowNode(InParam, InGuid)
 	{
 		RegisterInputConnection(&String);
 	}
 
-	virtual void Evaluate(Dataflow::FContext& Context, const FDataflowOutput* Out) const override;
+	virtual void Evaluate(UE::Dataflow::FContext& Context, const FDataflowOutput* Out) const override;
 
 };
 
@@ -148,13 +148,13 @@ public:
 	UPROPERTY(EditAnywhere, Category = "Print", meta = (DataflowInput));
 	FString String = FString("");
 
-	FLogStringDataflowNode(const Dataflow::FNodeParameters& InParam, FGuid InGuid = FGuid::NewGuid())
+	FLogStringDataflowNode(const UE::Dataflow::FNodeParameters& InParam, FGuid InGuid = FGuid::NewGuid())
 		: FDataflowNode(InParam, InGuid)
 	{
 		RegisterInputConnection(&String);
 	}
 
-	virtual void Evaluate(Dataflow::FContext& Context, const FDataflowOutput* Out) const override;
+	virtual void Evaluate(UE::Dataflow::FContext& Context, const FDataflowOutput* Out) const override;
 
 };
 
@@ -181,14 +181,14 @@ public:
 	UPROPERTY(meta = (DataflowOutput))
 	FBox BoundingBox = FBox(ForceInit);
 
-	FBoundingBoxDataflowNode(const Dataflow::FNodeParameters& InParam, FGuid InGuid = FGuid::NewGuid())
+	FBoundingBoxDataflowNode(const UE::Dataflow::FNodeParameters& InParam, FGuid InGuid = FGuid::NewGuid())
 		: FDataflowNode(InParam, InGuid)
 	{
 		RegisterInputConnection(&Collection);
 		RegisterOutputConnection(&BoundingBox);
 	}
 
-	virtual void Evaluate(Dataflow::FContext& Context, const FDataflowOutput* Out) const override;
+	virtual void Evaluate(UE::Dataflow::FContext& Context, const FDataflowOutput* Out) const override;
 
 };
 
@@ -247,14 +247,14 @@ public:
 		return Size.X;
 	}
 
-	FGetBoxLengthsDataflowNode(const Dataflow::FNodeParameters& InParam, FGuid InGuid = FGuid::NewGuid())
+	FGetBoxLengthsDataflowNode(const UE::Dataflow::FNodeParameters& InParam, FGuid InGuid = FGuid::NewGuid())
 		: FDataflowNode(InParam, InGuid)
 	{
 		RegisterInputConnection(&Boxes);
 		RegisterOutputConnection(&Lengths);
 	}
 
-	virtual void Evaluate(Dataflow::FContext& Context, const FDataflowOutput* Out) const override;
+	virtual void Evaluate(UE::Dataflow::FContext& Context, const FDataflowOutput* Out) const override;
 
 };
 
@@ -290,7 +290,7 @@ public:
 	UPROPERTY(meta = (DataflowOutput))
 	float Volume = 0.0;
 
-	FExpandBoundingBoxDataflowNode(const Dataflow::FNodeParameters& InParam, FGuid InGuid = FGuid::NewGuid())
+	FExpandBoundingBoxDataflowNode(const UE::Dataflow::FNodeParameters& InParam, FGuid InGuid = FGuid::NewGuid())
 		: FDataflowNode(InParam, InGuid)
 	{
 		RegisterInputConnection(&BoundingBox);
@@ -301,7 +301,7 @@ public:
 		RegisterOutputConnection(&Volume);
 	}
 
-	virtual void Evaluate(Dataflow::FContext& Context, const FDataflowOutput* Out) const override;
+	virtual void Evaluate(UE::Dataflow::FContext& Context, const FDataflowOutput* Out) const override;
 
 };
 
@@ -333,7 +333,7 @@ public:
 	float Z = 0.f;
 
 
-	FExpandVectorDataflowNode(const Dataflow::FNodeParameters& InParam, FGuid InGuid = FGuid::NewGuid())
+	FExpandVectorDataflowNode(const UE::Dataflow::FNodeParameters& InParam, FGuid InGuid = FGuid::NewGuid())
 		: FDataflowNode(InParam, InGuid)
 	{
 		RegisterInputConnection(&Vector);
@@ -342,7 +342,7 @@ public:
 		RegisterOutputConnection(&Z);
 	}
 
-	virtual void Evaluate(Dataflow::FContext& Context, const FDataflowOutput* Out) const override;
+	virtual void Evaluate(UE::Dataflow::FContext& Context, const FDataflowOutput* Out) const override;
 
 };
 
@@ -367,7 +367,7 @@ public:
 	UPROPERTY(meta = (DataflowOutput))
 	FString String = FString("");
 
-	FStringAppendDataflowNode(const Dataflow::FNodeParameters& InParam, FGuid InGuid = FGuid::NewGuid())
+	FStringAppendDataflowNode(const UE::Dataflow::FNodeParameters& InParam, FGuid InGuid = FGuid::NewGuid())
 		: FDataflowNode(InParam, InGuid)
 	{
 		RegisterInputConnection(&String1);
@@ -375,7 +375,7 @@ public:
 		RegisterOutputConnection(&String);
 	}
 
-	virtual void Evaluate(Dataflow::FContext& Context, const FDataflowOutput* Out) const override;
+	virtual void Evaluate(UE::Dataflow::FContext& Context, const FDataflowOutput* Out) const override;
 
 };
 
@@ -399,14 +399,14 @@ public:
 	UPROPERTY(meta = (DataflowOutput))
 	int32 Hash = 0;
 
-	FHashStringDataflowNode(const Dataflow::FNodeParameters& InParam, FGuid InGuid = FGuid::NewGuid())
+	FHashStringDataflowNode(const UE::Dataflow::FNodeParameters& InParam, FGuid InGuid = FGuid::NewGuid())
 		: FDataflowNode(InParam, InGuid)
 	{
 		RegisterInputConnection(&String);
 		RegisterOutputConnection(&Hash);
 	}
 
-	virtual void Evaluate(Dataflow::FContext& Context, const FDataflowOutput* Out) const override;
+	virtual void Evaluate(UE::Dataflow::FContext& Context, const FDataflowOutput* Out) const override;
 
 };
 
@@ -431,14 +431,14 @@ public:
 	UPROPERTY(meta = (DataflowOutput))
 	int32 Hash = 0;
 
-	FHashVectorDataflowNode(const Dataflow::FNodeParameters& InParam, FGuid InGuid = FGuid::NewGuid())
+	FHashVectorDataflowNode(const UE::Dataflow::FNodeParameters& InParam, FGuid InGuid = FGuid::NewGuid())
 		: FDataflowNode(InParam, InGuid)
 	{
 		RegisterInputConnection(&Vector);
 		RegisterOutputConnection(&Hash);
 	}
 
-	virtual void Evaluate(Dataflow::FContext& Context, const FDataflowOutput* Out) const override;
+	virtual void Evaluate(UE::Dataflow::FContext& Context, const FDataflowOutput* Out) const override;
 
 };
 
@@ -467,7 +467,7 @@ public:
 	UPROPERTY(meta = (DataflowOutput))
 	TArray<FBox> BoundingBoxes;
 
-	FGetBoundingBoxesFromCollectionDataflowNode(const Dataflow::FNodeParameters& InParam, FGuid InGuid = FGuid::NewGuid())
+	FGetBoundingBoxesFromCollectionDataflowNode(const UE::Dataflow::FNodeParameters& InParam, FGuid InGuid = FGuid::NewGuid())
 		: FDataflowNode(InParam, InGuid)
 	{
 		RegisterInputConnection(&Collection);
@@ -475,7 +475,7 @@ public:
 		RegisterOutputConnection(&BoundingBoxes);
 	}
 
-	virtual void Evaluate(Dataflow::FContext& Context, const FDataflowOutput* Out) const override;
+	virtual void Evaluate(UE::Dataflow::FContext& Context, const FDataflowOutput* Out) const override;
 
 };
 
@@ -498,14 +498,14 @@ public:
 	UPROPERTY(meta = (DataflowOutput))
 	int32 RootIndex = INDEX_NONE;
 
-	FGetRootIndexFromCollectionDataflowNode(const Dataflow::FNodeParameters& InParam, FGuid InGuid = FGuid::NewGuid())
+	FGetRootIndexFromCollectionDataflowNode(const UE::Dataflow::FNodeParameters& InParam, FGuid InGuid = FGuid::NewGuid())
 		: FDataflowNode(InParam, InGuid)
 	{
 		RegisterInputConnection(&Collection);
 		RegisterOutputConnection(&RootIndex);
 	}
 
-	virtual void Evaluate(Dataflow::FContext& Context, const FDataflowOutput* Out) const override;
+	virtual void Evaluate(UE::Dataflow::FContext& Context, const FDataflowOutput* Out) const override;
 
 };
 
@@ -535,7 +535,7 @@ public:
 	UPROPERTY(meta = (DataflowOutput))
 	TArray<FVector> Centroids = TArray<FVector>();
 
-	FGetCentroidsFromCollectionDataflowNode(const Dataflow::FNodeParameters& InParam, FGuid InGuid = FGuid::NewGuid())
+	FGetCentroidsFromCollectionDataflowNode(const UE::Dataflow::FNodeParameters& InParam, FGuid InGuid = FGuid::NewGuid())
 		: FDataflowNode(InParam, InGuid)
 	{
 		RegisterInputConnection(&Collection);
@@ -543,7 +543,7 @@ public:
 		RegisterOutputConnection(&Centroids);
 	}
 
-	virtual void Evaluate(Dataflow::FContext& Context, const FDataflowOutput* Out) const override;
+	virtual void Evaluate(UE::Dataflow::FContext& Context, const FDataflowOutput* Out) const override;
 
 };
 
@@ -615,14 +615,14 @@ public:
 	UPROPERTY(EditAnywhere, Category = "General");
 	bool bInvertTransformation = false;
 
-	FTransformCollectionDataflowNode(const Dataflow::FNodeParameters& InParam, FGuid InGuid = FGuid::NewGuid())
+	FTransformCollectionDataflowNode(const UE::Dataflow::FNodeParameters& InParam, FGuid InGuid = FGuid::NewGuid())
 		: FDataflowNode(InParam, InGuid)
 	{
 		RegisterInputConnection(&Collection);
 		RegisterOutputConnection(&Collection, &Collection);
 	}
 
-	virtual void Evaluate(Dataflow::FContext& Context, const FDataflowOutput* Out) const override;
+	virtual void Evaluate(UE::Dataflow::FContext& Context, const FDataflowOutput* Out) const override;
 
 };
 
@@ -644,14 +644,14 @@ public:
 	UPROPERTY(meta = (DataflowInput, DataflowOutput, DataflowPassthrough = "Collection", DataflowIntrinsic))
 	FManagedArrayCollection Collection;
 
-	FBakeTransformsInCollectionDataflowNode(const Dataflow::FNodeParameters& InParam, FGuid InGuid = FGuid::NewGuid())
+	FBakeTransformsInCollectionDataflowNode(const UE::Dataflow::FNodeParameters& InParam, FGuid InGuid = FGuid::NewGuid())
 		: FDataflowNode(InParam, InGuid)
 	{
 		RegisterInputConnection(&Collection);
 		RegisterOutputConnection(&Collection, &Collection);
 	}
 
-	virtual void Evaluate(Dataflow::FContext& Context, const FDataflowOutput* Out) const override;
+	virtual void Evaluate(UE::Dataflow::FContext& Context, const FDataflowOutput* Out) const override;
 
 };
 
@@ -709,14 +709,14 @@ public:
 	UPROPERTY(EditAnywhere, Category = "General");
 	bool bInvertTransformation = false;
 
-	FTransformMeshDataflowNode(const Dataflow::FNodeParameters& InParam, FGuid InGuid = FGuid::NewGuid())
+	FTransformMeshDataflowNode(const UE::Dataflow::FNodeParameters& InParam, FGuid InGuid = FGuid::NewGuid())
 		: FDataflowNode(InParam, InGuid)
 	{
 		RegisterInputConnection(&Mesh);
 		RegisterOutputConnection(&Mesh, &Mesh);
 	}
 
-	virtual void Evaluate(Dataflow::FContext& Context, const FDataflowOutput* Out) const override;
+	virtual void Evaluate(UE::Dataflow::FContext& Context, const FDataflowOutput* Out) const override;
 
 };
 
@@ -762,7 +762,7 @@ public:
 	UPROPERTY(meta = (DataflowOutput));
 	bool Result = false;
 
-	FCompareIntDataflowNode(const Dataflow::FNodeParameters& InParam, FGuid InGuid = FGuid::NewGuid())
+	FCompareIntDataflowNode(const UE::Dataflow::FNodeParameters& InParam, FGuid InGuid = FGuid::NewGuid())
 		: FDataflowNode(InParam, InGuid)
 	{
 		RegisterInputConnection(&IntA);
@@ -770,7 +770,7 @@ public:
 		RegisterOutputConnection(&Result);
 	}
 
-	virtual void Evaluate(Dataflow::FContext& Context, const FDataflowOutput* Out) const override;
+	virtual void Evaluate(UE::Dataflow::FContext& Context, const FDataflowOutput* Out) const override;
 
 };
 
@@ -803,7 +803,7 @@ public:
 	UPROPERTY(meta = (DataflowOutput));
 	bool Result = false;
 
-	FCompareFloatDataflowNode(const Dataflow::FNodeParameters& InParam, FGuid InGuid = FGuid::NewGuid())
+	FCompareFloatDataflowNode(const UE::Dataflow::FNodeParameters& InParam, FGuid InGuid = FGuid::NewGuid())
 		: FDataflowNode(InParam, InGuid)
 	{
 		RegisterInputConnection(&FloatA);
@@ -811,7 +811,7 @@ public:
 		RegisterOutputConnection(&Result);
 	}
 
-	virtual void Evaluate(Dataflow::FContext& Context, const FDataflowOutput* Out) const override;
+	virtual void Evaluate(UE::Dataflow::FContext& Context, const FDataflowOutput* Out) const override;
 
 };
 
@@ -846,7 +846,7 @@ public:
 	UPROPERTY(meta = (DataflowOutput))
 	TObjectPtr<UDynamicMesh> Mesh;
 
-	FBranchMeshDataflowNode(const Dataflow::FNodeParameters& InParam, FGuid InGuid = FGuid::NewGuid())
+	FBranchMeshDataflowNode(const UE::Dataflow::FNodeParameters& InParam, FGuid InGuid = FGuid::NewGuid())
 		: FDataflowNode(InParam, InGuid)
 	{
 		RegisterInputConnection(&MeshA);
@@ -855,7 +855,7 @@ public:
 		RegisterOutputConnection(&Mesh);
 	}
 
-	virtual void Evaluate(Dataflow::FContext& Context, const FDataflowOutput* Out) const override;
+	virtual void Evaluate(UE::Dataflow::FContext& Context, const FDataflowOutput* Out) const override;
 
 };
 
@@ -888,7 +888,7 @@ public:
 	UPROPERTY(meta = (DataflowOutput))
 	FManagedArrayCollection ChosenCollection;
 
-	FBranchCollectionDataflowNode(const Dataflow::FNodeParameters& InParam, FGuid InGuid = FGuid::NewGuid())
+	FBranchCollectionDataflowNode(const UE::Dataflow::FNodeParameters& InParam, FGuid InGuid = FGuid::NewGuid())
 		: FDataflowNode(InParam, InGuid)
 	{
 		RegisterInputConnection(&TrueCollection);
@@ -897,7 +897,7 @@ public:
 		RegisterOutputConnection(&ChosenCollection);
 	}
 
-	virtual void Evaluate(Dataflow::FContext& Context, const FDataflowOutput* Out) const override;
+	virtual void Evaluate(UE::Dataflow::FContext& Context, const FDataflowOutput* Out) const override;
 
 };
 
@@ -923,14 +923,14 @@ public:
 	UPROPERTY(meta = (DataflowOutput))
 	FString String;
 
-	FGetSchemaDataflowNode(const Dataflow::FNodeParameters& InParam, FGuid InGuid = FGuid::NewGuid())
+	FGetSchemaDataflowNode(const UE::Dataflow::FNodeParameters& InParam, FGuid InGuid = FGuid::NewGuid())
 		: FDataflowNode(InParam, InGuid)
 	{
 		RegisterInputConnection(&Collection);
 		RegisterOutputConnection(&String);
 	}
 
-	virtual void Evaluate(Dataflow::FContext& Context, const FDataflowOutput* Out) const override;
+	virtual void Evaluate(UE::Dataflow::FContext& Context, const FDataflowOutput* Out) const override;
 
 };
 
@@ -966,7 +966,7 @@ public:
 	UPROPERTY(EditAnywhere, Category = "Removal", meta = (DataflowInput))
 	bool bClusterCrumbling = false;
 
-	FRemoveOnBreakDataflowNode(const Dataflow::FNodeParameters& InParam, FGuid InGuid = FGuid::NewGuid())
+	FRemoveOnBreakDataflowNode(const UE::Dataflow::FNodeParameters& InParam, FGuid InGuid = FGuid::NewGuid())
 		: FDataflowNode(InParam, InGuid)
 	{
 		RegisterInputConnection(&Collection);
@@ -978,7 +978,7 @@ public:
 		RegisterOutputConnection(&Collection, &Collection);
 	}
 
-	virtual void Evaluate(Dataflow::FContext& Context, const FDataflowOutput* Out) const override;
+	virtual void Evaluate(UE::Dataflow::FContext& Context, const FDataflowOutput* Out) const override;
 
 };
 
@@ -1021,7 +1021,7 @@ public:
 	UPROPERTY(meta = (DataflowInput, DisplayName = "TransformSelection", DataflowIntrinsic))
 	FDataflowTransformSelection TransformSelection;
 
-	FSetAnchorStateDataflowNode(const Dataflow::FNodeParameters& InParam, FGuid InGuid = FGuid::NewGuid())
+	FSetAnchorStateDataflowNode(const UE::Dataflow::FNodeParameters& InParam, FGuid InGuid = FGuid::NewGuid())
 		: FDataflowNode(InParam, InGuid)
 	{
 		RegisterInputConnection(&Collection);
@@ -1029,7 +1029,7 @@ public:
 		RegisterOutputConnection(&Collection, &Collection);
 	}
 
-	virtual void Evaluate(Dataflow::FContext& Context, const FDataflowOutput* Out) const override;
+	virtual void Evaluate(UE::Dataflow::FContext& Context, const FDataflowOutput* Out) const override;
 
 };
 
@@ -1119,7 +1119,7 @@ public:
 	UPROPERTY(meta = (DataflowInput, DataflowOutput, DataflowPassthrough = "Collection", DataflowIntrinsic))
 	FManagedArrayCollection Collection;
 
-	FProximityDataflowNode(const Dataflow::FNodeParameters& InParam, FGuid InGuid = FGuid::NewGuid())
+	FProximityDataflowNode(const UE::Dataflow::FNodeParameters& InParam, FGuid InGuid = FGuid::NewGuid())
 		: FDataflowNode(InParam, InGuid)
 	{
 		RegisterInputConnection(&Collection);
@@ -1128,7 +1128,7 @@ public:
 		RegisterOutputConnection(&Collection, &Collection);
 	}
 
-	virtual void Evaluate(Dataflow::FContext& Context, const FDataflowOutput* Out) const override;
+	virtual void Evaluate(UE::Dataflow::FContext& Context, const FDataflowOutput* Out) const override;
 
 };
 
@@ -1153,7 +1153,7 @@ public:
 	UPROPERTY(EditAnywhere, Category = "Pivot", meta = (DataflowInput))
 	FTransform Transform;
 
-	FCollectionSetPivotDataflowNode(const Dataflow::FNodeParameters& InParam, FGuid InGuid = FGuid::NewGuid())
+	FCollectionSetPivotDataflowNode(const UE::Dataflow::FNodeParameters& InParam, FGuid InGuid = FGuid::NewGuid())
 		: FDataflowNode(InParam, InGuid)
 	{
 		RegisterInputConnection(&Collection);
@@ -1161,7 +1161,7 @@ public:
 		RegisterOutputConnection(&Collection, &Collection);
 	}
 
-	virtual void Evaluate(Dataflow::FContext& Context, const FDataflowOutput* Out) const override;
+	virtual void Evaluate(UE::Dataflow::FContext& Context, const FDataflowOutput* Out) const override;
 
 };
 
@@ -1250,7 +1250,7 @@ public:
 	UPROPERTY(EditAnywhere, Category = "Attribute", meta = (DisplayName = "Number of Elements"));
 	int32 NumElements = 0;
 
-	FAddCustomCollectionAttributeDataflowNode(const Dataflow::FNodeParameters& InParam, FGuid InGuid = FGuid::NewGuid())
+	FAddCustomCollectionAttributeDataflowNode(const UE::Dataflow::FNodeParameters& InParam, FGuid InGuid = FGuid::NewGuid())
 		: FDataflowNode(InParam, InGuid)
 	{
 		RegisterInputConnection(&Collection);
@@ -1258,7 +1258,7 @@ public:
 		RegisterOutputConnection(&Collection, &Collection);
 	}
 
-	virtual void Evaluate(Dataflow::FContext& Context, const FDataflowOutput* Out) const override;
+	virtual void Evaluate(UE::Dataflow::FContext& Context, const FDataflowOutput* Out) const override;
 
 };
 
@@ -1291,14 +1291,14 @@ public:
 	UPROPERTY(meta = (DataflowOutput));
 	int32 NumElements = 0;
 
-	FGetNumElementsInCollectionGroupDataflowNode(const Dataflow::FNodeParameters& InParam, FGuid InGuid = FGuid::NewGuid())
+	FGetNumElementsInCollectionGroupDataflowNode(const UE::Dataflow::FNodeParameters& InParam, FGuid InGuid = FGuid::NewGuid())
 		: FDataflowNode(InParam, InGuid)
 	{
 		RegisterInputConnection(&Collection);
 		RegisterOutputConnection(&NumElements);
 	}
 
-	virtual void Evaluate(Dataflow::FContext& Context, const FDataflowOutput* Out) const override;
+	virtual void Evaluate(UE::Dataflow::FContext& Context, const FDataflowOutput* Out) const override;
 
 };
 
@@ -1367,7 +1367,7 @@ public:
 	UPROPERTY(meta = (DataflowOutput, DisplayName = "Linear Color Array"));
 	TArray<FLinearColor> LinearColorAttributeData;
 
-	FGetCollectionAttributeDataTypedDataflowNode(const Dataflow::FNodeParameters& InParam, FGuid InGuid = FGuid::NewGuid())
+	FGetCollectionAttributeDataTypedDataflowNode(const UE::Dataflow::FNodeParameters& InParam, FGuid InGuid = FGuid::NewGuid())
 		: FDataflowNode(InParam, InGuid)
 	{
 		RegisterInputConnection(&Collection);
@@ -1382,7 +1382,7 @@ public:
 		RegisterOutputConnection(&LinearColorAttributeData);
 	}
 
-	virtual void Evaluate(Dataflow::FContext& Context, const FDataflowOutput* Out) const override;
+	virtual void Evaluate(UE::Dataflow::FContext& Context, const FDataflowOutput* Out) const override;
 
 };
 
@@ -1451,7 +1451,7 @@ public:
 	UPROPERTY(meta = (DataflowInput, DisplayName = "Linear Color Array"));
 	TArray<FLinearColor> LinearColorAttributeData;
 
-	FSetCollectionAttributeDataTypedDataflowNode(const Dataflow::FNodeParameters& InParam, FGuid InGuid = FGuid::NewGuid())
+	FSetCollectionAttributeDataTypedDataflowNode(const UE::Dataflow::FNodeParameters& InParam, FGuid InGuid = FGuid::NewGuid())
 		: FDataflowNode(InParam, InGuid)
 	{
 		RegisterInputConnection(&Collection);
@@ -1467,7 +1467,7 @@ public:
 		RegisterOutputConnection(&Collection, &Collection);
 	}
 
-	virtual void Evaluate(Dataflow::FContext& Context, const FDataflowOutput* Out) const override;
+	virtual void Evaluate(UE::Dataflow::FContext& Context, const FDataflowOutput* Out) const override;
 
 };
 
@@ -1493,14 +1493,14 @@ public:
 	TArray<int32> VertexList;
 
 
-	FSelectionToVertexListDataflowNode(const Dataflow::FNodeParameters& InParam, FGuid InGuid = FGuid::NewGuid())
+	FSelectionToVertexListDataflowNode(const UE::Dataflow::FNodeParameters& InParam, FGuid InGuid = FGuid::NewGuid())
 		: FDataflowNode(InParam, InGuid)
 	{
 		RegisterInputConnection(&VertexSelection);
 		RegisterOutputConnection(&VertexList);
 	}
 
-	virtual void Evaluate(Dataflow::FContext& Context, const FDataflowOutput* Out) const override;
+	virtual void Evaluate(UE::Dataflow::FContext& Context, const FDataflowOutput* Out) const override;
 
 };
 
@@ -1525,7 +1525,7 @@ public:
 	UPROPERTY(meta = ( DataflowOutput, DisplayName = "Out Transform"));
 	FTransform OutTransform = FTransform::Identity;
 
-	FMultiplyTransformDataflowNode(const Dataflow::FNodeParameters& InParam, FGuid InGuid = FGuid::NewGuid())
+	FMultiplyTransformDataflowNode(const UE::Dataflow::FNodeParameters& InParam, FGuid InGuid = FGuid::NewGuid())
 		: FDataflowNode(InParam, InGuid)
 	{
 		RegisterInputConnection(&InLeftTransform);
@@ -1533,7 +1533,7 @@ public:
 		RegisterOutputConnection(&OutTransform);
 	}
 
-	virtual void Evaluate(Dataflow::FContext& Context, const FDataflowOutput* Out) const override;
+	virtual void Evaluate(UE::Dataflow::FContext& Context, const FDataflowOutput* Out) const override;
 
 };
 
@@ -1555,14 +1555,14 @@ public:
 	UPROPERTY(meta = (DataflowOutput, DisplayName = "OutTransform"));
 	FTransform OutTransform = FTransform::Identity;
 
-	FInvertTransformDataflowNode(const Dataflow::FNodeParameters& InParam, FGuid InGuid = FGuid::NewGuid())
+	FInvertTransformDataflowNode(const UE::Dataflow::FNodeParameters& InParam, FGuid InGuid = FGuid::NewGuid())
 		: FDataflowNode(InParam, InGuid)
 	{
 		RegisterInputConnection(&InTransform);
 		RegisterOutputConnection(&OutTransform);
 	}
 
-	virtual void Evaluate(Dataflow::FContext& Context, const FDataflowOutput* Out) const override;
+	virtual void Evaluate(UE::Dataflow::FContext& Context, const FDataflowOutput* Out) const override;
 };
 
 /**
@@ -1593,7 +1593,7 @@ public:
 	UPROPERTY(meta = (DataflowOutput))
 	float ReturnValue = 0.f;
 
-	FBranchFloatDataflowNode(const Dataflow::FNodeParameters& InParam, FGuid InGuid = FGuid::NewGuid())
+	FBranchFloatDataflowNode(const UE::Dataflow::FNodeParameters& InParam, FGuid InGuid = FGuid::NewGuid())
 		: FDataflowNode(InParam, InGuid)
 	{
 		RegisterInputConnection(&A);
@@ -1602,7 +1602,7 @@ public:
 		RegisterOutputConnection(&ReturnValue);
 	}
 
-	virtual void Evaluate(Dataflow::FContext& Context, const FDataflowOutput* Out) const override;
+	virtual void Evaluate(UE::Dataflow::FContext& Context, const FDataflowOutput* Out) const override;
 
 };
 
@@ -1634,7 +1634,7 @@ public:
 	UPROPERTY(meta = (DataflowOutput))
 	int32 ReturnValue = 0;
 
-	FBranchIntDataflowNode(const Dataflow::FNodeParameters& InParam, FGuid InGuid = FGuid::NewGuid())
+	FBranchIntDataflowNode(const UE::Dataflow::FNodeParameters& InParam, FGuid InGuid = FGuid::NewGuid())
 		: FDataflowNode(InParam, InGuid)
 	{
 		RegisterInputConnection(&A);
@@ -1643,7 +1643,7 @@ public:
 		RegisterOutputConnection(&ReturnValue);
 	}
 
-	virtual void Evaluate(Dataflow::FContext& Context, const FDataflowOutput* Out) const override;
+	virtual void Evaluate(UE::Dataflow::FContext& Context, const FDataflowOutput* Out) const override;
 
 };
 
@@ -1660,7 +1660,7 @@ public:
  //USTRUCT(meta = (Experimental))
  //struct FLogStringDataflowNode_v2 : public FDataflowNode
 
-namespace Dataflow
+namespace UE::Dataflow
 {
 	void GeometryCollectionEngineNodes();
 }
diff --git a/Engine/Plugins/Experimental/GeometryCollectionPlugin/Source/GeometryCollectionNodes/Public/Dataflow/GeometryCollectionOverrideNodes.h b/Engine/Plugins/Experimental/GeometryCollectionPlugin/Source/GeometryCollectionNodes/Public/Dataflow/GeometryCollectionOverrideNodes.h
--- a/Engine/Plugins/Experimental/GeometryCollectionPlugin/Source/GeometryCollectionNodes/Public/Dataflow/GeometryCollectionOverrideNodes.h
+++ b/Engine/Plugins/Experimental/GeometryCollectionPlugin/Source/GeometryCollectionNodes/Public/Dataflow/GeometryCollectionOverrideNodes.h
@@ -26,14 +26,14 @@ public:
 	UPROPERTY(meta = (DataflowOutput))
 	float FloatDefault = 0.f;
 
-	FGetFloatOverrideFromAssetDataflowNode(const Dataflow::FNodeParameters& InParam, FGuid InGuid = FGuid::NewGuid())
+	FGetFloatOverrideFromAssetDataflowNode(const UE::Dataflow::FNodeParameters& InParam, FGuid InGuid = FGuid::NewGuid())
 		: FDataflowOverrideNode(InParam, InGuid)
 	{
 		RegisterOutputConnection(&Float);
 		RegisterOutputConnection(&FloatDefault);
 	}
 
-	virtual void Evaluate(Dataflow::FContext& Context, const FDataflowOutput* Out) const override;
+	virtual void Evaluate(UE::Dataflow::FContext& Context, const FDataflowOutput* Out) const override;
 };
 
 /**
@@ -54,14 +54,14 @@ public:
 	UPROPERTY(meta = (DataflowOutput))
 	int32 IntDefault = 0;
 
-	FGetIntOverrideFromAssetDataflowNode(const Dataflow::FNodeParameters& InParam, FGuid InGuid = FGuid::NewGuid())
+	FGetIntOverrideFromAssetDataflowNode(const UE::Dataflow::FNodeParameters& InParam, FGuid InGuid = FGuid::NewGuid())
 		: FDataflowOverrideNode(InParam, InGuid)
 	{
 		RegisterOutputConnection(&Int);
 		RegisterOutputConnection(&IntDefault);
 	}
 
-	virtual void Evaluate(Dataflow::FContext& Context, const FDataflowOutput* Out) const override;
+	virtual void Evaluate(UE::Dataflow::FContext& Context, const FDataflowOutput* Out) const override;
 };
 
 /**
@@ -82,14 +82,14 @@ public:
 	UPROPERTY(meta = (DataflowOutput))
 	bool BoolDefault = false;
 
-	FGetBoolOverrideFromAssetDataflowNode(const Dataflow::FNodeParameters& InParam, FGuid InGuid = FGuid::NewGuid())
+	FGetBoolOverrideFromAssetDataflowNode(const UE::Dataflow::FNodeParameters& InParam, FGuid InGuid = FGuid::NewGuid())
 		: FDataflowOverrideNode(InParam, InGuid)
 	{
 		RegisterOutputConnection(&Bool);
 		RegisterOutputConnection(&BoolDefault);
 	}
 
-	virtual void Evaluate(Dataflow::FContext& Context, const FDataflowOutput* Out) const override;
+	virtual void Evaluate(UE::Dataflow::FContext& Context, const FDataflowOutput* Out) const override;
 };
 
 /**
@@ -110,17 +110,17 @@ public:
 	UPROPERTY(meta = (DataflowOutput))
 	FString StringDefault;
 
-	FGetStringOverrideFromAssetDataflowNode(const Dataflow::FNodeParameters& InParam, FGuid InGuid = FGuid::NewGuid())
+	FGetStringOverrideFromAssetDataflowNode(const UE::Dataflow::FNodeParameters& InParam, FGuid InGuid = FGuid::NewGuid())
 		: FDataflowOverrideNode(InParam, InGuid)
 	{
 		RegisterOutputConnection(&String);
 		RegisterOutputConnection(&StringDefault);
 	}
 
-	virtual void Evaluate(Dataflow::FContext& Context, const FDataflowOutput* Out) const override;
+	virtual void Evaluate(UE::Dataflow::FContext& Context, const FDataflowOutput* Out) const override;
 };
 
-namespace Dataflow
+namespace UE::Dataflow
 {
 	void GeometryCollectionOverrideNodes();
 }
diff --git a/Engine/Plugins/Experimental/GeometryCollectionPlugin/Source/GeometryCollectionNodes/Public/Dataflow/GeometryCollectionProcessingNodes.h b/Engine/Plugins/Experimental/GeometryCollectionPlugin/Source/GeometryCollectionNodes/Public/Dataflow/GeometryCollectionProcessingNodes.h
--- a/Engine/Plugins/Experimental/GeometryCollectionPlugin/Source/GeometryCollectionNodes/Public/Dataflow/GeometryCollectionProcessingNodes.h
+++ b/Engine/Plugins/Experimental/GeometryCollectionPlugin/Source/GeometryCollectionNodes/Public/Dataflow/GeometryCollectionProcessingNodes.h
@@ -21,17 +21,17 @@ public:
 	UPROPERTY(meta = (DataflowInput, DataflowOutput, DisplayName = "Collection"))
 		FManagedArrayCollection Collection;
 
-	FCloseGeometryOnCollectionDataflowNode(const Dataflow::FNodeParameters& InParam, FGuid InGuid = FGuid::NewGuid())
+	FCloseGeometryOnCollectionDataflowNode(const UE::Dataflow::FNodeParameters& InParam, FGuid InGuid = FGuid::NewGuid())
 		: FDataflowNode(InParam, InGuid)
 	{
 		RegisterInputConnection(&Collection);
 		RegisterOutputConnection(&Collection);
 	}
 
-	virtual void Evaluate(Dataflow::FContext& Context, const FDataflowOutput* Out) const override;
+	virtual void Evaluate(UE::Dataflow::FContext& Context, const FDataflowOutput* Out) const override;
 };
 
-namespace Dataflow
+namespace UE::Dataflow
 {
 	void GeometryCollectionProcessingNodes();
 
diff --git a/Engine/Plugins/Experimental/GeometryCollectionPlugin/Source/GeometryCollectionNodes/Public/Dataflow/GeometryCollectionSamplingNodes.h b/Engine/Plugins/Experimental/GeometryCollectionPlugin/Source/GeometryCollectionNodes/Public/Dataflow/GeometryCollectionSamplingNodes.h
--- a/Engine/Plugins/Experimental/GeometryCollectionPlugin/Source/GeometryCollectionNodes/Public/Dataflow/GeometryCollectionSamplingNodes.h
+++ b/Engine/Plugins/Experimental/GeometryCollectionPlugin/Source/GeometryCollectionNodes/Public/Dataflow/GeometryCollectionSamplingNodes.h
@@ -58,7 +58,7 @@ public:
 	UPROPERTY(meta = (DataflowOutput))
 	int32 NumSamplePoints = 0;
 
-	FUniformPointSamplingDataflowNode(const Dataflow::FNodeParameters& InParam, FGuid InGuid = FGuid::NewGuid())
+	FUniformPointSamplingDataflowNode(const UE::Dataflow::FNodeParameters& InParam, FGuid InGuid = FGuid::NewGuid())
 		: FDataflowNode(InParam, InGuid)
 	{
 		RegisterInputConnection(&TargetMesh);
@@ -72,7 +72,7 @@ public:
 		RegisterOutputConnection(&NumSamplePoints);
 	}
 
-	virtual void Evaluate(Dataflow::FContext& Context, const FDataflowOutput* Out) const override;
+	virtual void Evaluate(UE::Dataflow::FContext& Context, const FDataflowOutput* Out) const override;
 };
 
 /**
@@ -139,7 +139,7 @@ public:
 	UPROPERTY(meta = (DataflowOutput))
 	int32 NumSamplePoints = 0;
 
-	FNonUniformPointSamplingDataflowNode(const Dataflow::FNodeParameters& InParam, FGuid InGuid = FGuid::NewGuid())
+	FNonUniformPointSamplingDataflowNode(const UE::Dataflow::FNodeParameters& InParam, FGuid InGuid = FGuid::NewGuid())
 		: FDataflowNode(InParam, InGuid)
 	{
 		RegisterInputConnection(&TargetMesh);
@@ -156,7 +156,7 @@ public:
 		RegisterOutputConnection(&NumSamplePoints);
 	}
 
-	virtual void Evaluate(Dataflow::FContext& Context, const FDataflowOutput* Out) const override;
+	virtual void Evaluate(UE::Dataflow::FContext& Context, const FDataflowOutput* Out) const override;
 };
 
 /**
@@ -233,7 +233,7 @@ public:
 	UPROPERTY(meta = (DataflowOutput))
 	int32 NumSamplePoints = 0;
 
-	FVertexWeightedPointSamplingDataflowNode(const Dataflow::FNodeParameters& InParam, FGuid InGuid = FGuid::NewGuid())
+	FVertexWeightedPointSamplingDataflowNode(const UE::Dataflow::FNodeParameters& InParam, FGuid InGuid = FGuid::NewGuid())
 		: FDataflowNode(InParam, InGuid)
 	{
 		RegisterInputConnection(&TargetMesh);
@@ -251,11 +251,11 @@ public:
 		RegisterOutputConnection(&NumSamplePoints);
 	}
 
-	virtual void Evaluate(Dataflow::FContext& Context, const FDataflowOutput* Out) const override;
+	virtual void Evaluate(UE::Dataflow::FContext& Context, const FDataflowOutput* Out) const override;
 };
 
 
-namespace Dataflow
+namespace UE::Dataflow
 {
 	void GeometryCollectionSamplingNodes();
 }
