Change Dataflow namespace to UE::Dataflow to be ok with coding standard

diff --git a/Engine/Plugins/Experimental/GeometryCollectionPlugin/Source/GeometryCollectionNodes/Public/Dataflow/GeometryCollectionSelectionNodes.h b/Engine/Plugins/Experimental/GeometryCollectionPlugin/Source/GeometryCollectionNodes/Public/Dataflow/GeometryCollectionSelectionNodes.h
--- a/Engine/Plugins/Experimental/GeometryCollectionPlugin/Source/GeometryCollectionNodes/Public/Dataflow/GeometryCollectionSelectionNodes.h
+++ b/Engine/Plugins/Experimental/GeometryCollectionPlugin/Source/GeometryCollectionNodes/Public/Dataflow/GeometryCollectionSelectionNodes.h
@@ -35,7 +35,7 @@ public:
 	UPROPERTY(meta = (DataflowOutput, DisplayName = "TransformSelection"))
 	FDataflowTransformSelection TransformSelection;
 
-	FCollectionTransformSelectionAllDataflowNode(const Dataflow::FNodeParameters& InParam, FGuid InGuid = FGuid::NewGuid())
+	FCollectionTransformSelectionAllDataflowNode(const UE::Dataflow::FNodeParameters& InParam, FGuid InGuid = FGuid::NewGuid())
 		: FDataflowNode(InParam, InGuid)
 	{
 		RegisterInputConnection(&Collection);
@@ -43,7 +43,7 @@ public:
 		RegisterOutputConnection(&TransformSelection);
 	}
 
-	virtual void Evaluate(Dataflow::FContext& Context, const FDataflowOutput* Out) const override;
+	virtual void Evaluate(UE::Dataflow::FContext& Context, const FDataflowOutput* Out) const override;
 
 };
 
@@ -92,7 +92,7 @@ public:
 	UPROPERTY(meta = (DataflowOutput, DisplayName = "TransformSelection", DataflowPassthrough = "TransformSelectionA"))
 	FDataflowTransformSelection TransformSelection;
 
-	FCollectionTransformSelectionSetOperationDataflowNode(const Dataflow::FNodeParameters& InParam, FGuid InGuid = FGuid::NewGuid())
+	FCollectionTransformSelectionSetOperationDataflowNode(const UE::Dataflow::FNodeParameters& InParam, FGuid InGuid = FGuid::NewGuid())
 		: FDataflowNode(InParam, InGuid)
 	{
 		RegisterInputConnection(&TransformSelectionA);
@@ -100,7 +100,7 @@ public:
 		RegisterOutputConnection(&TransformSelection, &TransformSelectionA);
 	}
 
-	virtual void Evaluate(Dataflow::FContext& Context, const FDataflowOutput* Out) const override;
+	virtual void Evaluate(UE::Dataflow::FContext& Context, const FDataflowOutput* Out) const override;
 
 };
 
@@ -129,7 +129,7 @@ public:
 	UPROPERTY(meta = (DataflowOutput))
 	FString String;
 
-	FCollectionTransformSelectionInfoDataflowNode(const Dataflow::FNodeParameters& InParam, FGuid InGuid = FGuid::NewGuid())
+	FCollectionTransformSelectionInfoDataflowNode(const UE::Dataflow::FNodeParameters& InParam, FGuid InGuid = FGuid::NewGuid())
 		: FDataflowNode(InParam, InGuid)
 	{
 		RegisterInputConnection(&TransformSelection);
@@ -137,7 +137,7 @@ public:
 		RegisterOutputConnection(&String);
 	}
 
-	virtual void Evaluate(Dataflow::FContext& Context, const FDataflowOutput* Out) const override;
+	virtual void Evaluate(UE::Dataflow::FContext& Context, const FDataflowOutput* Out) const override;
 
 };
 
@@ -162,7 +162,7 @@ public:
 	UPROPERTY(meta = (DataflowOutput, DisplayName = "TransformSelection"))
 	FDataflowTransformSelection TransformSelection;
 
-	FCollectionTransformSelectionNoneDataflowNode(const Dataflow::FNodeParameters& InParam, FGuid InGuid = FGuid::NewGuid())
+	FCollectionTransformSelectionNoneDataflowNode(const UE::Dataflow::FNodeParameters& InParam, FGuid InGuid = FGuid::NewGuid())
 		: FDataflowNode(InParam, InGuid)
 	{
 		RegisterInputConnection(&Collection);
@@ -170,7 +170,7 @@ public:
 		RegisterOutputConnection(&TransformSelection);
 	}
 
-	virtual void Evaluate(Dataflow::FContext& Context, const FDataflowOutput* Out) const override;
+	virtual void Evaluate(UE::Dataflow::FContext& Context, const FDataflowOutput* Out) const override;
 
 };
 
@@ -191,14 +191,14 @@ public:
 	UPROPERTY(meta = (DataflowInput, DataflowOutput, DisplayName = "TransformSelection", DataflowPassthrough = "TransformSelection", DataflowIntrinsic))
 	FDataflowTransformSelection TransformSelection;
 
-	FCollectionTransformSelectionInvertDataflowNode(const Dataflow::FNodeParameters& InParam, FGuid InGuid = FGuid::NewGuid())
+	FCollectionTransformSelectionInvertDataflowNode(const UE::Dataflow::FNodeParameters& InParam, FGuid InGuid = FGuid::NewGuid())
 		: FDataflowNode(InParam, InGuid)
 	{
 		RegisterInputConnection(&TransformSelection);
 		RegisterOutputConnection(&TransformSelection, &TransformSelection);
 	}
 
-	virtual void Evaluate(Dataflow::FContext& Context, const FDataflowOutput* Out) const override;
+	virtual void Evaluate(UE::Dataflow::FContext& Context, const FDataflowOutput* Out) const override;
 
 };
 
@@ -235,7 +235,7 @@ public:
 	UPROPERTY(meta = (DataflowOutput, DisplayName = "TransformSelection"))
 	FDataflowTransformSelection TransformSelection;
 
-	FCollectionTransformSelectionRandomDataflowNode(const Dataflow::FNodeParameters& InParam, FGuid InGuid = FGuid::NewGuid())
+	FCollectionTransformSelectionRandomDataflowNode(const UE::Dataflow::FNodeParameters& InParam, FGuid InGuid = FGuid::NewGuid())
 		: FDataflowNode(InParam, InGuid)
 	{
 		RandomSeed = FMath::FRandRange(-1e5, 1e5);
@@ -246,7 +246,7 @@ public:
 		RegisterOutputConnection(&TransformSelection);
 	}
 
-	virtual void Evaluate(Dataflow::FContext& Context, const FDataflowOutput* Out) const override;
+	virtual void Evaluate(UE::Dataflow::FContext& Context, const FDataflowOutput* Out) const override;
 
 };
 
@@ -271,7 +271,7 @@ public:
 	UPROPERTY(meta = (DataflowOutput, DisplayName = "TransformSelection"))
 	FDataflowTransformSelection TransformSelection;
 
-	FCollectionTransformSelectionRootDataflowNode(const Dataflow::FNodeParameters& InParam, FGuid InGuid = FGuid::NewGuid())
+	FCollectionTransformSelectionRootDataflowNode(const UE::Dataflow::FNodeParameters& InParam, FGuid InGuid = FGuid::NewGuid())
 		: FDataflowNode(InParam, InGuid)
 	{
 		RegisterInputConnection(&Collection);
@@ -279,7 +279,7 @@ public:
 		RegisterOutputConnection(&TransformSelection);
 	}
 
-	virtual void Evaluate(Dataflow::FContext& Context, const FDataflowOutput* Out) const override;
+	virtual void Evaluate(UE::Dataflow::FContext& Context, const FDataflowOutput* Out) const override;
 
 };
 
@@ -309,7 +309,7 @@ public:
 	UPROPERTY(meta = (DataflowOutput, DisplayName = "TransformSelection"))
 	FDataflowTransformSelection TransformSelection;
 
-	FCollectionTransformSelectionCustomDataflowNode(const Dataflow::FNodeParameters& InParam, FGuid InGuid = FGuid::NewGuid())
+	FCollectionTransformSelectionCustomDataflowNode(const UE::Dataflow::FNodeParameters& InParam, FGuid InGuid = FGuid::NewGuid())
 		: FDataflowNode(InParam, InGuid)
 	{
 		RegisterInputConnection(&Collection);
@@ -318,7 +318,7 @@ public:
 		RegisterOutputConnection(&TransformSelection);
 	}
 
-	virtual void Evaluate(Dataflow::FContext& Context, const FDataflowOutput* Out) const override;
+	virtual void Evaluate(UE::Dataflow::FContext& Context, const FDataflowOutput* Out) const override;
 
 };
 
@@ -347,7 +347,7 @@ public:
 	UPROPERTY(meta = (DataflowOutput, DisplayName = "TransformSelection"))
 	FDataflowTransformSelection TransformSelection;
 
-	FCollectionTransformSelectionFromIndexArrayDataflowNode(const Dataflow::FNodeParameters& InParam, FGuid InGuid = FGuid::NewGuid())
+	FCollectionTransformSelectionFromIndexArrayDataflowNode(const UE::Dataflow::FNodeParameters& InParam, FGuid InGuid = FGuid::NewGuid())
 		: FDataflowNode(InParam, InGuid)
 	{
 		RegisterInputConnection(&Collection);
@@ -356,7 +356,7 @@ public:
 		RegisterOutputConnection(&TransformSelection);
 	}
 
-	virtual void Evaluate(Dataflow::FContext& Context, const FDataflowOutput* Out) const override;
+	virtual void Evaluate(UE::Dataflow::FContext& Context, const FDataflowOutput* Out) const override;
 
 };
 
@@ -381,7 +381,7 @@ public:
 	UPROPERTY(meta = (DataflowInput, DataflowOutput, DataflowPassthrough = "Collection", DataflowIntrinsic))
 	FManagedArrayCollection Collection;
 
-	FCollectionTransformSelectionParentDataflowNode(const Dataflow::FNodeParameters& InParam, FGuid InGuid = FGuid::NewGuid())
+	FCollectionTransformSelectionParentDataflowNode(const UE::Dataflow::FNodeParameters& InParam, FGuid InGuid = FGuid::NewGuid())
 		: FDataflowNode(InParam, InGuid)
 	{
 		RegisterInputConnection(&TransformSelection);
@@ -390,7 +390,7 @@ public:
 		RegisterOutputConnection(&TransformSelection, &TransformSelection);
 	}
 
-	virtual void Evaluate(Dataflow::FContext& Context, const FDataflowOutput* Out) const override;
+	virtual void Evaluate(UE::Dataflow::FContext& Context, const FDataflowOutput* Out) const override;
 
 };
 
@@ -423,7 +423,7 @@ public:
 	UPROPERTY(EditAnywhere, Category = "Random", meta = (DataflowInput, EditCondition = "bDeterministic"))
 	float RandomSeed = 0.f;
 
-	FCollectionTransformSelectionByPercentageDataflowNode(const Dataflow::FNodeParameters& InParam, FGuid InGuid = FGuid::NewGuid())
+	FCollectionTransformSelectionByPercentageDataflowNode(const UE::Dataflow::FNodeParameters& InParam, FGuid InGuid = FGuid::NewGuid())
 		: FDataflowNode(InParam, InGuid)
 	{
 		RandomSeed = FMath::RandRange(-100000, 100000);
@@ -433,7 +433,7 @@ public:
 		RegisterOutputConnection(&TransformSelection, &TransformSelection);
 	}
 
-	virtual void Evaluate(Dataflow::FContext& Context, const FDataflowOutput* Out) const override;
+	virtual void Evaluate(UE::Dataflow::FContext& Context, const FDataflowOutput* Out) const override;
 
 };
 
@@ -458,7 +458,7 @@ public:
 	UPROPERTY(meta = (DataflowInput, DataflowOutput, DataflowPassthrough = "Collection", DataflowIntrinsic))
 	FManagedArrayCollection Collection;
 
-	FCollectionTransformSelectionChildrenDataflowNode(const Dataflow::FNodeParameters& InParam, FGuid InGuid = FGuid::NewGuid())
+	FCollectionTransformSelectionChildrenDataflowNode(const UE::Dataflow::FNodeParameters& InParam, FGuid InGuid = FGuid::NewGuid())
 		: FDataflowNode(InParam, InGuid)
 	{
 		RegisterInputConnection(&TransformSelection);
@@ -467,7 +467,7 @@ public:
 		RegisterOutputConnection(&TransformSelection, &TransformSelection);
 	}
 
-	virtual void Evaluate(Dataflow::FContext& Context, const FDataflowOutput* Out) const override;
+	virtual void Evaluate(UE::Dataflow::FContext& Context, const FDataflowOutput* Out) const override;
 
 };
 
@@ -492,7 +492,7 @@ public:
 	UPROPERTY(meta = (DataflowInput, DataflowOutput, DataflowPassthrough = "Collection", DataflowIntrinsic))
 	FManagedArrayCollection Collection;
 
-	FCollectionTransformSelectionSiblingsDataflowNode(const Dataflow::FNodeParameters& InParam, FGuid InGuid = FGuid::NewGuid())
+	FCollectionTransformSelectionSiblingsDataflowNode(const UE::Dataflow::FNodeParameters& InParam, FGuid InGuid = FGuid::NewGuid())
 		: FDataflowNode(InParam, InGuid)
 	{
 		RegisterInputConnection(&TransformSelection);
@@ -501,7 +501,7 @@ public:
 		RegisterOutputConnection(&TransformSelection, &TransformSelection);
 	}
 
-	virtual void Evaluate(Dataflow::FContext& Context, const FDataflowOutput* Out) const override;
+	virtual void Evaluate(UE::Dataflow::FContext& Context, const FDataflowOutput* Out) const override;
 
 };
 
@@ -526,7 +526,7 @@ public:
 	UPROPERTY(meta = (DataflowInput, DataflowOutput, DataflowPassthrough = "Collection", DataflowIntrinsic))
 	FManagedArrayCollection Collection;
 
-	FCollectionTransformSelectionLevelDataflowNode(const Dataflow::FNodeParameters& InParam, FGuid InGuid = FGuid::NewGuid())
+	FCollectionTransformSelectionLevelDataflowNode(const UE::Dataflow::FNodeParameters& InParam, FGuid InGuid = FGuid::NewGuid())
 		: FDataflowNode(InParam, InGuid)
 	{
 		RegisterInputConnection(&TransformSelection);
@@ -535,7 +535,7 @@ public:
 		RegisterOutputConnection(&TransformSelection, &TransformSelection);
 	}
 
-	virtual void Evaluate(Dataflow::FContext& Context, const FDataflowOutput* Out) const override;
+	virtual void Evaluate(UE::Dataflow::FContext& Context, const FDataflowOutput* Out) const override;
 
 };
 
@@ -568,7 +568,7 @@ public:
 	UPROPERTY(meta = (DataflowOutput, DisplayName = "TransformSelection"))
 	FDataflowTransformSelection TransformSelection;
 
-	FCollectionTransformSelectionTargetLevelDataflowNode(const Dataflow::FNodeParameters& InParam, FGuid InGuid = FGuid::NewGuid())
+	FCollectionTransformSelectionTargetLevelDataflowNode(const UE::Dataflow::FNodeParameters& InParam, FGuid InGuid = FGuid::NewGuid())
 		: FDataflowNode(InParam, InGuid)
 	{
 		RegisterInputConnection(&Collection);
@@ -577,7 +577,7 @@ public:
 		RegisterOutputConnection(&TransformSelection);
 	}
 
-	virtual void Evaluate(Dataflow::FContext& Context, const FDataflowOutput* Out) const override;
+	virtual void Evaluate(UE::Dataflow::FContext& Context, const FDataflowOutput* Out) const override;
 
 };
 
@@ -607,7 +607,7 @@ public:
 	UPROPERTY(EditAnywhere, Category = Options)
 	bool bAllowContactInParentLevels = true;
 
-	FCollectionTransformSelectionContactDataflowNode(const Dataflow::FNodeParameters& InParam, FGuid InGuid = FGuid::NewGuid())
+	FCollectionTransformSelectionContactDataflowNode(const UE::Dataflow::FNodeParameters& InParam, FGuid InGuid = FGuid::NewGuid())
 		: FDataflowNode(InParam, InGuid)
 	{
 		RegisterInputConnection(&TransformSelection);
@@ -616,7 +616,7 @@ public:
 		RegisterOutputConnection(&TransformSelection, &TransformSelection);
 	}
 
-	virtual void Evaluate(Dataflow::FContext& Context, const FDataflowOutput* Out) const override;
+	virtual void Evaluate(UE::Dataflow::FContext& Context, const FDataflowOutput* Out) const override;
 
 };
 
@@ -641,7 +641,7 @@ public:
 	UPROPERTY(meta = (DataflowOutput, DisplayName = "TransformSelection"))
 	FDataflowTransformSelection TransformSelection;
 
-	FCollectionTransformSelectionLeafDataflowNode(const Dataflow::FNodeParameters& InParam, FGuid InGuid = FGuid::NewGuid())
+	FCollectionTransformSelectionLeafDataflowNode(const UE::Dataflow::FNodeParameters& InParam, FGuid InGuid = FGuid::NewGuid())
 		: FDataflowNode(InParam, InGuid)
 	{
 		RegisterInputConnection(&Collection);
@@ -649,7 +649,7 @@ public:
 		RegisterOutputConnection(&TransformSelection);
 	}
 
-	virtual void Evaluate(Dataflow::FContext& Context, const FDataflowOutput* Out) const override;
+	virtual void Evaluate(UE::Dataflow::FContext& Context, const FDataflowOutput* Out) const override;
 
 };
 
@@ -676,7 +676,7 @@ public:
 	UPROPERTY(meta = (DataflowOutput, DisplayName = "TransformSelection"))
 	FDataflowTransformSelection TransformSelection;
 
-	FCollectionTransformSelectionClusterDataflowNode(const Dataflow::FNodeParameters& InParam, FGuid InGuid = FGuid::NewGuid())
+	FCollectionTransformSelectionClusterDataflowNode(const UE::Dataflow::FNodeParameters& InParam, FGuid InGuid = FGuid::NewGuid())
 		: FDataflowNode(InParam, InGuid)
 	{
 		RegisterInputConnection(&Collection);
@@ -684,7 +684,7 @@ public:
 		RegisterOutputConnection(&TransformSelection);
 	}
 
-	virtual void Evaluate(Dataflow::FContext& Context, const FDataflowOutput* Out) const override;
+	virtual void Evaluate(UE::Dataflow::FContext& Context, const FDataflowOutput* Out) const override;
 
 };
 
@@ -709,7 +709,7 @@ public:
 	UPROPERTY(meta = (DataflowOutput, DisplayName = "TransformSelection"))
 	FDataflowTransformSelection TransformSelection;
 
-	FCollectionTransformSelectionClusterDataflowNode_v2(const Dataflow::FNodeParameters& InParam, FGuid InGuid = FGuid::NewGuid())
+	FCollectionTransformSelectionClusterDataflowNode_v2(const UE::Dataflow::FNodeParameters& InParam, FGuid InGuid = FGuid::NewGuid())
 		: FDataflowNode(InParam, InGuid)
 	{
 		RegisterInputConnection(&Collection);
@@ -717,7 +717,7 @@ public:
 		RegisterOutputConnection(&TransformSelection);
 	}
 
-	virtual void Evaluate(Dataflow::FContext& Context, const FDataflowOutput* Out) const override;
+	virtual void Evaluate(UE::Dataflow::FContext& Context, const FDataflowOutput* Out) const override;
 
 };
 
@@ -771,7 +771,7 @@ public:
 	UPROPERTY(meta = (DataflowOutput))
 	TArray<int> Indices;
 
-	FSelectFloatArrayIndicesInRangeDataflowNode(const Dataflow::FNodeParameters& InParam, FGuid InGuid = FGuid::NewGuid())
+	FSelectFloatArrayIndicesInRangeDataflowNode(const UE::Dataflow::FNodeParameters& InParam, FGuid InGuid = FGuid::NewGuid())
 		: FDataflowNode(InParam, InGuid)
 	{
 		RegisterInputConnection(&Values);
@@ -780,7 +780,7 @@ public:
 		RegisterOutputConnection(&Indices);
 	}
 
-	virtual void Evaluate(Dataflow::FContext& Context, const FDataflowOutput* Out) const override;
+	virtual void Evaluate(UE::Dataflow::FContext& Context, const FDataflowOutput* Out) const override;
 
 };
 
@@ -825,7 +825,7 @@ public:
 	UPROPERTY(meta = (DataflowOutput, DisplayName = "TransformSelection"))
 	FDataflowTransformSelection TransformSelection;
 
-	FCollectionTransformSelectionBySizeDataflowNode(const Dataflow::FNodeParameters& InParam, FGuid InGuid = FGuid::NewGuid())
+	FCollectionTransformSelectionBySizeDataflowNode(const UE::Dataflow::FNodeParameters& InParam, FGuid InGuid = FGuid::NewGuid())
 		: FDataflowNode(InParam, InGuid)
 	{
 		RegisterInputConnection(&Collection);
@@ -835,7 +835,7 @@ public:
 		RegisterOutputConnection(&TransformSelection);
 	}
 
-	virtual void Evaluate(Dataflow::FContext& Context, const FDataflowOutput* Out) const override;
+	virtual void Evaluate(UE::Dataflow::FContext& Context, const FDataflowOutput* Out) const override;
 
 };
 
@@ -876,7 +876,7 @@ public:
 	UPROPERTY(meta = (DataflowOutput, DisplayName = "TransformSelection"))
 	FDataflowTransformSelection TransformSelection;
 
-	FCollectionTransformSelectionByVolumeDataflowNode(const Dataflow::FNodeParameters& InParam, FGuid InGuid = FGuid::NewGuid())
+	FCollectionTransformSelectionByVolumeDataflowNode(const UE::Dataflow::FNodeParameters& InParam, FGuid InGuid = FGuid::NewGuid())
 		: FDataflowNode(InParam, InGuid)
 	{
 		RegisterInputConnection(&Collection);
@@ -886,7 +886,7 @@ public:
 		RegisterOutputConnection(&TransformSelection);
 	}
 
-	virtual void Evaluate(Dataflow::FContext& Context, const FDataflowOutput* Out) const override;
+	virtual void Evaluate(UE::Dataflow::FContext& Context, const FDataflowOutput* Out) const override;
 
 };
 
@@ -941,7 +941,7 @@ public:
 	UPROPERTY(meta = (DataflowOutput, DisplayName = "TransformSelection"))
 	FDataflowTransformSelection TransformSelection;
 
-	FCollectionTransformSelectionInBoxDataflowNode(const Dataflow::FNodeParameters& InParam, FGuid InGuid = FGuid::NewGuid())
+	FCollectionTransformSelectionInBoxDataflowNode(const UE::Dataflow::FNodeParameters& InParam, FGuid InGuid = FGuid::NewGuid())
 		: FDataflowNode(InParam, InGuid)
 	{
 		RegisterInputConnection(&Collection);
@@ -951,7 +951,7 @@ public:
 		RegisterOutputConnection(&TransformSelection);
 	}
 
-	virtual void Evaluate(Dataflow::FContext& Context, const FDataflowOutput* Out) const override;
+	virtual void Evaluate(UE::Dataflow::FContext& Context, const FDataflowOutput* Out) const override;
 
 };
 
@@ -992,7 +992,7 @@ public:
 	UPROPERTY(meta = (DataflowOutput, DisplayName = "TransformSelection"))
 	FDataflowTransformSelection TransformSelection;
 
-	FCollectionTransformSelectionInSphereDataflowNode(const Dataflow::FNodeParameters& InParam, FGuid InGuid = FGuid::NewGuid())
+	FCollectionTransformSelectionInSphereDataflowNode(const UE::Dataflow::FNodeParameters& InParam, FGuid InGuid = FGuid::NewGuid())
 		: FDataflowNode(InParam, InGuid)
 	{
 		RegisterInputConnection(&Collection);
@@ -1002,7 +1002,7 @@ public:
 		RegisterOutputConnection(&TransformSelection);
 	}
 
-	virtual void Evaluate(Dataflow::FContext& Context, const FDataflowOutput* Out) const override;
+	virtual void Evaluate(UE::Dataflow::FContext& Context, const FDataflowOutput* Out) const override;
 
 };
 
@@ -1051,7 +1051,7 @@ public:
 	UPROPERTY(meta = (DataflowOutput, DisplayName = "TransformSelection"))
 	FDataflowTransformSelection TransformSelection;
 
-	FCollectionTransformSelectionByFloatAttrDataflowNode(const Dataflow::FNodeParameters& InParam, FGuid InGuid = FGuid::NewGuid())
+	FCollectionTransformSelectionByFloatAttrDataflowNode(const UE::Dataflow::FNodeParameters& InParam, FGuid InGuid = FGuid::NewGuid())
 		: FDataflowNode(InParam, InGuid)
 	{
 		RegisterInputConnection(&Collection);
@@ -1061,7 +1061,7 @@ public:
 		RegisterOutputConnection(&TransformSelection);
 	}
 
-	virtual void Evaluate(Dataflow::FContext& Context, const FDataflowOutput* Out) const override;
+	virtual void Evaluate(UE::Dataflow::FContext& Context, const FDataflowOutput* Out) const override;
 
 };
 
@@ -1110,7 +1110,7 @@ public:
 	UPROPERTY(meta = (DataflowOutput, DisplayName = "TransformSelection"))
 	FDataflowTransformSelection TransformSelection;
 
-	FCollectionTransformSelectionByIntAttrDataflowNode(const Dataflow::FNodeParameters& InParam, FGuid InGuid = FGuid::NewGuid())
+	FCollectionTransformSelectionByIntAttrDataflowNode(const UE::Dataflow::FNodeParameters& InParam, FGuid InGuid = FGuid::NewGuid())
 		: FDataflowNode(InParam, InGuid)
 	{
 		RegisterInputConnection(&Collection);
@@ -1120,7 +1120,7 @@ public:
 		RegisterOutputConnection(&TransformSelection);
 	}
 
-	virtual void Evaluate(Dataflow::FContext& Context, const FDataflowOutput* Out) const override;
+	virtual void Evaluate(UE::Dataflow::FContext& Context, const FDataflowOutput* Out) const override;
 
 };
 
@@ -1150,7 +1150,7 @@ public:
 	UPROPERTY(meta = (DataflowOutput, DisplayName = "VertexSelection"))
 	FDataflowVertexSelection VertexSelection;
 
-	FCollectionVertexSelectionCustomDataflowNode(const Dataflow::FNodeParameters& InParam, FGuid InGuid = FGuid::NewGuid())
+	FCollectionVertexSelectionCustomDataflowNode(const UE::Dataflow::FNodeParameters& InParam, FGuid InGuid = FGuid::NewGuid())
 		: FDataflowNode(InParam, InGuid)
 	{
 		RegisterInputConnection(&Collection);
@@ -1159,7 +1159,7 @@ public:
 		RegisterOutputConnection(&VertexSelection);
 	}
 
-	virtual void Evaluate(Dataflow::FContext& Context, const FDataflowOutput* Out) const override;
+	virtual void Evaluate(UE::Dataflow::FContext& Context, const FDataflowOutput* Out) const override;
 
 };
 
@@ -1189,7 +1189,7 @@ public:
 	UPROPERTY(meta = (DataflowOutput, DisplayName = "FaceSelection"))
 	FDataflowFaceSelection FaceSelection;
 
-	FCollectionFaceSelectionCustomDataflowNode(const Dataflow::FNodeParameters& InParam, FGuid InGuid = FGuid::NewGuid())
+	FCollectionFaceSelectionCustomDataflowNode(const UE::Dataflow::FNodeParameters& InParam, FGuid InGuid = FGuid::NewGuid())
 		: FDataflowNode(InParam, InGuid)
 	{
 		RegisterInputConnection(&Collection);
@@ -1198,7 +1198,7 @@ public:
 		RegisterOutputConnection(&FaceSelection);
 	}
 
-	virtual void Evaluate(Dataflow::FContext& Context, const FDataflowOutput* Out) const override;
+	virtual void Evaluate(UE::Dataflow::FContext& Context, const FDataflowOutput* Out) const override;
 
 };
 
@@ -1235,7 +1235,7 @@ public:
 	UPROPERTY(EditAnywhere, Category = "Selection")
 	bool bAllElementsMustBeSelected = false;
 
-	FCollectionSelectionConvertDataflowNode(const Dataflow::FNodeParameters& InParam, FGuid InGuid = FGuid::NewGuid())
+	FCollectionSelectionConvertDataflowNode(const UE::Dataflow::FNodeParameters& InParam, FGuid InGuid = FGuid::NewGuid())
 		: FDataflowNode(InParam, InGuid)
 	{
 		RegisterInputConnection(&Collection);
@@ -1248,7 +1248,7 @@ public:
 		RegisterOutputConnection(&VertexSelection, &VertexSelection);
 	}
 
-	virtual void Evaluate(Dataflow::FContext& Context, const FDataflowOutput* Out) const override;
+	virtual void Evaluate(UE::Dataflow::FContext& Context, const FDataflowOutput* Out) const override;
 
 };
 
@@ -1269,14 +1269,14 @@ public:
 	UPROPERTY(meta = (DataflowInput, DataflowOutput, DisplayName = "FaceSelection", DataflowPassthrough = "FaceSelection", DataflowIntrinsic))
 	FDataflowFaceSelection FaceSelection;
 
-	FCollectionFaceSelectionInvertDataflowNode(const Dataflow::FNodeParameters& InParam, FGuid InGuid = FGuid::NewGuid())
+	FCollectionFaceSelectionInvertDataflowNode(const UE::Dataflow::FNodeParameters& InParam, FGuid InGuid = FGuid::NewGuid())
 		: FDataflowNode(InParam, InGuid)
 	{
 		RegisterInputConnection(&FaceSelection);
 		RegisterOutputConnection(&FaceSelection, &FaceSelection);
 	}
 
-	virtual void Evaluate(Dataflow::FContext& Context, const FDataflowOutput* Out) const override;
+	virtual void Evaluate(UE::Dataflow::FContext& Context, const FDataflowOutput* Out) const override;
 
 };
 
@@ -1309,7 +1309,7 @@ public:
 	UPROPERTY(EditAnywhere, Category = "Random", meta = (DataflowInput, EditCondition = "bDeterministic"))
 	float RandomSeed = 0.f;
 
-	FCollectionVertexSelectionByPercentageDataflowNode(const Dataflow::FNodeParameters& InParam, FGuid InGuid = FGuid::NewGuid())
+	FCollectionVertexSelectionByPercentageDataflowNode(const UE::Dataflow::FNodeParameters& InParam, FGuid InGuid = FGuid::NewGuid())
 		: FDataflowNode(InParam, InGuid)
 	{
 		RandomSeed = FMath::RandRange(-100000, 100000);
@@ -1319,7 +1319,7 @@ public:
 		RegisterOutputConnection(&VertexSelection, &VertexSelection);
 	}
 
-	virtual void Evaluate(Dataflow::FContext& Context, const FDataflowOutput* Out) const override;
+	virtual void Evaluate(UE::Dataflow::FContext& Context, const FDataflowOutput* Out) const override;
 
 };
 
@@ -1352,7 +1352,7 @@ public:
 	UPROPERTY(meta = (DataflowOutput, DisplayName = "VertexSelection", DataflowPassthrough = "VertexSelectionA"))
 	FDataflowVertexSelection VertexSelection;
 
-	FCollectionVertexSelectionSetOperationDataflowNode(const Dataflow::FNodeParameters& InParam, FGuid InGuid = FGuid::NewGuid())
+	FCollectionVertexSelectionSetOperationDataflowNode(const UE::Dataflow::FNodeParameters& InParam, FGuid InGuid = FGuid::NewGuid())
 		: FDataflowNode(InParam, InGuid)
 	{
 		RegisterInputConnection(&VertexSelectionA);
@@ -1360,7 +1360,7 @@ public:
 		RegisterOutputConnection(&VertexSelection, &VertexSelectionA);
 	}
 
-	virtual void Evaluate(Dataflow::FContext& Context, const FDataflowOutput* Out) const override;
+	virtual void Evaluate(UE::Dataflow::FContext& Context, const FDataflowOutput* Out) const override;
 
 };
 
@@ -1374,7 +1374,7 @@ enum class ESelectionByAttrGroup : uint8
 	Material UMETA(DisplayName = "Material")
 };
 
-namespace Dataflow::Private
+namespace UE::Dataflow::Private
 {
 	inline FName GetAttributeFromEnumAsName(const ESelectionByAttrGroup Value)
 	{
@@ -1457,7 +1457,7 @@ public:
 	UPROPERTY(meta = (DataflowOutput, DisplayName = "MaterialSelection"))
 	FDataflowMaterialSelection MaterialSelection;
 
-	FCollectionSelectionByAttrDataflowNode(const Dataflow::FNodeParameters& InParam, FGuid InGuid = FGuid::NewGuid())
+	FCollectionSelectionByAttrDataflowNode(const UE::Dataflow::FNodeParameters& InParam, FGuid InGuid = FGuid::NewGuid())
 		: FDataflowNode(InParam, InGuid)
 	{
 		RegisterInputConnection(&Collection);
@@ -1470,10 +1470,10 @@ public:
 		RegisterOutputConnection(&MaterialSelection);
 	}
 
-	virtual void Evaluate(Dataflow::FContext& Context, const FDataflowOutput* Out) const override;
+	virtual void Evaluate(UE::Dataflow::FContext& Context, const FDataflowOutput* Out) const override;
 };
 
-namespace Dataflow
+namespace UE::Dataflow
 {
 	void GeometryCollectionSelectionNodes();
 }
diff --git a/Engine/Plugins/Experimental/GeometryCollectionPlugin/Source/GeometryCollectionNodes/Public/Dataflow/GeometryCollectionSkeletalMeshToCollectionNode.h b/Engine/Plugins/Experimental/GeometryCollectionPlugin/Source/GeometryCollectionNodes/Public/Dataflow/GeometryCollectionSkeletalMeshToCollectionNode.h
--- a/Engine/Plugins/Experimental/GeometryCollectionPlugin/Source/GeometryCollectionNodes/Public/Dataflow/GeometryCollectionSkeletalMeshToCollectionNode.h
+++ b/Engine/Plugins/Experimental/GeometryCollectionPlugin/Source/GeometryCollectionNodes/Public/Dataflow/GeometryCollectionSkeletalMeshToCollectionNode.h
@@ -27,13 +27,13 @@ public:
 	UPROPERTY(meta = (DataflowOutput, DisplayName = "Collection"))
 	FManagedArrayCollection Collection;
 
-	FSkeletalMeshToCollectionDataflowNode(const Dataflow::FNodeParameters& InParam, FGuid InGuid = FGuid::NewGuid())
+	FSkeletalMeshToCollectionDataflowNode(const UE::Dataflow::FNodeParameters& InParam, FGuid InGuid = FGuid::NewGuid())
 		: FDataflowNode(InParam, InGuid)
 	{
 		RegisterInputConnection(&SkeletalMesh);
 		RegisterOutputConnection(&Collection);
 	}
 
-	virtual void Evaluate(Dataflow::FContext& Context, const FDataflowOutput* Out) const override;
+	virtual void Evaluate(UE::Dataflow::FContext& Context, const FDataflowOutput* Out) const override;
 };
 
diff --git a/Engine/Plugins/Experimental/GeometryCollectionPlugin/Source/GeometryCollectionNodes/Public/Dataflow/GeometryCollectionSkeletonToCollectionNode.h b/Engine/Plugins/Experimental/GeometryCollectionPlugin/Source/GeometryCollectionNodes/Public/Dataflow/GeometryCollectionSkeletonToCollectionNode.h
--- a/Engine/Plugins/Experimental/GeometryCollectionPlugin/Source/GeometryCollectionNodes/Public/Dataflow/GeometryCollectionSkeletonToCollectionNode.h
+++ b/Engine/Plugins/Experimental/GeometryCollectionPlugin/Source/GeometryCollectionNodes/Public/Dataflow/GeometryCollectionSkeletonToCollectionNode.h
@@ -25,14 +25,14 @@ public:
 	UPROPERTY(meta = (DataflowOutput, DisplayName = "Collection"))
 	FManagedArrayCollection Collection;
 
-	FSkeletonToCollectionDataflowNode(const Dataflow::FNodeParameters& InParam, FGuid InGuid = FGuid::NewGuid())
+	FSkeletonToCollectionDataflowNode(const UE::Dataflow::FNodeParameters& InParam, FGuid InGuid = FGuid::NewGuid())
 		: FDataflowNode(InParam, InGuid)
 	{
 		RegisterInputConnection(&Skeleton);
 		RegisterOutputConnection(&Collection);
 	}
 
-	virtual void Evaluate(Dataflow::FContext& Context, const FDataflowOutput* Out) const override;
+	virtual void Evaluate(UE::Dataflow::FContext& Context, const FDataflowOutput* Out) const override;
 };
 
 
diff --git a/Engine/Plugins/Experimental/GeometryCollectionPlugin/Source/GeometryCollectionNodes/Public/Dataflow/GeometryCollectionStaticMeshToCollectionNode.h b/Engine/Plugins/Experimental/GeometryCollectionPlugin/Source/GeometryCollectionNodes/Public/Dataflow/GeometryCollectionStaticMeshToCollectionNode.h
--- a/Engine/Plugins/Experimental/GeometryCollectionPlugin/Source/GeometryCollectionNodes/Public/Dataflow/GeometryCollectionStaticMeshToCollectionNode.h
+++ b/Engine/Plugins/Experimental/GeometryCollectionPlugin/Source/GeometryCollectionNodes/Public/Dataflow/GeometryCollectionStaticMeshToCollectionNode.h
@@ -55,6 +55,6 @@ public:
 	UPROPERTY(meta = (DataflowOutput, DisplayName = "InstancedMeshes"))
 	TArray<FGeometryCollectionAutoInstanceMesh> InstancedMeshes;
 
-	FStaticMeshToCollectionDataflowNode(const Dataflow::FNodeParameters& InParam, FGuid InGuid = FGuid::NewGuid());
-	virtual void Evaluate(Dataflow::FContext& Context, const FDataflowOutput* Out) const override;
+	FStaticMeshToCollectionDataflowNode(const UE::Dataflow::FNodeParameters& InParam, FGuid InGuid = FGuid::NewGuid());
+	virtual void Evaluate(UE::Dataflow::FContext& Context, const FDataflowOutput* Out) const override;
 };
diff --git a/Engine/Plugins/Experimental/GeometryCollectionPlugin/Source/GeometryCollectionNodes/Public/Dataflow/GeometryCollectionTransferVertexAttributeNode.h b/Engine/Plugins/Experimental/GeometryCollectionPlugin/Source/GeometryCollectionNodes/Public/Dataflow/GeometryCollectionTransferVertexAttributeNode.h
--- a/Engine/Plugins/Experimental/GeometryCollectionPlugin/Source/GeometryCollectionNodes/Public/Dataflow/GeometryCollectionTransferVertexAttributeNode.h
+++ b/Engine/Plugins/Experimental/GeometryCollectionPlugin/Source/GeometryCollectionNodes/Public/Dataflow/GeometryCollectionTransferVertexAttributeNode.h
@@ -115,7 +115,7 @@ public:
 		EditConditionHides))
 	float BoundMultiplier = 0.01f;
 
-	FGeometryCollectionTransferVertexAttributeNode(const Dataflow::FNodeParameters& InParam, FGuid InGuid = FGuid::NewGuid())
+	FGeometryCollectionTransferVertexAttributeNode(const UE::Dataflow::FNodeParameters& InParam, FGuid InGuid = FGuid::NewGuid())
 		: FDataflowNode(InParam, InGuid)
 	{
 		RegisterInputConnection(&Collection);
@@ -128,7 +128,7 @@ public:
 
 private:
 
-	virtual void Evaluate(Dataflow::FContext& Context, const FDataflowOutput* Out) const override;
+	virtual void Evaluate(UE::Dataflow::FContext& Context, const FDataflowOutput* Out) const override;
 	
 	typedef Chaos::TSphere<Chaos::FReal, 3> SphereType;
 	typedef Chaos::TBoundingVolumeHierarchy<TArray<SphereType*>, TArray<int32>, Chaos::FReal, 3> BVH;
diff --git a/Engine/Plugins/Experimental/GeometryCollectionPlugin/Source/GeometryCollectionNodes/Public/Dataflow/GeometryCollectionTriangleBoundaryIndicesNode.h b/Engine/Plugins/Experimental/GeometryCollectionPlugin/Source/GeometryCollectionNodes/Public/Dataflow/GeometryCollectionTriangleBoundaryIndicesNode.h
--- a/Engine/Plugins/Experimental/GeometryCollectionPlugin/Source/GeometryCollectionNodes/Public/Dataflow/GeometryCollectionTriangleBoundaryIndicesNode.h
+++ b/Engine/Plugins/Experimental/GeometryCollectionPlugin/Source/GeometryCollectionNodes/Public/Dataflow/GeometryCollectionTriangleBoundaryIndicesNode.h
@@ -28,12 +28,12 @@ public:
 		TArray<int32> BoundaryIndicesOut;
 
 
-	FTriangleBoundaryIndicesNode(const Dataflow::FNodeParameters& InParam, FGuid InGuid = FGuid::NewGuid())
+	FTriangleBoundaryIndicesNode(const UE::Dataflow::FNodeParameters& InParam, FGuid InGuid = FGuid::NewGuid())
 		: FDataflowNode(InParam, InGuid)
 	{
 		RegisterInputConnection(&Collection);
 		RegisterOutputConnection(&BoundaryIndicesOut);
 	}
 
-	virtual void Evaluate(Dataflow::FContext& Context, const FDataflowOutput* Out) const override;
+	virtual void Evaluate(UE::Dataflow::FContext& Context, const FDataflowOutput* Out) const override;
 };
\ No newline at end of file
diff --git a/Engine/Plugins/Experimental/GeometryCollectionPlugin/Source/GeometryCollectionNodes/Public/Dataflow/GeometryCollectionUtilityNodes.h b/Engine/Plugins/Experimental/GeometryCollectionPlugin/Source/GeometryCollectionNodes/Public/Dataflow/GeometryCollectionUtilityNodes.h
--- a/Engine/Plugins/Experimental/GeometryCollectionPlugin/Source/GeometryCollectionNodes/Public/Dataflow/GeometryCollectionUtilityNodes.h
+++ b/Engine/Plugins/Experimental/GeometryCollectionPlugin/Source/GeometryCollectionNodes/Public/Dataflow/GeometryCollectionUtilityNodes.h
@@ -99,9 +99,9 @@ public:
 	UPROPERTY(meta = (DataflowOutput))
 	FDataflowConvexDecompositionSettings DecompositionSettings;
 
-	FMakeDataflowConvexDecompositionSettingsNode(const Dataflow::FNodeParameters& InParam, FGuid InGuid = FGuid::NewGuid());
+	FMakeDataflowConvexDecompositionSettingsNode(const UE::Dataflow::FNodeParameters& InParam, FGuid InGuid = FGuid::NewGuid());
 
-	virtual void Evaluate(Dataflow::FContext& Context, const FDataflowOutput* Out) const override;
+	virtual void Evaluate(UE::Dataflow::FContext& Context, const FDataflowOutput* Out) const override;
 };
 
 USTRUCT(meta = (DataflowGeometryCollection))
@@ -141,9 +141,9 @@ public:
 	UPROPERTY(EditAnywhere, Category = Convex, meta = (DataflowInput))
 	FDataflowConvexDecompositionSettings ConvexDecompositionSettings;
 
-	FCreateLeafConvexHullsDataflowNode(const Dataflow::FNodeParameters& InParam, FGuid InGuid = FGuid::NewGuid());
+	FCreateLeafConvexHullsDataflowNode(const UE::Dataflow::FNodeParameters& InParam, FGuid InGuid = FGuid::NewGuid());
 
-	virtual void Evaluate(Dataflow::FContext& Context, const FDataflowOutput* Out) const override;
+	virtual void Evaluate(UE::Dataflow::FContext& Context, const FDataflowOutput* Out) const override;
 };
 
 USTRUCT(meta = (DataflowGeometryCollection))
@@ -179,9 +179,9 @@ public:
 	UPROPERTY(EditAnywhere, Category = "Convex")
 	bool bUseExistingVertices = false;
 
-	FSimplifyConvexHullsDataflowNode(const Dataflow::FNodeParameters& InParam, FGuid InGuid = FGuid::NewGuid());
+	FSimplifyConvexHullsDataflowNode(const UE::Dataflow::FNodeParameters& InParam, FGuid InGuid = FGuid::NewGuid());
 
-	virtual void Evaluate(Dataflow::FContext& Context, const FDataflowOutput* Out) const override;
+	virtual void Evaluate(UE::Dataflow::FContext& Context, const FDataflowOutput* Out) const override;
 };
 
 /**
@@ -219,9 +219,9 @@ public:
 	UPROPERTY(EditAnywhere, Category = AutomaticOverlapRemoval, meta = (DataflowInput, DisplayName = "Max Removal Fraction", ClampMin = 0.01f, ClampMax = 1.f))
 	float CanRemoveFraction = 0.3f;
 
-	FCreateNonOverlappingConvexHullsDataflowNode(const Dataflow::FNodeParameters& InParam, FGuid InGuid = FGuid::NewGuid());
+	FCreateNonOverlappingConvexHullsDataflowNode(const UE::Dataflow::FNodeParameters& InParam, FGuid InGuid = FGuid::NewGuid());
 
-	virtual void Evaluate(Dataflow::FContext& Context, const FDataflowOutput* Out) const override;
+	virtual void Evaluate(UE::Dataflow::FContext& Context, const FDataflowOutput* Out) const override;
 
 };
 
@@ -323,9 +323,9 @@ public:
 	UPROPERTY(EditAnywhere, Category = NegativeSpace, meta = (DataflowInput, ClampMin = 0, Units = cm, EditCondition = "bProtectNegativeSpace", EditConditionHides))
 	double MinRadius = 10.0;
 
-	FGenerateClusterConvexHullsFromLeafHullsDataflowNode(const Dataflow::FNodeParameters& InParam, FGuid InGuid = FGuid::NewGuid());
+	FGenerateClusterConvexHullsFromLeafHullsDataflowNode(const UE::Dataflow::FNodeParameters& InParam, FGuid InGuid = FGuid::NewGuid());
 
-	virtual void Evaluate(Dataflow::FContext& Context, const FDataflowOutput* Out) const override;
+	virtual void Evaluate(UE::Dataflow::FContext& Context, const FDataflowOutput* Out) const override;
 
 };
 
@@ -400,9 +400,9 @@ public:
 	UPROPERTY(EditAnywhere, Category = NegativeSpace, meta = (DataflowInput, ClampMin = 0, Units = cm, EditCondition = "bProtectNegativeSpace", EditConditionHides))
 	double MinRadius = 10.0;
 
-	FGenerateClusterConvexHullsFromChildrenHullsDataflowNode(const Dataflow::FNodeParameters& InParam, FGuid InGuid = FGuid::NewGuid());
+	FGenerateClusterConvexHullsFromChildrenHullsDataflowNode(const UE::Dataflow::FNodeParameters& InParam, FGuid InGuid = FGuid::NewGuid());
 
-	virtual void Evaluate(Dataflow::FContext& Context, const FDataflowOutput* Out) const override;
+	virtual void Evaluate(UE::Dataflow::FContext& Context, const FDataflowOutput* Out) const override;
 
 };
 
@@ -422,7 +422,7 @@ public:
 	UPROPERTY(meta = (DataflowInput))
 	FDataflowTransformSelection TransformSelection;
 
-	FClearConvexHullsDataflowNode(const Dataflow::FNodeParameters& InParam, FGuid InGuid = FGuid::NewGuid())
+	FClearConvexHullsDataflowNode(const UE::Dataflow::FNodeParameters& InParam, FGuid InGuid = FGuid::NewGuid())
 		: FDataflowNode(InParam, InGuid)
 	{
 		RegisterInputConnection(&Collection);
@@ -431,7 +431,7 @@ public:
 		RegisterOutputConnection(&Collection, &Collection);
 	}
 
-	virtual void Evaluate(Dataflow::FContext& Context, const FDataflowOutput* Out) const override;
+	virtual void Evaluate(UE::Dataflow::FContext& Context, const FDataflowOutput* Out) const override;
 };
 
 
@@ -502,9 +502,9 @@ public:
 	UPROPERTY(EditAnywhere, Category = NegativeSpace, meta = (DataflowInput, ClampMin = 0, Units = cm, EditCondition = "bProtectNegativeSpace", EditConditionHides))
 	double MinRadius = 10.0;
 
-	FMergeConvexHullsDataflowNode(const Dataflow::FNodeParameters& InParam, FGuid InGuid = FGuid::NewGuid());
+	FMergeConvexHullsDataflowNode(const UE::Dataflow::FNodeParameters& InParam, FGuid InGuid = FGuid::NewGuid());
 
-	virtual void Evaluate(Dataflow::FContext& Context, const FDataflowOutput* Out) const override;
+	virtual void Evaluate(UE::Dataflow::FContext& Context, const FDataflowOutput* Out) const override;
 
 };
 
@@ -524,9 +524,9 @@ public:
 	UPROPERTY(meta = (DataflowInput, DataflowOutput))
 	FManagedArrayCollection Collection;
 
-	FUpdateVolumeAttributesDataflowNode(const Dataflow::FNodeParameters& InParam, FGuid InGuid = FGuid::NewGuid());
+	FUpdateVolumeAttributesDataflowNode(const UE::Dataflow::FNodeParameters& InParam, FGuid InGuid = FGuid::NewGuid());
 
-	virtual void Evaluate(Dataflow::FContext& Context, const FDataflowOutput* Out) const override;
+	virtual void Evaluate(UE::Dataflow::FContext& Context, const FDataflowOutput* Out) const override;
 };
 
 /**
@@ -558,9 +558,9 @@ public:
 	UPROPERTY(EditAnywhere, Category = Options)
 	bool bVolumeOfUnion = false;
 
-	FGetConvexHullVolumeDataflowNode(const Dataflow::FNodeParameters& InParam, FGuid InGuid = FGuid::NewGuid());
+	FGetConvexHullVolumeDataflowNode(const UE::Dataflow::FNodeParameters& InParam, FGuid InGuid = FGuid::NewGuid());
 
-	virtual void Evaluate(Dataflow::FContext& Context, const FDataflowOutput* Out) const override;
+	virtual void Evaluate(UE::Dataflow::FContext& Context, const FDataflowOutput* Out) const override;
 };
 
 /**
@@ -643,7 +643,7 @@ public:
 	UPROPERTY(EditAnywhere, Category = "Collision", meta = (DataflowInput, DisplayName = "Point Spacing", UIMin = 0.f, EditCondition = "AddSamplesForCollision"));
 	float CollisionSampleSpacing = 50.f;
 
-	FFixTinyGeoDataflowNode(const Dataflow::FNodeParameters& InParam, FGuid InGuid = FGuid::NewGuid())
+	FFixTinyGeoDataflowNode(const UE::Dataflow::FNodeParameters& InParam, FGuid InGuid = FGuid::NewGuid())
 		: FDataflowNode(InParam, InGuid)
 	{
 		RegisterInputConnection(&Collection);
@@ -652,7 +652,7 @@ public:
 		RegisterOutputConnection(&Collection, &Collection);
 	}
 
-	virtual void Evaluate(Dataflow::FContext& Context, const FDataflowOutput* Out) const override;
+	virtual void Evaluate(UE::Dataflow::FContext& Context, const FDataflowOutput* Out) const override;
 };
 
 /**
@@ -691,7 +691,7 @@ public:
 	UPROPERTY(EditAnywhere, Category = RecomputeSettings, AdvancedDisplay)
 	bool bOnlyInternalSurfaces = true;
 
-	FRecomputeNormalsInGeometryCollectionDataflowNode(const Dataflow::FNodeParameters& InParam, FGuid InGuid = FGuid::NewGuid())
+	FRecomputeNormalsInGeometryCollectionDataflowNode(const UE::Dataflow::FNodeParameters& InParam, FGuid InGuid = FGuid::NewGuid())
 		: FDataflowNode(InParam, InGuid)
 	{
 		RegisterInputConnection(&Collection);
@@ -699,7 +699,7 @@ public:
 		RegisterOutputConnection(&Collection, &Collection);
 	}
 
-	virtual void Evaluate(Dataflow::FContext& Context, const FDataflowOutput* Out) const override;
+	virtual void Evaluate(UE::Dataflow::FContext& Context, const FDataflowOutput* Out) const override;
 };
 
 /**
@@ -740,7 +740,7 @@ public:
 	UPROPERTY(EditAnywhere, Category = "Collision", meta = (DataflowInput, DisplayName = "Point Spacing", UIMin = 0.f, EditCondition = "AddSamplesForCollision"));
 	float CollisionSampleSpacing = 50.f;
 
-	FResampleGeometryCollectionDataflowNode(const Dataflow::FNodeParameters& InParam, FGuid InGuid = FGuid::NewGuid())
+	FResampleGeometryCollectionDataflowNode(const UE::Dataflow::FNodeParameters& InParam, FGuid InGuid = FGuid::NewGuid())
 		: FDataflowNode(InParam, InGuid)
 	{
 		RegisterInputConnection(&Collection);
@@ -749,7 +749,7 @@ public:
 		RegisterOutputConnection(&Collection, &Collection);
 	}
 
-	virtual void Evaluate(Dataflow::FContext& Context, const FDataflowOutput* Out) const override;
+	virtual void Evaluate(UE::Dataflow::FContext& Context, const FDataflowOutput* Out) const override;
 };
 
 /**
@@ -781,17 +781,17 @@ public:
 	UPROPERTY(EditAnywhere, Category = Clustering)
 	bool bRemoveDanglingClusters = false;
 
-	FValidateGeometryCollectionDataflowNode(const Dataflow::FNodeParameters& InParam, FGuid InGuid = FGuid::NewGuid())
+	FValidateGeometryCollectionDataflowNode(const UE::Dataflow::FNodeParameters& InParam, FGuid InGuid = FGuid::NewGuid())
 		: FDataflowNode(InParam, InGuid)
 	{
 		RegisterInputConnection(&Collection);
 		RegisterOutputConnection(&Collection, &Collection);
 	}
 
-	virtual void Evaluate(Dataflow::FContext& Context, const FDataflowOutput* Out) const override;
+	virtual void Evaluate(UE::Dataflow::FContext& Context, const FDataflowOutput* Out) const override;
 };
 
-namespace Dataflow
+namespace UE::Dataflow
 {
 	void GeometryCollectionUtilityNodes();
 }
diff --git a/Engine/Plugins/Experimental/GeometryCollectionPlugin/Source/GeometryCollectionNodes/Public/Dataflow/GeometryCollectionVertexScalarToVertexIndicesNode.h b/Engine/Plugins/Experimental/GeometryCollectionPlugin/Source/GeometryCollectionNodes/Public/Dataflow/GeometryCollectionVertexScalarToVertexIndicesNode.h
--- a/Engine/Plugins/Experimental/GeometryCollectionPlugin/Source/GeometryCollectionNodes/Public/Dataflow/GeometryCollectionVertexScalarToVertexIndicesNode.h
+++ b/Engine/Plugins/Experimental/GeometryCollectionPlugin/Source/GeometryCollectionNodes/Public/Dataflow/GeometryCollectionVertexScalarToVertexIndicesNode.h
@@ -33,9 +33,9 @@ public:
 	UPROPERTY(Meta = (DataflowOutput, DisplayName = "VertexIndices"))
 	TArray<int32> VertexIndices = {};
 
-	FGeometryCollectionVertexScalarToVertexIndicesNode(const Dataflow::FNodeParameters& InParam, FGuid InGuid = FGuid::NewGuid());
+	FGeometryCollectionVertexScalarToVertexIndicesNode(const UE::Dataflow::FNodeParameters& InParam, FGuid InGuid = FGuid::NewGuid());
 
 private:
 
-	virtual void Evaluate(Dataflow::FContext& Context, const FDataflowOutput* Out) const override;
+	virtual void Evaluate(UE::Dataflow::FContext& Context, const FDataflowOutput* Out) const override;
 };
diff --git a/Engine/Plugins/Experimental/GeometryCollectionPlugin/Source/GeometryCollectionNodes/Public/Dataflow/GeometryCollectionVerticesNodes.h b/Engine/Plugins/Experimental/GeometryCollectionPlugin/Source/GeometryCollectionNodes/Public/Dataflow/GeometryCollectionVerticesNodes.h
--- a/Engine/Plugins/Experimental/GeometryCollectionPlugin/Source/GeometryCollectionNodes/Public/Dataflow/GeometryCollectionVerticesNodes.h
+++ b/Engine/Plugins/Experimental/GeometryCollectionPlugin/Source/GeometryCollectionNodes/Public/Dataflow/GeometryCollectionVerticesNodes.h
@@ -37,7 +37,7 @@ public:
 	FString AttributeName = FString("Vertex");
 
 
-	FTransformCollectionAttributeDataflowNode(const Dataflow::FNodeParameters& InParam, FGuid InGuid = FGuid::NewGuid())
+	FTransformCollectionAttributeDataflowNode(const UE::Dataflow::FNodeParameters& InParam, FGuid InGuid = FGuid::NewGuid())
 		: FDataflowNode(InParam, InGuid)
 	{
 		RegisterInputConnection(&TransformIn);
@@ -45,10 +45,10 @@ public:
 		RegisterOutputConnection(&Collection);
 	}
 
-	virtual void Evaluate(Dataflow::FContext& Context, const FDataflowOutput* Out) const override;
+	virtual void Evaluate(UE::Dataflow::FContext& Context, const FDataflowOutput* Out) const override;
 };
 
-namespace Dataflow
+namespace UE::Dataflow
 {
 	void GeometryCollectionVerticesNodes();
 }
diff --git a/Engine/Plugins/Experimental/GeometryCollectionPlugin/Source/GeometryCollectionNodes/Public/Dataflow/SetVertexColorFromFloatArrayNode.h b/Engine/Plugins/Experimental/GeometryCollectionPlugin/Source/GeometryCollectionNodes/Public/Dataflow/SetVertexColorFromFloatArrayNode.h
--- a/Engine/Plugins/Experimental/GeometryCollectionPlugin/Source/GeometryCollectionNodes/Public/Dataflow/SetVertexColorFromFloatArrayNode.h
+++ b/Engine/Plugins/Experimental/GeometryCollectionPlugin/Source/GeometryCollectionNodes/Public/Dataflow/SetVertexColorFromFloatArrayNode.h
@@ -35,7 +35,7 @@ public:
 	FLinearColor Color = FLinearColor(FColor::White);
 
 
-	FSetVertexColorFromFloatArrayDataflowNode(const Dataflow::FNodeParameters& InParam, FGuid InGuid = FGuid::NewGuid())
+	FSetVertexColorFromFloatArrayDataflowNode(const UE::Dataflow::FNodeParameters& InParam, FGuid InGuid = FGuid::NewGuid())
 		: FDataflowNode(InParam, InGuid)
 	{
 		RegisterInputConnection(&Collection);
@@ -43,6 +43,6 @@ public:
 		RegisterOutputConnection(&Collection, &Collection);
 	}
 
-	virtual void Evaluate(Dataflow::FContext& Context, const FDataflowOutput* Out) const override;
+	virtual void Evaluate(UE::Dataflow::FContext& Context, const FDataflowOutput* Out) const override;
 
 };
diff --git a/Engine/Plugins/Experimental/GeometryCollectionPlugin/Source/GeometryCollectionNodes/Public/Dataflow/SetVertexColorFromVertexIndicesNode.h b/Engine/Plugins/Experimental/GeometryCollectionPlugin/Source/GeometryCollectionNodes/Public/Dataflow/SetVertexColorFromVertexIndicesNode.h
--- a/Engine/Plugins/Experimental/GeometryCollectionPlugin/Source/GeometryCollectionNodes/Public/Dataflow/SetVertexColorFromVertexIndicesNode.h
+++ b/Engine/Plugins/Experimental/GeometryCollectionPlugin/Source/GeometryCollectionNodes/Public/Dataflow/SetVertexColorFromVertexIndicesNode.h
@@ -28,7 +28,7 @@ public:
 	UPROPERTY(EditAnywhere, Category = "Color")
 	FLinearColor SelectedColor = FLinearColor(FColor::Yellow);
 
-	FSetVertexColorFromVertexIndicesDataflowNode(const Dataflow::FNodeParameters& InParam, FGuid InGuid = FGuid::NewGuid())
+	FSetVertexColorFromVertexIndicesDataflowNode(const UE::Dataflow::FNodeParameters& InParam, FGuid InGuid = FGuid::NewGuid())
 		: FDataflowNode(InParam, InGuid)
 	{
 		RegisterInputConnection(&Collection);
@@ -36,6 +36,6 @@ public:
 		RegisterOutputConnection(&Collection, &Collection);
 	}
 
-	virtual void Evaluate(Dataflow::FContext& Context, const FDataflowOutput* Out) const override;
+	virtual void Evaluate(UE::Dataflow::FContext& Context, const FDataflowOutput* Out) const override;
 
 };
diff --git a/Engine/Plugins/Experimental/GeometryCollectionPlugin/Source/GeometryCollectionNodes/Public/Dataflow/SetVertexColorFromVertexSelectionNode.h b/Engine/Plugins/Experimental/GeometryCollectionPlugin/Source/GeometryCollectionNodes/Public/Dataflow/SetVertexColorFromVertexSelectionNode.h
--- a/Engine/Plugins/Experimental/GeometryCollectionPlugin/Source/GeometryCollectionNodes/Public/Dataflow/SetVertexColorFromVertexSelectionNode.h
+++ b/Engine/Plugins/Experimental/GeometryCollectionPlugin/Source/GeometryCollectionNodes/Public/Dataflow/SetVertexColorFromVertexSelectionNode.h
@@ -29,7 +29,7 @@ public:
 	UPROPERTY(EditAnywhere, Category = "Color")
 	FLinearColor SelectedColor = FLinearColor(FColor::Yellow);
 
-	FSetVertexColorFromVertexSelectionDataflowNode(const Dataflow::FNodeParameters& InParam, FGuid InGuid = FGuid::NewGuid())
+	FSetVertexColorFromVertexSelectionDataflowNode(const UE::Dataflow::FNodeParameters& InParam, FGuid InGuid = FGuid::NewGuid())
 		: FDataflowNode(InParam, InGuid)
 	{
 		RegisterInputConnection(&Collection);
@@ -37,6 +37,6 @@ public:
 		RegisterOutputConnection(&Collection, &Collection);
 	}
 
-	virtual void Evaluate(Dataflow::FContext& Context, const FDataflowOutput* Out) const override;
+	virtual void Evaluate(UE::Dataflow::FContext& Context, const FDataflowOutput* Out) const override;
 
 };
diff --git a/Engine/Source/Runtime/Experimental/Dataflow/Core/Private/Dataflow/DataflowConnection.cpp b/Engine/Source/Runtime/Experimental/Dataflow/Core/Private/Dataflow/DataflowConnection.cpp
--- a/Engine/Source/Runtime/Experimental/Dataflow/Core/Private/Dataflow/DataflowConnection.cpp
+++ b/Engine/Source/Runtime/Experimental/Dataflow/Core/Private/Dataflow/DataflowConnection.cpp
@@ -7,10 +7,10 @@
 
 #include UE_INLINE_GENERATED_CPP_BY_NAME(DataflowConnection)
 
-const Dataflow::FConnectionKey Dataflow::FConnectionKey::Invalid = { (uint32)INDEX_NONE, INDEX_NONE, (uint32)INDEX_NONE };
-const Dataflow::FPin Dataflow::FPin::InvalidPin = { Dataflow::FPin::EDirection::NONE, NAME_None, NAME_None };
+const UE::Dataflow::FConnectionKey UE::Dataflow::FConnectionKey::Invalid = { (uint32)INDEX_NONE, INDEX_NONE, (uint32)INDEX_NONE };
+const UE::Dataflow::FPin UE::Dataflow::FPin::InvalidPin = { UE::Dataflow::FPin::EDirection::NONE, NAME_None, NAME_None };
 
-FDataflowConnection::FDataflowConnection(Dataflow::FPin::EDirection InDirection, FName InType, FName InName, FDataflowNode* InOwningNode, const FProperty* InProperty, FGuid InGuid)
+FDataflowConnection::FDataflowConnection(UE::Dataflow::FPin::EDirection InDirection, FName InType, FName InName, FDataflowNode* InOwningNode, const FProperty* InProperty, FGuid InGuid)
 	: Type(InType)
 	, Name(InName)
 	, OwningNode(InOwningNode)
@@ -22,7 +22,7 @@ FDataflowConnection::FDataflowConnection(Dataflow::FPin::EDirection InDirection,
 	InitFromType();
 }
 
-FDataflowConnection::FDataflowConnection(Dataflow::FPin::EDirection InDirection, const Dataflow::FConnectionParameters& Params)
+FDataflowConnection::FDataflowConnection(UE::Dataflow::FPin::EDirection InDirection, const UE::Dataflow::FConnectionParameters& Params)
 	: Type(Params.Type)
 	, Name(Params.Name)
 	, OwningNode(Params.Owner)
@@ -62,9 +62,9 @@ FGuid FDataflowConnection::GetOwningNodeGuid() const
 	return OwningNode ? OwningNode->GetGuid() : FGuid();
 }
 
-Dataflow::FTimestamp FDataflowConnection::GetOwningNodeTimestamp() const
+UE::Dataflow::FTimestamp FDataflowConnection::GetOwningNodeTimestamp() const
 {
-	return OwningNode ? OwningNode->GetTimestamp() : Dataflow::FTimestamp::Invalid;
+	return OwningNode ? OwningNode->GetTimestamp() : UE::Dataflow::FTimestamp::Invalid;
 }
 
 uint32 FDataflowConnection::GetOwningNodeValueHash() const
diff --git a/Engine/Source/Runtime/Experimental/Dataflow/Core/Private/Dataflow/DataflowCoreNodes.cpp b/Engine/Source/Runtime/Experimental/Dataflow/Core/Private/Dataflow/DataflowCoreNodes.cpp
--- a/Engine/Source/Runtime/Experimental/Dataflow/Core/Private/Dataflow/DataflowCoreNodes.cpp
+++ b/Engine/Source/Runtime/Experimental/Dataflow/Core/Private/Dataflow/DataflowCoreNodes.cpp
@@ -4,7 +4,7 @@
 #include "Dataflow/DataflowNode.h"
 #include "Dataflow/DataflowNodeFactory.h"
 
-namespace Dataflow
+namespace UE::Dataflow
 {
 	void RegisterCoreNodes()
 	{
@@ -15,7 +15,7 @@ namespace Dataflow
 	}
 }
 
-FDataflowReRouteNode::FDataflowReRouteNode(const Dataflow::FNodeParameters& Param, FGuid InGuid)
+FDataflowReRouteNode::FDataflowReRouteNode(const UE::Dataflow::FNodeParameters& Param, FGuid InGuid)
 	: Super(Param, InGuid)
 {
 	RegisterInputConnection(&Value);
@@ -23,7 +23,7 @@ FDataflowReRouteNode::FDataflowReRouteNode(const Dataflow::FNodeParameters& Para
 		.SetPassthroughInput(&Value);
 }
 
-void FDataflowReRouteNode::Evaluate(Dataflow::FContext& Context, const FDataflowOutput* Out) const
+void FDataflowReRouteNode::Evaluate(UE::Dataflow::FContext& Context, const FDataflowOutput* Out) const
 {
 	ForwardInput(Context, &Value, &Value);
 }
@@ -38,7 +38,7 @@ bool FDataflowReRouteNode::OnOutputTypeChanged(const FDataflowOutput* Input)
 	return SetInputConcreteType(&Value, Input->GetType());
 }
 
-FDataflowBranchNode::FDataflowBranchNode(const Dataflow::FNodeParameters& Param, FGuid InGuid)
+FDataflowBranchNode::FDataflowBranchNode(const UE::Dataflow::FNodeParameters& Param, FGuid InGuid)
 	: Super(Param, InGuid)
 {
 	RegisterInputConnection(&TrueValue);
@@ -47,7 +47,7 @@ FDataflowBranchNode::FDataflowBranchNode(const Dataflow::FNodeParameters& Param,
 	RegisterOutputConnection(&Result);
 }
 
-void FDataflowBranchNode::Evaluate(Dataflow::FContext& Context, const FDataflowOutput* Out) const
+void FDataflowBranchNode::Evaluate(UE::Dataflow::FContext& Context, const FDataflowOutput* Out) const
 {
 	if (Out->IsA(&Result))
 	{
@@ -86,7 +86,7 @@ bool FDataflowBranchNode::OnOutputTypeChanged(const FDataflowOutput* Input)
 		);
 }
 
-FDataflowSelectNode::FDataflowSelectNode(const Dataflow::FNodeParameters& Param, FGuid InGuid)
+FDataflowSelectNode::FDataflowSelectNode(const UE::Dataflow::FNodeParameters& Param, FGuid InGuid)
 	: Super(Param, InGuid)
 {
 	// Add two sets of pins to start.
@@ -100,14 +100,14 @@ FDataflowSelectNode::FDataflowSelectNode(const Dataflow::FNodeParameters& Param,
 	check(NumRequiredDataflowInputs + NumInitialInputs == GetNumInputs()); // Update NumRequiredDataflowInputs when adding more inputs. This is used by Serialize
 }
 
-void FDataflowSelectNode::Evaluate(Dataflow::FContext& Context, const FDataflowOutput* Out) const
+void FDataflowSelectNode::Evaluate(UE::Dataflow::FContext& Context, const FDataflowOutput* Out) const
 {
 	if (Out->IsA(&Result))
 	{
 		const int32 InSelectedIndex = GetValue<int32>(Context, &SelectedIndex);
 		if (Inputs.IsValidIndex(InSelectedIndex))
 		{
-			const Dataflow::TConnectionReference<FDataflowAnyType> SelectedInputReference = GetConnectionReference(InSelectedIndex);
+			const UE::Dataflow::TConnectionReference<FDataflowAnyType> SelectedInputReference = GetConnectionReference(InSelectedIndex);
 			if (IsConnected(SelectedInputReference))
 			{
 				ForwardInput(Context, SelectedInputReference, &Result);
@@ -141,7 +141,7 @@ bool FDataflowSelectNode::OnOutputTypeChanged(const FDataflowOutput* Input)
 	return bResult;
 }
 
-TArray<Dataflow::FPin> FDataflowSelectNode::AddPins()
+TArray<UE::Dataflow::FPin> FDataflowSelectNode::AddPins()
 {
 	const int32 Index = Inputs.AddDefaulted();
 	const FDataflowInput& Input = RegisterInputArrayConnection(GetConnectionReference(Index));
@@ -152,21 +152,21 @@ TArray<Dataflow::FPin> FDataflowSelectNode::AddPins()
 		check(Input0);
 		SetInputConcreteType(GetConnectionReference(Index), Input0->GetType());
 	}
-	return { { Dataflow::FPin::EDirection::INPUT, Input.GetType(), Input.GetName() } };
+	return { { UE::Dataflow::FPin::EDirection::INPUT, Input.GetType(), Input.GetName() } };
 }
 
-TArray<Dataflow::FPin> FDataflowSelectNode::GetPinsToRemove() const
+TArray<UE::Dataflow::FPin> FDataflowSelectNode::GetPinsToRemove() const
 {
 	const int32 Index = Inputs.Num() - 1;
 	check(Inputs.IsValidIndex(Index));
 	if (const FDataflowInput* const Input = FindInput(GetConnectionReference(Index)))
 	{
-		return { { Dataflow::FPin::EDirection::INPUT, Input->GetType(), Input->GetName() } };
+		return { { UE::Dataflow::FPin::EDirection::INPUT, Input->GetType(), Input->GetName() } };
 	}
 	return Super::GetPinsToRemove();
 }
 
-void FDataflowSelectNode::OnPinRemoved(const Dataflow::FPin& Pin)
+void FDataflowSelectNode::OnPinRemoved(const UE::Dataflow::FPin& Pin)
 {
 	const int32 Index = Inputs.Num() - 1;
 	check(Inputs.IsValidIndex(Index));
@@ -220,18 +220,18 @@ void FDataflowSelectNode::Serialize(FArchive& Ar)
 	}
 }
 
-Dataflow::TConnectionReference<FDataflowAnyType> FDataflowSelectNode::GetConnectionReference(int32 Index) const
+UE::Dataflow::TConnectionReference<FDataflowAnyType> FDataflowSelectNode::GetConnectionReference(int32 Index) const
 {
 	return { &Inputs[Index], Index, &Inputs };
 }
 
-FDataflowPrintNode::FDataflowPrintNode(const Dataflow::FNodeParameters& InParam, FGuid InGuid)
+FDataflowPrintNode::FDataflowPrintNode(const UE::Dataflow::FNodeParameters& InParam, FGuid InGuid)
 	: FDataflowNode(InParam, InGuid)
 {
 	RegisterInputConnection(&Value);
 }
 
-void FDataflowPrintNode::Evaluate(Dataflow::FContext& Context, const FDataflowOutput* Out) const
+void FDataflowPrintNode::Evaluate(UE::Dataflow::FContext& Context, const FDataflowOutput* Out) const
 {
 	const FString InValue = GetValue(Context, &Value);
 	UE_LOG(LogTemp, Warning, TEXT("[Dataflow Print] %s"), *InValue);
diff --git a/Engine/Source/Runtime/Experimental/Dataflow/Core/Private/Dataflow/DataflowInputOutput.cpp b/Engine/Source/Runtime/Experimental/Dataflow/Core/Private/Dataflow/DataflowInputOutput.cpp
--- a/Engine/Source/Runtime/Experimental/Dataflow/Core/Private/Dataflow/DataflowInputOutput.cpp
+++ b/Engine/Source/Runtime/Experimental/Dataflow/Core/Private/Dataflow/DataflowInputOutput.cpp
@@ -14,15 +14,15 @@ FDataflowInput FDataflowInput::NoOpInput = FDataflowInput();
 FDataflowOutput FDataflowOutput::NoOpOutput = FDataflowOutput();
 
 
-FDataflowInput::FDataflowInput(const Dataflow::FInputParameters& Param, FGuid InGuid)
-	: FDataflowConnection(Dataflow::FPin::EDirection::INPUT, Param)
+FDataflowInput::FDataflowInput(const UE::Dataflow::FInputParameters& Param, FGuid InGuid)
+	: FDataflowConnection(UE::Dataflow::FPin::EDirection::INPUT, Param)
 	, Connection(nullptr)
 {
 	Guid = InGuid;
 }
 
-FDataflowInput::FDataflowInput(const Dataflow::FInputParameters& Param)
-	: FDataflowConnection(Dataflow::FPin::EDirection::INPUT, Param)
+FDataflowInput::FDataflowInput(const UE::Dataflow::FInputParameters& Param)
+	: FDataflowConnection(UE::Dataflow::FPin::EDirection::INPUT, Param)
 	, Connection(nullptr)
 {
 }
@@ -69,12 +69,12 @@ const TArray< const FDataflowOutput* > FDataflowInput::GetConnectedOutputs() con
 	return RetList;
 }
 
-void FDataflowInput::Invalidate(const Dataflow::FTimestamp& ModifiedTimestamp)
+void FDataflowInput::Invalidate(const UE::Dataflow::FTimestamp& ModifiedTimestamp)
 {
 	OwningNode->Invalidate(ModifiedTimestamp);
 }
 
-void FDataflowInput::PullValue(Dataflow::FContext& Context) const
+void FDataflowInput::PullValue(UE::Dataflow::FContext& Context) const
 {
 	if (GetConnectedOutputs().Num())
 	{
@@ -119,7 +119,7 @@ void FDataflowInput::FixAndPropagateType(FName InType)
 	}
 }
 
-FDataflowArrayInput::FDataflowArrayInput(int32 InIndex, const Dataflow::FArrayInputParameters& Param)
+FDataflowArrayInput::FDataflowArrayInput(int32 InIndex, const UE::Dataflow::FArrayInputParameters& Param)
 	: FDataflowInput(Param)
 	, Index(InIndex)
 	, ElementOffset(Param.InnerOffset)
@@ -149,15 +149,15 @@ void* FDataflowArrayInput::RealAddress() const
 //
 
 
-FDataflowOutput::FDataflowOutput(const Dataflow::FOutputParameters& Param, FGuid InGuid)
-	: FDataflowConnection(Dataflow::FPin::EDirection::OUTPUT, Param)
+FDataflowOutput::FDataflowOutput(const UE::Dataflow::FOutputParameters& Param, FGuid InGuid)
+	: FDataflowConnection(UE::Dataflow::FPin::EDirection::OUTPUT, Param)
 {
 	Guid = InGuid;
 	OutputLock = MakeShared<FCriticalSection>();
 }
 
-FDataflowOutput::FDataflowOutput(const Dataflow::FOutputParameters& Param)
-	: FDataflowConnection(Dataflow::FPin::EDirection::OUTPUT, Param)
+FDataflowOutput::FDataflowOutput(const UE::Dataflow::FOutputParameters& Param)
+	: FDataflowConnection(UE::Dataflow::FPin::EDirection::OUTPUT, Param)
 {
 	OutputLock = MakeShared<FCriticalSection>();
 }
@@ -202,7 +202,7 @@ bool FDataflowOutput::RemoveConnection(FDataflowConnection* InInput)
 	Connections.RemoveSwap((FDataflowInput*)InInput); return true;
 }
 
-FDataflowOutput& FDataflowOutput::SetPassthroughInput(const Dataflow::FConnectionReference& Reference)
+FDataflowOutput& FDataflowOutput::SetPassthroughInput(const UE::Dataflow::FConnectionReference& Reference)
 {
 	check(OwningNode);
 	const FDataflowInput* const PassthroughInput = OwningNode->FindInput(Reference);
@@ -211,9 +211,9 @@ FDataflowOutput& FDataflowOutput::SetPassthroughInput(const Dataflow::FConnectio
 	return *this;
 }
 
-FDataflowOutput& FDataflowOutput::SetPassthroughInput(const Dataflow::FConnectionKey& Key)
+FDataflowOutput& FDataflowOutput::SetPassthroughInput(const UE::Dataflow::FConnectionKey& Key)
 {
-	check(Key == Dataflow::FConnectionKey::Invalid || !OwningNode || OwningNode->FindInput(Key));
+	check(Key == UE::Dataflow::FConnectionKey::Invalid || !OwningNode || OwningNode->FindInput(Key));
 	PassthroughKey = Key;
 	return *this;
 }
@@ -223,7 +223,7 @@ const FDataflowInput* FDataflowOutput::GetPassthroughInput() const
 	return OwningNode ? OwningNode->FindInput(PassthroughKey) : nullptr;
 }
 
-void FDataflowOutput::Invalidate(const Dataflow::FTimestamp& ModifiedTimestamp)
+void FDataflowOutput::Invalidate(const UE::Dataflow::FTimestamp& ModifiedTimestamp)
 {
 	for (FDataflowConnection* Con : GetConnections())
 	{
@@ -231,7 +231,7 @@ void FDataflowOutput::Invalidate(const Dataflow::FTimestamp& ModifiedTimestamp)
 	}
 }
 
-bool FDataflowOutput::Evaluate(Dataflow::FContext& Context) const
+bool FDataflowOutput::Evaluate(UE::Dataflow::FContext& Context) const
 {
 	check(OwningNode);
 
@@ -248,11 +248,11 @@ bool FDataflowOutput::Evaluate(Dataflow::FContext& Context) const
 	return false;
 }
 
-bool FDataflowOutput::EvaluateImpl(Dataflow::FContext& Context) const
+bool FDataflowOutput::EvaluateImpl(UE::Dataflow::FContext& Context) const
 {
 	UE_LOG(LogChaosDataflow, Verbose, TEXT("FDataflowOutput::EvaluateImpl(): Node [%s], Output [%s]"), *GetOwningNode()->GetName().ToString(), *GetName().ToString());
 
-	Dataflow::FContextScopedCallstack Callstack(Context, this);
+	UE::Dataflow::FContextScopedCallstack Callstack(Context, this);
 	if (Callstack.IsLoopDetected())
 	{ 
 		ensureMsgf(false, TEXT("Connection %s is already in the callstack, this is certainly because of a loop in the graph"), *GetName().ToString());
@@ -278,13 +278,13 @@ bool FDataflowOutput::EvaluateImpl(Dataflow::FContext& Context) const
 	return true;
 }
 
-TFuture<bool> FDataflowOutput::EvaluateParallel(Dataflow::FContext& Context) const
+TFuture<bool> FDataflowOutput::EvaluateParallel(UE::Dataflow::FContext& Context) const
 {
 	return Async(EAsyncExecution::TaskGraph, [&]() -> bool { return this->Evaluate(Context); });
 }
 
 
-void FDataflowOutput::ForwardInput(const Dataflow::FConnectionReference& InputReference, Dataflow::FContext& Context) const
+void FDataflowOutput::ForwardInput(const UE::Dataflow::FConnectionReference& InputReference, UE::Dataflow::FContext& Context) const
 {
 	if (Property && OwningNode)
 	{
@@ -293,7 +293,7 @@ void FDataflowOutput::ForwardInput(const Dataflow::FConnectionReference& InputRe
 	}
 }
 
-void FDataflowOutput::ForwardInput(const FDataflowInput* Input, Dataflow::FContext& Context) const
+void FDataflowOutput::ForwardInput(const FDataflowInput* Input, UE::Dataflow::FContext& Context) const
 {
 	if (Property && OwningNode)
 	{
diff --git a/Engine/Source/Runtime/Experimental/Dataflow/Core/Private/Dataflow/DataflowMathNodes.cpp b/Engine/Source/Runtime/Experimental/Dataflow/Core/Private/Dataflow/DataflowMathNodes.cpp
--- a/Engine/Source/Runtime/Experimental/Dataflow/Core/Private/Dataflow/DataflowMathNodes.cpp
+++ b/Engine/Source/Runtime/Experimental/Dataflow/Core/Private/Dataflow/DataflowMathNodes.cpp
@@ -7,7 +7,7 @@
 
 #include "MathUtil.h"
 
-namespace Dataflow
+namespace UE::Dataflow
 {
 	void RegisterDataflowMathNodes()
 	{
@@ -57,7 +57,7 @@ namespace Dataflow
 
 //-----------------------------------------------------------------------------------------------
 
-FDataflowMathOneInputOperatorNode::FDataflowMathOneInputOperatorNode(const Dataflow::FNodeParameters& InParam, FGuid InGuid)
+FDataflowMathOneInputOperatorNode::FDataflowMathOneInputOperatorNode(const UE::Dataflow::FNodeParameters& InParam, FGuid InGuid)
 	:FDataflowNode(InParam, InGuid)
 {
 }
@@ -72,7 +72,7 @@ void FDataflowMathOneInputOperatorNode::RegisterInputsAndOutputs()
 	SetOutputConcreteType(&Result, TDataflowSingleTypePolicy<double>::TypeName);
 }
 
-void FDataflowMathOneInputOperatorNode::Evaluate(Dataflow::FContext& Context, const FDataflowOutput* Out) const
+void FDataflowMathOneInputOperatorNode::Evaluate(UE::Dataflow::FContext& Context, const FDataflowOutput* Out) const
 {
 	if (Out->IsA(&Result))
 	{
@@ -84,7 +84,7 @@ void FDataflowMathOneInputOperatorNode::Evaluate(Dataflow::FContext& Context, co
 
 //-----------------------------------------------------------------------------------------------
 
-FDataflowMathTwoInputsOperatorNode::FDataflowMathTwoInputsOperatorNode(const Dataflow::FNodeParameters& InParam, FGuid InGuid)
+FDataflowMathTwoInputsOperatorNode::FDataflowMathTwoInputsOperatorNode(const UE::Dataflow::FNodeParameters& InParam, FGuid InGuid)
 	:FDataflowNode(InParam, InGuid)
 {
 }
@@ -100,7 +100,7 @@ void FDataflowMathTwoInputsOperatorNode::RegisterInputsAndOutputs()
 	SetOutputConcreteType(&Result, TDataflowSingleTypePolicy<double>::TypeName);
 }
 
-void FDataflowMathTwoInputsOperatorNode::Evaluate(Dataflow::FContext& Context, const FDataflowOutput* Out) const
+void FDataflowMathTwoInputsOperatorNode::Evaluate(UE::Dataflow::FContext& Context, const FDataflowOutput* Out) const
 {
 	if (Out->IsA(&Result))
 	{
@@ -113,53 +113,53 @@ void FDataflowMathTwoInputsOperatorNode::Evaluate(Dataflow::FContext& Context, c
 
 //-----------------------------------------------------------------------------------------------
 
-FDataflowMathAddNode::FDataflowMathAddNode(const Dataflow::FNodeParameters& InParam, FGuid InGuid)
+FDataflowMathAddNode::FDataflowMathAddNode(const UE::Dataflow::FNodeParameters& InParam, FGuid InGuid)
 	: FDataflowMathTwoInputsOperatorNode(InParam, InGuid)
 {
 	RegisterInputsAndOutputs();
 }
 
-double FDataflowMathAddNode::ComputeResult(Dataflow::FContext& Context, double InA, double InB) const
+double FDataflowMathAddNode::ComputeResult(UE::Dataflow::FContext& Context, double InA, double InB) const
 {
 	return (InA + InB);
 }
 
 //-----------------------------------------------------------------------------------------------
 
-FDataflowMathSubtractNode::FDataflowMathSubtractNode(const Dataflow::FNodeParameters& InParam, FGuid InGuid)
+FDataflowMathSubtractNode::FDataflowMathSubtractNode(const UE::Dataflow::FNodeParameters& InParam, FGuid InGuid)
 	: FDataflowMathTwoInputsOperatorNode(InParam, InGuid)
 {
 	RegisterInputsAndOutputs();
 }
 
-double FDataflowMathSubtractNode::ComputeResult(Dataflow::FContext& Context, double InA, double InB) const
+double FDataflowMathSubtractNode::ComputeResult(UE::Dataflow::FContext& Context, double InA, double InB) const
 {
 	return (InA - InB);
 }
 
 //-----------------------------------------------------------------------------------------------
 
-FDataflowMathMultiplyNode::FDataflowMathMultiplyNode(const Dataflow::FNodeParameters& InParam, FGuid InGuid)
+FDataflowMathMultiplyNode::FDataflowMathMultiplyNode(const UE::Dataflow::FNodeParameters& InParam, FGuid InGuid)
 	: FDataflowMathTwoInputsOperatorNode(InParam, InGuid)
 {
 	RegisterInputsAndOutputs();
 }
 
-double FDataflowMathMultiplyNode::ComputeResult(Dataflow::FContext& Context, double InA, double InB) const
+double FDataflowMathMultiplyNode::ComputeResult(UE::Dataflow::FContext& Context, double InA, double InB) const
 {
 	return (InA * InB);
 }
 
 //-----------------------------------------------------------------------------------------------
 
-FDataflowMathDivideNode::FDataflowMathDivideNode(const Dataflow::FNodeParameters& InParam, FGuid InGuid)
+FDataflowMathDivideNode::FDataflowMathDivideNode(const UE::Dataflow::FNodeParameters& InParam, FGuid InGuid)
 	: FDataflowMathTwoInputsOperatorNode(InParam, InGuid)
 {
 	RegisterInputsAndOutputs();
 	RegisterInputConnection(&Fallback);
 }
 
-double FDataflowMathDivideNode::ComputeResult(Dataflow::FContext& Context, double InA, double InB) const
+double FDataflowMathDivideNode::ComputeResult(UE::Dataflow::FContext& Context, double InA, double InB) const
 {
 	if (InB == 0)
 	{
@@ -170,40 +170,40 @@ double FDataflowMathDivideNode::ComputeResult(Dataflow::FContext& Context, doubl
 
 //-----------------------------------------------------------------------------------------------
 
-FDataflowMathMinimumNode::FDataflowMathMinimumNode(const Dataflow::FNodeParameters& InParam, FGuid InGuid)
+FDataflowMathMinimumNode::FDataflowMathMinimumNode(const UE::Dataflow::FNodeParameters& InParam, FGuid InGuid)
 	: FDataflowMathTwoInputsOperatorNode(InParam, InGuid)
 {
 	RegisterInputsAndOutputs();
 }
 
-double FDataflowMathMinimumNode::ComputeResult(Dataflow::FContext& Context, double InA, double InB) const
+double FDataflowMathMinimumNode::ComputeResult(UE::Dataflow::FContext& Context, double InA, double InB) const
 {
 	return FMath::Min(InA, InB);
 }
 
 //-----------------------------------------------------------------------------------------------
 
-FDataflowMathMaximumNode::FDataflowMathMaximumNode(const Dataflow::FNodeParameters& InParam, FGuid InGuid)
+FDataflowMathMaximumNode::FDataflowMathMaximumNode(const UE::Dataflow::FNodeParameters& InParam, FGuid InGuid)
 	: FDataflowMathTwoInputsOperatorNode(InParam, InGuid)
 {
 	RegisterInputsAndOutputs();
 }
 
-double FDataflowMathMaximumNode::ComputeResult(Dataflow::FContext& Context, double InA, double InB) const
+double FDataflowMathMaximumNode::ComputeResult(UE::Dataflow::FContext& Context, double InA, double InB) const
 {
 	return FMath::Max(InA, InB);
 }
 
 //-----------------------------------------------------------------------------------------------
 
-FDataflowMathReciprocalNode::FDataflowMathReciprocalNode(const Dataflow::FNodeParameters& InParam, FGuid InGuid)
+FDataflowMathReciprocalNode::FDataflowMathReciprocalNode(const UE::Dataflow::FNodeParameters& InParam, FGuid InGuid)
 	: FDataflowMathOneInputOperatorNode(InParam, InGuid)
 {
 	RegisterInputsAndOutputs();
 	RegisterInputConnection(&Fallback);
 }
 
-double FDataflowMathReciprocalNode::ComputeResult(Dataflow::FContext& Context, double InA) const
+double FDataflowMathReciprocalNode::ComputeResult(UE::Dataflow::FContext& Context, double InA) const
 {
 	if (InA == 0)
 	{
@@ -214,39 +214,39 @@ double FDataflowMathReciprocalNode::ComputeResult(Dataflow::FContext& Context, d
 
 //-----------------------------------------------------------------------------------------------
 
-FDataflowMathSquareNode::FDataflowMathSquareNode(const Dataflow::FNodeParameters& InParam, FGuid InGuid)
+FDataflowMathSquareNode::FDataflowMathSquareNode(const UE::Dataflow::FNodeParameters& InParam, FGuid InGuid)
 	: FDataflowMathOneInputOperatorNode(InParam, InGuid)
 {
 	RegisterInputsAndOutputs();
 }
 
-double FDataflowMathSquareNode::ComputeResult(Dataflow::FContext& Context, double InA) const
+double FDataflowMathSquareNode::ComputeResult(UE::Dataflow::FContext& Context, double InA) const
 {
 	return (InA * InA);
 }
 
 //-----------------------------------------------------------------------------------------------
 
-FDataflowMathCubeNode::FDataflowMathCubeNode(const Dataflow::FNodeParameters& InParam, FGuid InGuid)
+FDataflowMathCubeNode::FDataflowMathCubeNode(const UE::Dataflow::FNodeParameters& InParam, FGuid InGuid)
 	: FDataflowMathOneInputOperatorNode(InParam, InGuid)
 {
 	RegisterInputsAndOutputs();
 }
 
-double FDataflowMathCubeNode::ComputeResult(Dataflow::FContext& Context, double InA) const
+double FDataflowMathCubeNode::ComputeResult(UE::Dataflow::FContext& Context, double InA) const
 {
 	return (InA * InA * InA);
 }
 
 //-----------------------------------------------------------------------------------------------
 
-FDataflowMathSquareRootNode::FDataflowMathSquareRootNode(const Dataflow::FNodeParameters& InParam, FGuid InGuid)
+FDataflowMathSquareRootNode::FDataflowMathSquareRootNode(const UE::Dataflow::FNodeParameters& InParam, FGuid InGuid)
 	: FDataflowMathOneInputOperatorNode(InParam, InGuid)
 {
 	RegisterInputsAndOutputs();
 }
 
-double FDataflowMathSquareRootNode::ComputeResult(Dataflow::FContext& Context, double InA) const
+double FDataflowMathSquareRootNode::ComputeResult(UE::Dataflow::FContext& Context, double InA) const
 {
 	if (InA < 0)
 	{
@@ -258,14 +258,14 @@ double FDataflowMathSquareRootNode::ComputeResult(Dataflow::FContext& Context, d
 
 //-----------------------------------------------------------------------------------------------
 
-FDataflowMathInverseSquareRootNode::FDataflowMathInverseSquareRootNode(const Dataflow::FNodeParameters& InParam, FGuid InGuid)
+FDataflowMathInverseSquareRootNode::FDataflowMathInverseSquareRootNode(const UE::Dataflow::FNodeParameters& InParam, FGuid InGuid)
 	: FDataflowMathOneInputOperatorNode(InParam, InGuid)
 {
 	RegisterInputsAndOutputs();
 	RegisterInputConnection(&Fallback);
 }
 
-double FDataflowMathInverseSquareRootNode::ComputeResult(Dataflow::FContext& Context, double InA) const
+double FDataflowMathInverseSquareRootNode::ComputeResult(UE::Dataflow::FContext& Context, double InA) const
 {
 	if (InA == 0)
 	{
@@ -276,111 +276,111 @@ double FDataflowMathInverseSquareRootNode::ComputeResult(Dataflow::FContext& Con
 
 //-----------------------------------------------------------------------------------------------
 
-FDataflowMathNegateNode::FDataflowMathNegateNode(const Dataflow::FNodeParameters& InParam, FGuid InGuid)
+FDataflowMathNegateNode::FDataflowMathNegateNode(const UE::Dataflow::FNodeParameters& InParam, FGuid InGuid)
 	: FDataflowMathOneInputOperatorNode(InParam, InGuid)
 {
 	RegisterInputsAndOutputs();
 }
 
-double FDataflowMathNegateNode::ComputeResult(Dataflow::FContext& Context, double InA) const
+double FDataflowMathNegateNode::ComputeResult(UE::Dataflow::FContext& Context, double InA) const
 {
 	return -InA;
 }
 
 //-----------------------------------------------------------------------------------------------
 
-FDataflowMathAbsNode::FDataflowMathAbsNode(const Dataflow::FNodeParameters& InParam, FGuid InGuid)
+FDataflowMathAbsNode::FDataflowMathAbsNode(const UE::Dataflow::FNodeParameters& InParam, FGuid InGuid)
 	: FDataflowMathOneInputOperatorNode(InParam, InGuid)
 {
 	RegisterInputsAndOutputs();
 }
 
-double FDataflowMathAbsNode::ComputeResult(Dataflow::FContext& Context, double InA) const
+double FDataflowMathAbsNode::ComputeResult(UE::Dataflow::FContext& Context, double InA) const
 {
 	return FMath::Abs(InA);
 }
 
 //-----------------------------------------------------------------------------------------------
 
-FDataflowMathFloorNode::FDataflowMathFloorNode(const Dataflow::FNodeParameters& InParam, FGuid InGuid)
+FDataflowMathFloorNode::FDataflowMathFloorNode(const UE::Dataflow::FNodeParameters& InParam, FGuid InGuid)
 	: FDataflowMathOneInputOperatorNode(InParam, InGuid)
 {
 	RegisterInputsAndOutputs();
 }
 
-double FDataflowMathFloorNode::ComputeResult(Dataflow::FContext& Context, double InA) const
+double FDataflowMathFloorNode::ComputeResult(UE::Dataflow::FContext& Context, double InA) const
 {
 	return FMath::FloorToDouble(InA);
 }
 
 //-----------------------------------------------------------------------------------------------
 
-FDataflowMathCeilNode::FDataflowMathCeilNode(const Dataflow::FNodeParameters& InParam, FGuid InGuid)
+FDataflowMathCeilNode::FDataflowMathCeilNode(const UE::Dataflow::FNodeParameters& InParam, FGuid InGuid)
 	: FDataflowMathOneInputOperatorNode(InParam, InGuid)
 {
 	RegisterInputsAndOutputs();
 }
 
-double FDataflowMathCeilNode::ComputeResult(Dataflow::FContext& Context, double InA) const
+double FDataflowMathCeilNode::ComputeResult(UE::Dataflow::FContext& Context, double InA) const
 {
 	return FMath::CeilToDouble(InA);
 }
 
 //-----------------------------------------------------------------------------------------------
 
-FDataflowMathRoundNode::FDataflowMathRoundNode(const Dataflow::FNodeParameters& InParam, FGuid InGuid)
+FDataflowMathRoundNode::FDataflowMathRoundNode(const UE::Dataflow::FNodeParameters& InParam, FGuid InGuid)
 	: FDataflowMathOneInputOperatorNode(InParam, InGuid)
 {
 	RegisterInputsAndOutputs();
 }
 
-double FDataflowMathRoundNode::ComputeResult(Dataflow::FContext& Context, double InA) const
+double FDataflowMathRoundNode::ComputeResult(UE::Dataflow::FContext& Context, double InA) const
 {
 	return FMath::RoundToDouble(InA);
 }
 
 //-----------------------------------------------------------------------------------------------
 
-FDataflowMathTruncNode::FDataflowMathTruncNode(const Dataflow::FNodeParameters& InParam, FGuid InGuid)
+FDataflowMathTruncNode::FDataflowMathTruncNode(const UE::Dataflow::FNodeParameters& InParam, FGuid InGuid)
 	: FDataflowMathOneInputOperatorNode(InParam, InGuid)
 {
 	RegisterInputsAndOutputs();
 }
 
-double FDataflowMathTruncNode::ComputeResult(Dataflow::FContext& Context, double InA) const
+double FDataflowMathTruncNode::ComputeResult(UE::Dataflow::FContext& Context, double InA) const
 {
 	return FMath::TruncToDouble(InA);
 }
 
 //-----------------------------------------------------------------------------------------------
 
-FDataflowMathFracNode::FDataflowMathFracNode(const Dataflow::FNodeParameters& InParam, FGuid InGuid)
+FDataflowMathFracNode::FDataflowMathFracNode(const UE::Dataflow::FNodeParameters& InParam, FGuid InGuid)
 	: FDataflowMathOneInputOperatorNode(InParam, InGuid)
 {
 	RegisterInputsAndOutputs();
 }
 
-double FDataflowMathFracNode::ComputeResult(Dataflow::FContext& Context, double InA) const
+double FDataflowMathFracNode::ComputeResult(UE::Dataflow::FContext& Context, double InA) const
 {
 	return FMath::Frac(InA);
 }
 
 //-----------------------------------------------------------------------------------------------
 
-FDataflowMathPowNode::FDataflowMathPowNode(const Dataflow::FNodeParameters& InParam, FGuid InGuid)
+FDataflowMathPowNode::FDataflowMathPowNode(const UE::Dataflow::FNodeParameters& InParam, FGuid InGuid)
 	: FDataflowMathTwoInputsOperatorNode(InParam, InGuid)
 {
 	RegisterInputsAndOutputs();
 }
 
-double FDataflowMathPowNode::ComputeResult(Dataflow::FContext& Context, double InA, double InB) const
+double FDataflowMathPowNode::ComputeResult(UE::Dataflow::FContext& Context, double InA, double InB) const
 {
 	return FMath::Pow(InA, InB);
 }
 
 //-----------------------------------------------------------------------------------------------
 
-FDataflowMathLogXNode::FDataflowMathLogXNode(const Dataflow::FNodeParameters& InParam, FGuid InGuid)
+FDataflowMathLogXNode::FDataflowMathLogXNode(const UE::Dataflow::FNodeParameters& InParam, FGuid InGuid)
 	: FDataflowMathOneInputOperatorNode(InParam, InGuid)
 {
 	Base.Value = 10.0; // default is base 10
@@ -388,7 +388,7 @@ FDataflowMathLogXNode::FDataflowMathLogXNode(const Dataflow::FNodeParameters& In
 	RegisterInputConnection(&Base);
 }
 
-double FDataflowMathLogXNode::ComputeResult(Dataflow::FContext& Context, double InA) const
+double FDataflowMathLogXNode::ComputeResult(UE::Dataflow::FContext& Context, double InA) const
 {
 	const double InBase = GetValue(Context, &Base);
 	if (InBase <= 0.f)
@@ -400,59 +400,59 @@ double FDataflowMathLogXNode::ComputeResult(Dataflow::FContext& Context, double
 
 //-----------------------------------------------------------------------------------------------
 
-FDataflowMathLogNode::FDataflowMathLogNode(const Dataflow::FNodeParameters& InParam, FGuid InGuid)
+FDataflowMathLogNode::FDataflowMathLogNode(const UE::Dataflow::FNodeParameters& InParam, FGuid InGuid)
 	: FDataflowMathOneInputOperatorNode(InParam, InGuid)
 {
 	RegisterInputsAndOutputs();
 }
 
-double FDataflowMathLogNode::ComputeResult(Dataflow::FContext& Context, double InA) const
+double FDataflowMathLogNode::ComputeResult(UE::Dataflow::FContext& Context, double InA) const
 {
 	return FMath::Loge(InA);
 }
 
 //-----------------------------------------------------------------------------------------------
 
-FDataflowMathExpNode::FDataflowMathExpNode(const Dataflow::FNodeParameters& InParam, FGuid InGuid)
+FDataflowMathExpNode::FDataflowMathExpNode(const UE::Dataflow::FNodeParameters& InParam, FGuid InGuid)
 	: FDataflowMathOneInputOperatorNode(InParam, InGuid)
 {
 	RegisterInputsAndOutputs();
 }
 
-double FDataflowMathExpNode::ComputeResult(Dataflow::FContext& Context, double InA) const
+double FDataflowMathExpNode::ComputeResult(UE::Dataflow::FContext& Context, double InA) const
 {
 	return FMath::Exp(InA);
 }
 
 //-----------------------------------------------------------------------------------------------
 
-FDataflowMathSignNode::FDataflowMathSignNode(const Dataflow::FNodeParameters& InParam, FGuid InGuid)
+FDataflowMathSignNode::FDataflowMathSignNode(const UE::Dataflow::FNodeParameters& InParam, FGuid InGuid)
 	: FDataflowMathOneInputOperatorNode(InParam, InGuid)
 {
 	RegisterInputsAndOutputs();
 }
 
-double FDataflowMathSignNode::ComputeResult(Dataflow::FContext& Context, double InA) const
+double FDataflowMathSignNode::ComputeResult(UE::Dataflow::FContext& Context, double InA) const
 {
 	return FMath::Sign(InA);
 }
 
 //-----------------------------------------------------------------------------------------------
 
-FDataflowMathOneMinusNode::FDataflowMathOneMinusNode(const Dataflow::FNodeParameters& InParam, FGuid InGuid)
+FDataflowMathOneMinusNode::FDataflowMathOneMinusNode(const UE::Dataflow::FNodeParameters& InParam, FGuid InGuid)
 	: FDataflowMathOneInputOperatorNode(InParam, InGuid)
 {
 	RegisterInputsAndOutputs();
 }
 
-double FDataflowMathOneMinusNode::ComputeResult(Dataflow::FContext& Context, double InA) const
+double FDataflowMathOneMinusNode::ComputeResult(UE::Dataflow::FContext& Context, double InA) const
 {
 	return (1.0 - InA);
 }
 
 //-----------------------------------------------------------------------------------------------
 
-FDataflowMathConstantNode::FDataflowMathConstantNode(const Dataflow::FNodeParameters& InParam, FGuid InGuid)
+FDataflowMathConstantNode::FDataflowMathConstantNode(const UE::Dataflow::FNodeParameters& InParam, FGuid InGuid)
 	: FDataflowNode(InParam, InGuid)
 {
 	RegisterOutputConnection(&Result);
@@ -486,7 +486,7 @@ double FDataflowMathConstantNode::GetConstant() const
 	return 0.0;
 }
 
-void FDataflowMathConstantNode::Evaluate(Dataflow::FContext& Context, const FDataflowOutput* Out) const
+void FDataflowMathConstantNode::Evaluate(UE::Dataflow::FContext& Context, const FDataflowOutput* Out) const
 {
 	if (Out->IsA(&Result))
 	{
@@ -500,117 +500,117 @@ void FDataflowMathConstantNode::Evaluate(Dataflow::FContext& Context, const FDat
 //
 //--------------------------------------------------------------------------
 
-FDataflowMathSinNode::FDataflowMathSinNode(const Dataflow::FNodeParameters& InParam, FGuid InGuid)
+FDataflowMathSinNode::FDataflowMathSinNode(const UE::Dataflow::FNodeParameters& InParam, FGuid InGuid)
 	: FDataflowMathOneInputOperatorNode(InParam, InGuid)
 {
 	RegisterInputsAndOutputs();
 }
 
-double FDataflowMathSinNode::ComputeResult(Dataflow::FContext& Context, double InA) const
+double FDataflowMathSinNode::ComputeResult(UE::Dataflow::FContext& Context, double InA) const
 {
 	return FMath::Sin(InA);
 }
 
 //-----------------------------------------------------------------------------------------------
 
-FDataflowMathCosNode::FDataflowMathCosNode(const Dataflow::FNodeParameters& InParam, FGuid InGuid)
+FDataflowMathCosNode::FDataflowMathCosNode(const UE::Dataflow::FNodeParameters& InParam, FGuid InGuid)
 	: FDataflowMathOneInputOperatorNode(InParam, InGuid)
 {
 	RegisterInputsAndOutputs();
 }
 
-double FDataflowMathCosNode::ComputeResult(Dataflow::FContext& Context, double InA) const
+double FDataflowMathCosNode::ComputeResult(UE::Dataflow::FContext& Context, double InA) const
 {
 	return FMath::Cos(InA);
 }
 
 //-----------------------------------------------------------------------------------------------
 
-FDataflowMathTanNode::FDataflowMathTanNode(const Dataflow::FNodeParameters& InParam, FGuid InGuid)
+FDataflowMathTanNode::FDataflowMathTanNode(const UE::Dataflow::FNodeParameters& InParam, FGuid InGuid)
 	: FDataflowMathOneInputOperatorNode(InParam, InGuid)
 {
 	RegisterInputsAndOutputs();
 }
 
-double FDataflowMathTanNode::ComputeResult(Dataflow::FContext& Context, double InA) const
+double FDataflowMathTanNode::ComputeResult(UE::Dataflow::FContext& Context, double InA) const
 {
 	return FMath::Tan(InA);
 }
 
 //-----------------------------------------------------------------------------------------------
 
-FDataflowMathArcSinNode::FDataflowMathArcSinNode(const Dataflow::FNodeParameters& InParam, FGuid InGuid)
+FDataflowMathArcSinNode::FDataflowMathArcSinNode(const UE::Dataflow::FNodeParameters& InParam, FGuid InGuid)
 	: FDataflowMathOneInputOperatorNode(InParam, InGuid)
 {
 	RegisterInputsAndOutputs();
 }
 
-double FDataflowMathArcSinNode::ComputeResult(Dataflow::FContext& Context, double InA) const
+double FDataflowMathArcSinNode::ComputeResult(UE::Dataflow::FContext& Context, double InA) const
 {
 	return FMath::Asin(InA);
 }
 
 //-----------------------------------------------------------------------------------------------
 
-FDataflowMathArcCosNode::FDataflowMathArcCosNode(const Dataflow::FNodeParameters& InParam, FGuid InGuid)
+FDataflowMathArcCosNode::FDataflowMathArcCosNode(const UE::Dataflow::FNodeParameters& InParam, FGuid InGuid)
 	: FDataflowMathOneInputOperatorNode(InParam, InGuid)
 {
 	RegisterInputsAndOutputs();
 }
 
-double FDataflowMathArcCosNode::ComputeResult(Dataflow::FContext& Context, double InA) const
+double FDataflowMathArcCosNode::ComputeResult(UE::Dataflow::FContext& Context, double InA) const
 {
 	return FMath::Acos(InA);
 }
 
 //-----------------------------------------------------------------------------------------------
 
-FDataflowMathArcTanNode::FDataflowMathArcTanNode(const Dataflow::FNodeParameters& InParam, FGuid InGuid)
+FDataflowMathArcTanNode::FDataflowMathArcTanNode(const UE::Dataflow::FNodeParameters& InParam, FGuid InGuid)
 	: FDataflowMathOneInputOperatorNode(InParam, InGuid)
 {
 	RegisterInputsAndOutputs();
 }
 
-double FDataflowMathArcTanNode::ComputeResult(Dataflow::FContext& Context, double InA) const
+double FDataflowMathArcTanNode::ComputeResult(UE::Dataflow::FContext& Context, double InA) const
 {
 	return FMath::Atan(InA);
 }
 
 //-----------------------------------------------------------------------------------------------
 
-FDataflowMathArcTan2Node::FDataflowMathArcTan2Node(const Dataflow::FNodeParameters& InParam, FGuid InGuid)
+FDataflowMathArcTan2Node::FDataflowMathArcTan2Node(const UE::Dataflow::FNodeParameters& InParam, FGuid InGuid)
 	: FDataflowMathTwoInputsOperatorNode(InParam, InGuid)
 {
 	RegisterInputsAndOutputs();
 }
 
-double FDataflowMathArcTan2Node::ComputeResult(Dataflow::FContext& Context, double InA, double InB) const
+double FDataflowMathArcTan2Node::ComputeResult(UE::Dataflow::FContext& Context, double InA, double InB) const
 {
 	return FMath::Atan2(InA, InB);
 }
 
 //-----------------------------------------------------------------------------------------------
 
-FDataflowMathDegToRadNode::FDataflowMathDegToRadNode(const Dataflow::FNodeParameters& InParam, FGuid InGuid)
+FDataflowMathDegToRadNode::FDataflowMathDegToRadNode(const UE::Dataflow::FNodeParameters& InParam, FGuid InGuid)
 	: FDataflowMathOneInputOperatorNode(InParam, InGuid)
 {
 	RegisterInputsAndOutputs();
 }
 
-double FDataflowMathDegToRadNode::ComputeResult(Dataflow::FContext& Context, double InA) const
+double FDataflowMathDegToRadNode::ComputeResult(UE::Dataflow::FContext& Context, double InA) const
 {
 	return FMath::DegreesToRadians(InA);
 }
 
 //-----------------------------------------------------------------------------------------------
 
-FDataflowMathRadToDegNode::FDataflowMathRadToDegNode(const Dataflow::FNodeParameters& InParam, FGuid InGuid)
+FDataflowMathRadToDegNode::FDataflowMathRadToDegNode(const UE::Dataflow::FNodeParameters& InParam, FGuid InGuid)
 	: FDataflowMathOneInputOperatorNode(InParam, InGuid)
 {
 	RegisterInputsAndOutputs();
 }
 
-double FDataflowMathRadToDegNode::ComputeResult(Dataflow::FContext& Context, double InA) const
+double FDataflowMathRadToDegNode::ComputeResult(UE::Dataflow::FContext& Context, double InA) const
 {
 	return FMath::RadiansToDegrees(InA);
 }
