Add namespace change guard around UE::Dataflow in GeometryCollection

diff --git a/Engine/Plugins/Experimental/GeometryCollectionPlugin/Source/GeometryCollectionNodes/Public/Dataflow/CreateColorArrayFromFloatArrayNode.h b/Engine/Plugins/Experimental/GeometryCollectionPlugin/Source/GeometryCollectionNodes/Public/Dataflow/CreateColorArrayFromFloatArrayNode.h
--- a/Engine/Plugins/Experimental/GeometryCollectionPlugin/Source/GeometryCollectionNodes/Public/Dataflow/CreateColorArrayFromFloatArrayNode.h
+++ b/Engine/Plugins/Experimental/GeometryCollectionPlugin/Source/GeometryCollectionNodes/Public/Dataflow/CreateColorArrayFromFloatArrayNode.h
@@ -8,6 +8,12 @@
 
 #include "CreateColorArrayFromFloatArrayNode.generated.h"
 
+#if UE_ENABLE_INCLUDE_ORDER_DEPRECATED_IN_5_5
+namespace Dataflow = UE::Dataflow;
+#else
+namespace UE_DEPRECATED(5.5, "Use UE::Dafaflow instead.") Dataflow {}
+#endif
+
 class FGeometryCollection;
 
 /** Set the vertex color on the collection based on the normalized float array. */
diff --git a/Engine/Plugins/Experimental/GeometryCollectionPlugin/Source/GeometryCollectionNodes/Public/Dataflow/GeometryCollectionArrayNodes.h b/Engine/Plugins/Experimental/GeometryCollectionPlugin/Source/GeometryCollectionNodes/Public/Dataflow/GeometryCollectionArrayNodes.h
--- a/Engine/Plugins/Experimental/GeometryCollectionPlugin/Source/GeometryCollectionNodes/Public/Dataflow/GeometryCollectionArrayNodes.h
+++ b/Engine/Plugins/Experimental/GeometryCollectionPlugin/Source/GeometryCollectionNodes/Public/Dataflow/GeometryCollectionArrayNodes.h
@@ -9,6 +9,11 @@
 
 #include "GeometryCollectionArrayNodes.generated.h"
 
+#if UE_ENABLE_INCLUDE_ORDER_DEPRECATED_IN_5_5
+namespace Dataflow = UE::Dataflow;
+#else
+namespace UE_DEPRECATED(5.5, "Use UE::Dafaflow instead.") Dataflow {}
+#endif
 
 class FGeometryCollection;
 
diff --git a/Engine/Plugins/Experimental/GeometryCollectionPlugin/Source/GeometryCollectionNodes/Public/Dataflow/GeometryCollectionAssetNodes.h b/Engine/Plugins/Experimental/GeometryCollectionPlugin/Source/GeometryCollectionNodes/Public/Dataflow/GeometryCollectionAssetNodes.h
--- a/Engine/Plugins/Experimental/GeometryCollectionPlugin/Source/GeometryCollectionNodes/Public/Dataflow/GeometryCollectionAssetNodes.h
+++ b/Engine/Plugins/Experimental/GeometryCollectionPlugin/Source/GeometryCollectionNodes/Public/Dataflow/GeometryCollectionAssetNodes.h
@@ -10,6 +10,12 @@
 
 #include "GeometryCollectionAssetNodes.generated.h"
 
+#if UE_ENABLE_INCLUDE_ORDER_DEPRECATED_IN_5_5
+namespace Dataflow = UE::Dataflow;
+#else
+namespace UE_DEPRECATED(5.5, "Use UE::Dafaflow instead.") Dataflow {}
+#endif
+
 class UGeometryCollection;
 class UBlueprint;
 
diff --git a/Engine/Plugins/Experimental/GeometryCollectionPlugin/Source/GeometryCollectionNodes/Public/Dataflow/GeometryCollectionClusteringNodes.h b/Engine/Plugins/Experimental/GeometryCollectionPlugin/Source/GeometryCollectionNodes/Public/Dataflow/GeometryCollectionClusteringNodes.h
--- a/Engine/Plugins/Experimental/GeometryCollectionPlugin/Source/GeometryCollectionNodes/Public/Dataflow/GeometryCollectionClusteringNodes.h
+++ b/Engine/Plugins/Experimental/GeometryCollectionPlugin/Source/GeometryCollectionNodes/Public/Dataflow/GeometryCollectionClusteringNodes.h
@@ -9,6 +9,12 @@
 
 #include "GeometryCollectionClusteringNodes.generated.h"
 
+#if UE_ENABLE_INCLUDE_ORDER_DEPRECATED_IN_5_5
+namespace Dataflow = UE::Dataflow;
+#else
+namespace UE_DEPRECATED(5.5, "Use UE::Dafaflow instead.") Dataflow {}
+#endif
+
 class FGeometryCollection;
 
 
diff --git a/Engine/Plugins/Experimental/GeometryCollectionPlugin/Source/GeometryCollectionNodes/Public/Dataflow/GeometryCollectionConversionNodes.h b/Engine/Plugins/Experimental/GeometryCollectionPlugin/Source/GeometryCollectionNodes/Public/Dataflow/GeometryCollectionConversionNodes.h
--- a/Engine/Plugins/Experimental/GeometryCollectionPlugin/Source/GeometryCollectionNodes/Public/Dataflow/GeometryCollectionConversionNodes.h
+++ b/Engine/Plugins/Experimental/GeometryCollectionPlugin/Source/GeometryCollectionNodes/Public/Dataflow/GeometryCollectionConversionNodes.h
@@ -7,6 +7,12 @@
 
 #include "GeometryCollectionConversionNodes.generated.h"
 
+#if UE_ENABLE_INCLUDE_ORDER_DEPRECATED_IN_5_5
+namespace Dataflow = UE::Dataflow;
+#else
+namespace UE_DEPRECATED(5.5, "Use UE::Dafaflow instead.") Dataflow {}
+#endif
+
 /**
  *
  * Description for this node
diff --git a/Engine/Plugins/Experimental/GeometryCollectionPlugin/Source/GeometryCollectionNodes/Public/Dataflow/GeometryCollectionDebugNodes.h b/Engine/Plugins/Experimental/GeometryCollectionPlugin/Source/GeometryCollectionNodes/Public/Dataflow/GeometryCollectionDebugNodes.h
--- a/Engine/Plugins/Experimental/GeometryCollectionPlugin/Source/GeometryCollectionNodes/Public/Dataflow/GeometryCollectionDebugNodes.h
+++ b/Engine/Plugins/Experimental/GeometryCollectionPlugin/Source/GeometryCollectionNodes/Public/Dataflow/GeometryCollectionDebugNodes.h
@@ -11,6 +11,12 @@
 
 #include "GeometryCollectionDebugNodes.generated.h"
 
+#if UE_ENABLE_INCLUDE_ORDER_DEPRECATED_IN_5_5
+namespace Dataflow = UE::Dataflow;
+#else
+namespace UE_DEPRECATED(5.5, "Use UE::Dafaflow instead.") Dataflow {}
+#endif
+
 //~ Dataflow nodes that are useful for debugging a geometry collection-related dataflow graph
 
 
diff --git a/Engine/Plugins/Experimental/GeometryCollectionPlugin/Source/GeometryCollectionNodes/Public/Dataflow/GeometryCollectionEditNodes.h b/Engine/Plugins/Experimental/GeometryCollectionPlugin/Source/GeometryCollectionNodes/Public/Dataflow/GeometryCollectionEditNodes.h
--- a/Engine/Plugins/Experimental/GeometryCollectionPlugin/Source/GeometryCollectionNodes/Public/Dataflow/GeometryCollectionEditNodes.h
+++ b/Engine/Plugins/Experimental/GeometryCollectionPlugin/Source/GeometryCollectionNodes/Public/Dataflow/GeometryCollectionEditNodes.h
@@ -9,6 +9,11 @@
 
 #include "GeometryCollectionEditNodes.generated.h"
 
+#if UE_ENABLE_INCLUDE_ORDER_DEPRECATED_IN_5_5
+namespace Dataflow = UE::Dataflow;
+#else
+namespace UE_DEPRECATED(5.5, "Use UE::Dafaflow instead.") Dataflow {}
+#endif
 
 class FGeometryCollection;
 
diff --git a/Engine/Plugins/Experimental/GeometryCollectionPlugin/Source/GeometryCollectionNodes/Public/Dataflow/GeometryCollectionEmbedNodes.h b/Engine/Plugins/Experimental/GeometryCollectionPlugin/Source/GeometryCollectionNodes/Public/Dataflow/GeometryCollectionEmbedNodes.h
--- a/Engine/Plugins/Experimental/GeometryCollectionPlugin/Source/GeometryCollectionNodes/Public/Dataflow/GeometryCollectionEmbedNodes.h
+++ b/Engine/Plugins/Experimental/GeometryCollectionPlugin/Source/GeometryCollectionNodes/Public/Dataflow/GeometryCollectionEmbedNodes.h
@@ -9,6 +9,11 @@
 
 //#include "GeometryCollectionEmbedNodes.generated.h"
 
+#if UE_ENABLE_INCLUDE_ORDER_DEPRECATED_IN_5_5
+namespace Dataflow = UE::Dataflow;
+#else
+namespace UE_DEPRECATED(5.5, "Use UE::Dafaflow instead.") Dataflow {}
+#endif
 
 class FGeometryCollection;
 
diff --git a/Engine/Plugins/Experimental/GeometryCollectionPlugin/Source/GeometryCollectionNodes/Public/Dataflow/GeometryCollectionFieldNodes.h b/Engine/Plugins/Experimental/GeometryCollectionPlugin/Source/GeometryCollectionNodes/Public/Dataflow/GeometryCollectionFieldNodes.h
--- a/Engine/Plugins/Experimental/GeometryCollectionPlugin/Source/GeometryCollectionNodes/Public/Dataflow/GeometryCollectionFieldNodes.h
+++ b/Engine/Plugins/Experimental/GeometryCollectionPlugin/Source/GeometryCollectionNodes/Public/Dataflow/GeometryCollectionFieldNodes.h
@@ -10,6 +10,11 @@
 
 #include "GeometryCollectionFieldNodes.generated.h"
 
+#if UE_ENABLE_INCLUDE_ORDER_DEPRECATED_IN_5_5
+namespace Dataflow = UE::Dataflow;
+#else
+namespace UE_DEPRECATED(5.5, "Use UE::Dafaflow instead.") Dataflow {}
+#endif
 
 /**
  *
diff --git a/Engine/Plugins/Experimental/GeometryCollectionPlugin/Source/GeometryCollectionNodes/Public/Dataflow/GeometryCollectionFracturingNodes.h b/Engine/Plugins/Experimental/GeometryCollectionPlugin/Source/GeometryCollectionNodes/Public/Dataflow/GeometryCollectionFracturingNodes.h
--- a/Engine/Plugins/Experimental/GeometryCollectionPlugin/Source/GeometryCollectionNodes/Public/Dataflow/GeometryCollectionFracturingNodes.h
+++ b/Engine/Plugins/Experimental/GeometryCollectionPlugin/Source/GeometryCollectionNodes/Public/Dataflow/GeometryCollectionFracturingNodes.h
@@ -10,6 +10,12 @@
 
 #include "GeometryCollectionFracturingNodes.generated.h"
 
+#if UE_ENABLE_INCLUDE_ORDER_DEPRECATED_IN_5_5
+namespace Dataflow = UE::Dataflow;
+#else
+namespace UE_DEPRECATED(5.5, "Use UE::Dafaflow instead.") Dataflow {}
+#endif
+
 class FGeometryCollection;
 class UStaticMesh;
 
diff --git a/Engine/Plugins/Experimental/GeometryCollectionPlugin/Source/GeometryCollectionNodes/Public/Dataflow/GeometryCollectionMakeNodes.h b/Engine/Plugins/Experimental/GeometryCollectionPlugin/Source/GeometryCollectionNodes/Public/Dataflow/GeometryCollectionMakeNodes.h
--- a/Engine/Plugins/Experimental/GeometryCollectionPlugin/Source/GeometryCollectionNodes/Public/Dataflow/GeometryCollectionMakeNodes.h
+++ b/Engine/Plugins/Experimental/GeometryCollectionPlugin/Source/GeometryCollectionNodes/Public/Dataflow/GeometryCollectionMakeNodes.h
@@ -13,6 +13,12 @@
 
 #include "GeometryCollectionMakeNodes.generated.h"
 
+#if UE_ENABLE_INCLUDE_ORDER_DEPRECATED_IN_5_5
+namespace Dataflow = UE::Dataflow;
+#else
+namespace UE_DEPRECATED(5.5, "Use UE::Dafaflow instead.") Dataflow {}
+#endif
+
 /**
  *
  * Description for this node
diff --git a/Engine/Plugins/Experimental/GeometryCollectionPlugin/Source/GeometryCollectionNodes/Public/Dataflow/GeometryCollectionMaterialNodes.h b/Engine/Plugins/Experimental/GeometryCollectionPlugin/Source/GeometryCollectionNodes/Public/Dataflow/GeometryCollectionMaterialNodes.h
--- a/Engine/Plugins/Experimental/GeometryCollectionPlugin/Source/GeometryCollectionNodes/Public/Dataflow/GeometryCollectionMaterialNodes.h
+++ b/Engine/Plugins/Experimental/GeometryCollectionPlugin/Source/GeometryCollectionNodes/Public/Dataflow/GeometryCollectionMaterialNodes.h
@@ -13,6 +13,11 @@
 
 #include "GeometryCollectionMaterialNodes.generated.h"
 
+#if UE_ENABLE_INCLUDE_ORDER_DEPRECATED_IN_5_5
+namespace Dataflow = UE::Dataflow;
+#else
+namespace UE_DEPRECATED(5.5, "Use UE::Dafaflow instead.") Dataflow {}
+#endif
 
 class FGeometryCollection;
 class UGeometryCollection;
diff --git a/Engine/Plugins/Experimental/GeometryCollectionPlugin/Source/GeometryCollectionNodes/Public/Dataflow/GeometryCollectionMathNodes.h b/Engine/Plugins/Experimental/GeometryCollectionPlugin/Source/GeometryCollectionNodes/Public/Dataflow/GeometryCollectionMathNodes.h
--- a/Engine/Plugins/Experimental/GeometryCollectionPlugin/Source/GeometryCollectionNodes/Public/Dataflow/GeometryCollectionMathNodes.h
+++ b/Engine/Plugins/Experimental/GeometryCollectionPlugin/Source/GeometryCollectionNodes/Public/Dataflow/GeometryCollectionMathNodes.h
@@ -7,6 +7,12 @@
 
 #include "GeometryCollectionMathNodes.generated.h"
 
+#if UE_ENABLE_INCLUDE_ORDER_DEPRECATED_IN_5_5
+namespace Dataflow = UE::Dataflow;
+#else
+namespace UE_DEPRECATED(5.5, "Use UE::Dafaflow instead.") Dataflow {}
+#endif
+
 /**
 *
 * 
diff --git a/Engine/Plugins/Experimental/GeometryCollectionPlugin/Source/GeometryCollectionNodes/Public/Dataflow/GeometryCollectionMeshNodes.h b/Engine/Plugins/Experimental/GeometryCollectionPlugin/Source/GeometryCollectionNodes/Public/Dataflow/GeometryCollectionMeshNodes.h
--- a/Engine/Plugins/Experimental/GeometryCollectionPlugin/Source/GeometryCollectionNodes/Public/Dataflow/GeometryCollectionMeshNodes.h
+++ b/Engine/Plugins/Experimental/GeometryCollectionPlugin/Source/GeometryCollectionNodes/Public/Dataflow/GeometryCollectionMeshNodes.h
@@ -10,6 +10,12 @@
 
 #include "GeometryCollectionMeshNodes.generated.h"
 
+#if UE_ENABLE_INCLUDE_ORDER_DEPRECATED_IN_5_5
+namespace Dataflow = UE::Dataflow;
+#else
+namespace UE_DEPRECATED(5.5, "Use UE::Dafaflow instead.") Dataflow {}
+#endif
+
 class FGeometryCollection;
 class UStaticMesh;
 class UDynamicMesh;
diff --git a/Engine/Plugins/Experimental/GeometryCollectionPlugin/Source/GeometryCollectionNodes/Public/Dataflow/GeometryCollectionNodes.h b/Engine/Plugins/Experimental/GeometryCollectionPlugin/Source/GeometryCollectionNodes/Public/Dataflow/GeometryCollectionNodes.h
--- a/Engine/Plugins/Experimental/GeometryCollectionPlugin/Source/GeometryCollectionNodes/Public/Dataflow/GeometryCollectionNodes.h
+++ b/Engine/Plugins/Experimental/GeometryCollectionPlugin/Source/GeometryCollectionNodes/Public/Dataflow/GeometryCollectionNodes.h
@@ -14,6 +14,11 @@
 
 #include "GeometryCollectionNodes.generated.h"
 
+#if UE_ENABLE_INCLUDE_ORDER_DEPRECATED_IN_5_5
+namespace Dataflow = UE::Dataflow;
+#else
+namespace UE_DEPRECATED(5.5, "Use UE::Dafaflow instead.") Dataflow {}
+#endif
 
 class FGeometryCollection;
 class UGeometryCollection;
diff --git a/Engine/Plugins/Experimental/GeometryCollectionPlugin/Source/GeometryCollectionNodes/Public/Dataflow/GeometryCollectionOverrideNodes.h b/Engine/Plugins/Experimental/GeometryCollectionPlugin/Source/GeometryCollectionNodes/Public/Dataflow/GeometryCollectionOverrideNodes.h
--- a/Engine/Plugins/Experimental/GeometryCollectionPlugin/Source/GeometryCollectionNodes/Public/Dataflow/GeometryCollectionOverrideNodes.h
+++ b/Engine/Plugins/Experimental/GeometryCollectionPlugin/Source/GeometryCollectionNodes/Public/Dataflow/GeometryCollectionOverrideNodes.h
@@ -8,6 +8,12 @@
 
 #include "GeometryCollectionOverrideNodes.generated.h"
 
+#if UE_ENABLE_INCLUDE_ORDER_DEPRECATED_IN_5_5
+namespace Dataflow = UE::Dataflow;
+#else
+namespace UE_DEPRECATED(5.5, "Use UE::Dafaflow instead.") Dataflow {}
+#endif
+
 /**
  *
  * 
diff --git a/Engine/Plugins/Experimental/GeometryCollectionPlugin/Source/GeometryCollectionNodes/Public/Dataflow/GeometryCollectionProcessingNodes.h b/Engine/Plugins/Experimental/GeometryCollectionPlugin/Source/GeometryCollectionNodes/Public/Dataflow/GeometryCollectionProcessingNodes.h
--- a/Engine/Plugins/Experimental/GeometryCollectionPlugin/Source/GeometryCollectionNodes/Public/Dataflow/GeometryCollectionProcessingNodes.h
+++ b/Engine/Plugins/Experimental/GeometryCollectionPlugin/Source/GeometryCollectionNodes/Public/Dataflow/GeometryCollectionProcessingNodes.h
@@ -8,6 +8,11 @@
 
 #include "GeometryCollectionProcessingNodes.generated.h"
 
+#if UE_ENABLE_INCLUDE_ORDER_DEPRECATED_IN_5_5
+namespace Dataflow = UE::Dataflow;
+#else
+namespace UE_DEPRECATED(5.5, "Use UE::Dafaflow instead.") Dataflow {}
+#endif
 
 USTRUCT(meta = (DataflowGeometryCollection))
 struct FCloseGeometryOnCollectionDataflowNode : public FDataflowNode
diff --git a/Engine/Plugins/Experimental/GeometryCollectionPlugin/Source/GeometryCollectionNodes/Public/Dataflow/GeometryCollectionSamplingNodes.h b/Engine/Plugins/Experimental/GeometryCollectionPlugin/Source/GeometryCollectionNodes/Public/Dataflow/GeometryCollectionSamplingNodes.h
--- a/Engine/Plugins/Experimental/GeometryCollectionPlugin/Source/GeometryCollectionNodes/Public/Dataflow/GeometryCollectionSamplingNodes.h
+++ b/Engine/Plugins/Experimental/GeometryCollectionPlugin/Source/GeometryCollectionNodes/Public/Dataflow/GeometryCollectionSamplingNodes.h
@@ -7,6 +7,12 @@
 
 #include "GeometryCollectionSamplingNodes.generated.h"
 
+#if UE_ENABLE_INCLUDE_ORDER_DEPRECATED_IN_5_5
+namespace Dataflow = UE::Dataflow;
+#else
+namespace UE_DEPRECATED(5.5, "Use UE::Dafaflow instead.") Dataflow {}
+#endif
+
 class FGeometryCollection;
 class UDynamicMesh;
 
diff --git a/Engine/Plugins/Experimental/GeometryCollectionPlugin/Source/GeometryCollectionNodes/Public/Dataflow/GeometryCollectionSelectionNodes.h b/Engine/Plugins/Experimental/GeometryCollectionPlugin/Source/GeometryCollectionNodes/Public/Dataflow/GeometryCollectionSelectionNodes.h
--- a/Engine/Plugins/Experimental/GeometryCollectionPlugin/Source/GeometryCollectionNodes/Public/Dataflow/GeometryCollectionSelectionNodes.h
+++ b/Engine/Plugins/Experimental/GeometryCollectionPlugin/Source/GeometryCollectionNodes/Public/Dataflow/GeometryCollectionSelectionNodes.h
@@ -12,6 +12,12 @@
 
 #include "GeometryCollectionSelectionNodes.generated.h"
 
+#if UE_ENABLE_INCLUDE_ORDER_DEPRECATED_IN_5_5
+namespace Dataflow = UE::Dataflow;
+#else
+namespace UE_DEPRECATED(5.5, "Use UE::Dafaflow instead.") Dataflow {}
+#endif
+
 class FGeometryCollection;
 
 
diff --git a/Engine/Plugins/Experimental/GeometryCollectionPlugin/Source/GeometryCollectionNodes/Public/Dataflow/GeometryCollectionSkeletalMeshToCollectionNode.h b/Engine/Plugins/Experimental/GeometryCollectionPlugin/Source/GeometryCollectionNodes/Public/Dataflow/GeometryCollectionSkeletalMeshToCollectionNode.h
--- a/Engine/Plugins/Experimental/GeometryCollectionPlugin/Source/GeometryCollectionNodes/Public/Dataflow/GeometryCollectionSkeletalMeshToCollectionNode.h
+++ b/Engine/Plugins/Experimental/GeometryCollectionPlugin/Source/GeometryCollectionNodes/Public/Dataflow/GeometryCollectionSkeletalMeshToCollectionNode.h
@@ -8,6 +8,12 @@
 
 #include "GeometryCollectionSkeletalMeshToCollectionNode.generated.h"
 
+#if UE_ENABLE_INCLUDE_ORDER_DEPRECATED_IN_5_5
+namespace Dataflow = UE::Dataflow;
+#else
+namespace UE_DEPRECATED(5.5, "Use UE::Dafaflow instead.") Dataflow {}
+#endif
+
 class USkeletalMesh;
 
 
diff --git a/Engine/Plugins/Experimental/GeometryCollectionPlugin/Source/GeometryCollectionNodes/Public/Dataflow/GeometryCollectionSkeletonToCollectionNode.h b/Engine/Plugins/Experimental/GeometryCollectionPlugin/Source/GeometryCollectionNodes/Public/Dataflow/GeometryCollectionSkeletonToCollectionNode.h
--- a/Engine/Plugins/Experimental/GeometryCollectionPlugin/Source/GeometryCollectionNodes/Public/Dataflow/GeometryCollectionSkeletonToCollectionNode.h
+++ b/Engine/Plugins/Experimental/GeometryCollectionPlugin/Source/GeometryCollectionNodes/Public/Dataflow/GeometryCollectionSkeletonToCollectionNode.h
@@ -8,6 +8,12 @@
 
 #include "GeometryCollectionSkeletonToCollectionNode.generated.h"
 
+#if UE_ENABLE_INCLUDE_ORDER_DEPRECATED_IN_5_5
+namespace Dataflow = UE::Dataflow;
+#else
+namespace UE_DEPRECATED(5.5, "Use UE::Dafaflow instead.") Dataflow {}
+#endif
+
 class USkeletalMesh;
 
 USTRUCT(meta = (DataflowGeometryCollection))
diff --git a/Engine/Plugins/Experimental/GeometryCollectionPlugin/Source/GeometryCollectionNodes/Public/Dataflow/GeometryCollectionStaticMeshToCollectionNode.h b/Engine/Plugins/Experimental/GeometryCollectionPlugin/Source/GeometryCollectionNodes/Public/Dataflow/GeometryCollectionStaticMeshToCollectionNode.h
--- a/Engine/Plugins/Experimental/GeometryCollectionPlugin/Source/GeometryCollectionNodes/Public/Dataflow/GeometryCollectionStaticMeshToCollectionNode.h
+++ b/Engine/Plugins/Experimental/GeometryCollectionPlugin/Source/GeometryCollectionNodes/Public/Dataflow/GeometryCollectionStaticMeshToCollectionNode.h
@@ -8,6 +8,12 @@
 
 #include "GeometryCollectionStaticMeshToCollectionNode.generated.h"
 
+#if UE_ENABLE_INCLUDE_ORDER_DEPRECATED_IN_5_5
+namespace Dataflow = UE::Dataflow;
+#else
+namespace UE_DEPRECATED(5.5, "Use UE::Dafaflow instead.") Dataflow {}
+#endif
+
 class UGeometryCollection;
 class UStaticMesh;
 class UMaterial;
diff --git a/Engine/Plugins/Experimental/GeometryCollectionPlugin/Source/GeometryCollectionNodes/Public/Dataflow/GeometryCollectionTransferVertexAttributeNode.h b/Engine/Plugins/Experimental/GeometryCollectionPlugin/Source/GeometryCollectionNodes/Public/Dataflow/GeometryCollectionTransferVertexAttributeNode.h
--- a/Engine/Plugins/Experimental/GeometryCollectionPlugin/Source/GeometryCollectionNodes/Public/Dataflow/GeometryCollectionTransferVertexAttributeNode.h
+++ b/Engine/Plugins/Experimental/GeometryCollectionPlugin/Source/GeometryCollectionNodes/Public/Dataflow/GeometryCollectionTransferVertexAttributeNode.h
@@ -10,6 +10,11 @@
 
 #include "GeometryCollectionTransferVertexAttributeNode.generated.h"
 
+#if UE_ENABLE_INCLUDE_ORDER_DEPRECATED_IN_5_5
+namespace Dataflow = UE::Dataflow;
+#else
+namespace UE_DEPRECATED(5.5, "Use UE::Dafaflow instead.") Dataflow {}
+#endif
 
 namespace UE::Private 
 {
diff --git a/Engine/Plugins/Experimental/GeometryCollectionPlugin/Source/GeometryCollectionNodes/Public/Dataflow/GeometryCollectionTriangleBoundaryIndicesNode.h b/Engine/Plugins/Experimental/GeometryCollectionPlugin/Source/GeometryCollectionNodes/Public/Dataflow/GeometryCollectionTriangleBoundaryIndicesNode.h
--- a/Engine/Plugins/Experimental/GeometryCollectionPlugin/Source/GeometryCollectionNodes/Public/Dataflow/GeometryCollectionTriangleBoundaryIndicesNode.h
+++ b/Engine/Plugins/Experimental/GeometryCollectionPlugin/Source/GeometryCollectionNodes/Public/Dataflow/GeometryCollectionTriangleBoundaryIndicesNode.h
@@ -8,6 +8,12 @@
 
 #include "GeometryCollectionTriangleBoundaryIndicesNode.generated.h"
 
+#if UE_ENABLE_INCLUDE_ORDER_DEPRECATED_IN_5_5
+namespace Dataflow = UE::Dataflow;
+#else
+namespace UE_DEPRECATED(5.5, "Use UE::Dafaflow instead.") Dataflow {}
+#endif
+
 //Outputs boundary nodes of a triangle mesh
 USTRUCT(meta = (DataflowFlesh))
 struct FTriangleBoundaryIndicesNode : public FDataflowNode
diff --git a/Engine/Plugins/Experimental/GeometryCollectionPlugin/Source/GeometryCollectionNodes/Public/Dataflow/GeometryCollectionUtilityNodes.h b/Engine/Plugins/Experimental/GeometryCollectionPlugin/Source/GeometryCollectionNodes/Public/Dataflow/GeometryCollectionUtilityNodes.h
--- a/Engine/Plugins/Experimental/GeometryCollectionPlugin/Source/GeometryCollectionNodes/Public/Dataflow/GeometryCollectionUtilityNodes.h
+++ b/Engine/Plugins/Experimental/GeometryCollectionPlugin/Source/GeometryCollectionNodes/Public/Dataflow/GeometryCollectionUtilityNodes.h
@@ -14,6 +14,11 @@
 
 #include "GeometryCollectionUtilityNodes.generated.h"
 
+#if UE_ENABLE_INCLUDE_ORDER_DEPRECATED_IN_5_5
+namespace Dataflow = UE::Dataflow;
+#else
+namespace UE_DEPRECATED(5.5, "Use UE::Dafaflow instead.") Dataflow {}
+#endif
 
 UENUM(BlueprintType)
 enum class EConvexOverlapRemovalMethodEnum : uint8
diff --git a/Engine/Plugins/Experimental/GeometryCollectionPlugin/Source/GeometryCollectionNodes/Public/Dataflow/GeometryCollectionVertexScalarToVertexIndicesNode.h b/Engine/Plugins/Experimental/GeometryCollectionPlugin/Source/GeometryCollectionNodes/Public/Dataflow/GeometryCollectionVertexScalarToVertexIndicesNode.h
--- a/Engine/Plugins/Experimental/GeometryCollectionPlugin/Source/GeometryCollectionNodes/Public/Dataflow/GeometryCollectionVertexScalarToVertexIndicesNode.h
+++ b/Engine/Plugins/Experimental/GeometryCollectionPlugin/Source/GeometryCollectionNodes/Public/Dataflow/GeometryCollectionVertexScalarToVertexIndicesNode.h
@@ -9,6 +9,12 @@
 
 #include "GeometryCollectionVertexScalarToVertexIndicesNode.generated.h"
 
+#if UE_ENABLE_INCLUDE_ORDER_DEPRECATED_IN_5_5
+namespace Dataflow = UE::Dataflow;
+#else
+namespace UE_DEPRECATED(5.5, "Use UE::Dafaflow instead.") Dataflow {}
+#endif
+
 /** Convert an vertex float array to a list of indices */
 USTRUCT(meta = (DataflowGeometryCollection))
 struct FGeometryCollectionVertexScalarToVertexIndicesNode : public FDataflowNode
diff --git a/Engine/Plugins/Experimental/GeometryCollectionPlugin/Source/GeometryCollectionNodes/Public/Dataflow/GeometryCollectionVerticesNodes.h b/Engine/Plugins/Experimental/GeometryCollectionPlugin/Source/GeometryCollectionNodes/Public/Dataflow/GeometryCollectionVerticesNodes.h
--- a/Engine/Plugins/Experimental/GeometryCollectionPlugin/Source/GeometryCollectionNodes/Public/Dataflow/GeometryCollectionVerticesNodes.h
+++ b/Engine/Plugins/Experimental/GeometryCollectionPlugin/Source/GeometryCollectionNodes/Public/Dataflow/GeometryCollectionVerticesNodes.h
@@ -8,6 +8,12 @@
 
 #include "GeometryCollectionVerticesNodes.generated.h"
 
+#if UE_ENABLE_INCLUDE_ORDER_DEPRECATED_IN_5_5
+namespace Dataflow = UE::Dataflow;
+#else
+namespace UE_DEPRECATED(5.5, "Use UE::Dafaflow instead.") Dataflow {}
+#endif
+
 class USkeletalMesh;
 
 USTRUCT(meta = (DataflowGeometryCollection))
diff --git a/Engine/Plugins/Experimental/GeometryCollectionPlugin/Source/GeometryCollectionNodes/Public/Dataflow/SetVertexColorFromFloatArrayNode.h b/Engine/Plugins/Experimental/GeometryCollectionPlugin/Source/GeometryCollectionNodes/Public/Dataflow/SetVertexColorFromFloatArrayNode.h
--- a/Engine/Plugins/Experimental/GeometryCollectionPlugin/Source/GeometryCollectionNodes/Public/Dataflow/SetVertexColorFromFloatArrayNode.h
+++ b/Engine/Plugins/Experimental/GeometryCollectionPlugin/Source/GeometryCollectionNodes/Public/Dataflow/SetVertexColorFromFloatArrayNode.h
@@ -8,6 +8,12 @@
 
 #include "SetVertexColorFromFloatArrayNode.generated.h"
 
+#if UE_ENABLE_INCLUDE_ORDER_DEPRECATED_IN_5_5
+namespace Dataflow = UE::Dataflow;
+#else
+namespace UE_DEPRECATED(5.5, "Use UE::Dafaflow instead.") Dataflow {}
+#endif
+
 class FGeometryCollection;
 
 /** Set the vertex color on the collection based on the normalized float array. */
diff --git a/Engine/Plugins/Experimental/GeometryCollectionPlugin/Source/GeometryCollectionNodes/Public/Dataflow/SetVertexColorFromVertexIndicesNode.h b/Engine/Plugins/Experimental/GeometryCollectionPlugin/Source/GeometryCollectionNodes/Public/Dataflow/SetVertexColorFromVertexIndicesNode.h
--- a/Engine/Plugins/Experimental/GeometryCollectionPlugin/Source/GeometryCollectionNodes/Public/Dataflow/SetVertexColorFromVertexIndicesNode.h
+++ b/Engine/Plugins/Experimental/GeometryCollectionPlugin/Source/GeometryCollectionNodes/Public/Dataflow/SetVertexColorFromVertexIndicesNode.h
@@ -8,6 +8,12 @@
 
 #include "SetVertexColorFromVertexIndicesNode.generated.h"
 
+#if UE_ENABLE_INCLUDE_ORDER_DEPRECATED_IN_5_5
+namespace Dataflow = UE::Dataflow;
+#else
+namespace UE_DEPRECATED(5.5, "Use UE::Dafaflow instead.") Dataflow {}
+#endif
+
 /** Set the vertex color of the collection based on the selection set. */
 USTRUCT()
 struct FSetVertexColorFromVertexIndicesDataflowNode : public FDataflowNode
diff --git a/Engine/Plugins/Experimental/GeometryCollectionPlugin/Source/GeometryCollectionNodes/Public/Dataflow/SetVertexColorFromVertexSelectionNode.h b/Engine/Plugins/Experimental/GeometryCollectionPlugin/Source/GeometryCollectionNodes/Public/Dataflow/SetVertexColorFromVertexSelectionNode.h
--- a/Engine/Plugins/Experimental/GeometryCollectionPlugin/Source/GeometryCollectionNodes/Public/Dataflow/SetVertexColorFromVertexSelectionNode.h
+++ b/Engine/Plugins/Experimental/GeometryCollectionPlugin/Source/GeometryCollectionNodes/Public/Dataflow/SetVertexColorFromVertexSelectionNode.h
@@ -9,6 +9,12 @@
 
 #include "SetVertexColorFromVertexSelectionNode.generated.h"
 
+#if UE_ENABLE_INCLUDE_ORDER_DEPRECATED_IN_5_5
+namespace Dataflow = UE::Dataflow;
+#else
+namespace UE_DEPRECATED(5.5, "Use UE::Dafaflow instead.") Dataflow {}
+#endif
+
 /** Set the vertex color of the collection based on the selection set. */
 USTRUCT()
 struct FSetVertexColorFromVertexSelectionDataflowNode : public FDataflowNode
