ChaosFlesh: refactor kinematic vertex attribute, transfer skin weights between collections, skeletal mesh to collection now contains skin weights, add particle radius cvar

diff --git a/Engine/Plugins/Experimental/ChaosFlesh/Source/ChaosFleshNodes/Private/Dataflow/ChaosFleshVisualizeFiberFieldNode.cpp b/Engine/Plugins/Experimental/ChaosFlesh/Source/ChaosFleshNodes/Private/Dataflow/ChaosFleshVisualizeFiberFieldNode.cpp
--- a/Engine/Plugins/Experimental/ChaosFlesh/Source/ChaosFleshNodes/Private/Dataflow/ChaosFleshVisualizeFiberFieldNode.cpp
+++ b/Engine/Plugins/Experimental/ChaosFlesh/Source/ChaosFleshNodes/Private/Dataflow/ChaosFleshVisualizeFiberFieldNode.cpp
@@ -4,6 +4,7 @@
 #include "GeometryCollection/Facades/CollectionKinematicBindingFacade.h"
 #include "GeometryCollection/Facades/CollectionPositionTargetFacade.h"
 #include "ChaosFlesh/TetrahedralCollection.h"
+#include "GeometryCollection/Facades/CollectionVertexBoneWeightsFacade.h"
 
 #include UE_INLINE_GENERATED_CPP_BY_NAME(ChaosFleshVisualizeFiberFieldNode)
 
@@ -92,37 +93,16 @@ void FVisualizeKinematicFacesNode::Evaluate(UE::Dataflow::FContext& Context, con
 		TManagedArray<FIntVector>* Indices = InCollection.FindAttribute<FIntVector>("Indices", FGeometryCollection::FacesGroup);
 		TManagedArray<bool>* FaceVisibility = InCollection.FindAttribute<bool>("Visible", FGeometryCollection::FacesGroup);
 		TManagedArray<FVector3f>* Vertices = InCollection.FindAttribute<FVector3f>("Vertex", FGeometryCollection::VerticesGroup);
-		if (Indices && FaceVisibility && Vertices)
+		GeometryCollection::Facades::FVertexBoneWeightsFacade VertexBoneWeightsFacade(InCollection);
+		if (Indices && FaceVisibility && Vertices && VertexBoneWeightsFacade.IsValid())
 		{
-			//Find kinematic particles
-			TArray<bool> ParticleIsKinematic;
-			ParticleIsKinematic.Init(false, Vertices->Num());
-			typedef GeometryCollection::Facades::FKinematicBindingFacade FKinematics;
-			FKinematics Kinematics(InCollection);
-			//RemoveElements
-			// Add Kinematics Node
-			for (int i = Kinematics.NumKinematicBindings() - 1; i >= 0; i--)
-			{
-				FKinematics::FBindingKey Key = Kinematics.GetKinematicBindingKey(i);
-
-				int32 BoneIndex = INDEX_NONE;
-				TArray<int32> BoundVerts;
-				TArray<float> BoundWeights;
-				Kinematics.GetBoneBindings(Key, BoneIndex, BoundVerts, BoundWeights);
-
-				for (int32 vdx : BoundVerts)
-				{
-					ParticleIsKinematic[vdx] = true;
-				}
-			}
-
 			FaceVisibility->Fill(false);
 			for (int32 FaceIdx = 0; FaceIdx < Indices->Num(); ++FaceIdx)
 			{
 				bool IsElementKinematic = true;
 				for (int32 j = 0; j < 3; j++)
 				{
-					if (!ParticleIsKinematic[(*Indices)[FaceIdx][j]])
+					if (!VertexBoneWeightsFacade.IsKinematicVertex((*Indices)[FaceIdx][j]))
 					{
 						IsElementKinematic = false;
 						break;
diff --git a/Engine/Plugins/Experimental/GeometryCollectionPlugin/Source/GeometryCollectionNodes/Private/Dataflow/GeometryCollectionNodesPlugin.cpp b/Engine/Plugins/Experimental/GeometryCollectionPlugin/Source/GeometryCollectionNodes/Private/Dataflow/GeometryCollectionNodesPlugin.cpp
--- a/Engine/Plugins/Experimental/GeometryCollectionPlugin/Source/GeometryCollectionNodes/Private/Dataflow/GeometryCollectionNodesPlugin.cpp
+++ b/Engine/Plugins/Experimental/GeometryCollectionPlugin/Source/GeometryCollectionNodes/Private/Dataflow/GeometryCollectionNodesPlugin.cpp
@@ -67,6 +67,7 @@ void IGeometryCollectionNodesPlugin::StartupModule()
 	DATAFLOW_NODE_REGISTER_CREATION_FACTORY(FSkeletonToCollectionDataflowNode);
 	DATAFLOW_NODE_REGISTER_CREATION_FACTORY(FSkeletalMeshToCollectionDataflowNode);
 	DATAFLOW_NODE_REGISTER_CREATION_FACTORY(FGeometryCollectionTransferVertexAttributeNode);
+	DATAFLOW_NODE_REGISTER_CREATION_FACTORY(FGeometryCollectionSetKinematicVertexSelectionNode);
 	DATAFLOW_NODE_REGISTER_CREATION_FACTORY(FStaticMeshToCollectionDataflowNode);
 	DATAFLOW_NODE_REGISTER_CREATION_FACTORY(FTriangleBoundaryIndicesNode);
 }
diff --git a/Engine/Plugins/Experimental/GeometryCollectionPlugin/Source/GeometryCollectionNodes/Private/Dataflow/GeometryCollectionTransferVertexAttributeNode.cpp b/Engine/Plugins/Experimental/GeometryCollectionPlugin/Source/GeometryCollectionNodes/Private/Dataflow/GeometryCollectionTransferVertexAttributeNode.cpp
--- a/Engine/Plugins/Experimental/GeometryCollectionPlugin/Source/GeometryCollectionNodes/Private/Dataflow/GeometryCollectionTransferVertexAttributeNode.cpp
+++ b/Engine/Plugins/Experimental/GeometryCollectionPlugin/Source/GeometryCollectionNodes/Private/Dataflow/GeometryCollectionTransferVertexAttributeNode.cpp
@@ -9,6 +9,7 @@
 #include "Dataflow/DataflowInputOutput.h"
 #include "GeometryCollection/GeometryCollectionAlgo.h"
 #include "GeometryCollection/ManagedArrayAccessor.h"
+#include "GeometryCollection/Facades/CollectionVertexBoneWeightsFacade.h"
 
 #include UE_INLINE_GENERATED_CPP_BY_NAME(GeometryCollectionTransferVertexAttributeNode)
 #define LOCTEXT_NAMESPACE "FGeometryCollectionTransferVertexAttributeNode"
@@ -83,6 +84,15 @@ namespace UE::Private {
 		TManagedArrayAccessor<int32> FaceCount;
 	};
 
+	struct FTransferData
+	{
+		TArray<int32> SourceIndices;
+		TArray<float> SourceWeights;
+		float FalloffScale;
+		FTransferData(): FalloffScale(0.f) {}
+		FTransferData(TArray<int32> SourceIndicesIn, TArray<float> SourceWeightsIn, float FalloffScaleIn) :
+			SourceIndices(SourceIndicesIn), SourceWeights(SourceWeightsIn), FalloffScale(FalloffScaleIn) {}
+	};
 }
 
 void FGeometryCollectionTransferVertexAttributeNode::Evaluate(UE::Dataflow::FContext& Context, const FDataflowOutput* Out) const
@@ -99,34 +109,95 @@ void FGeometryCollectionTransferVertexAttributeNode::Evaluate(UE::Dataflow::FCon
 
 		if (Target.IsValid() && Sample.IsValid())
 		{
-			if (const TManagedArray<float>* SourceAttributeFloatArray = Sample.GetAttributeArray<float>(Key.Attribute, Key.Group))
+			TArray<FIntVector2> AlignedGeometry = FindSourceToTargetGeometryMap(SampleCollection, TargetCollection);
+			TArray<UE::Private::FTransferData> TransferDataArray; //store barycentric weight info
+			if (AlignedGeometry.Num() == SampleCollection.NumElements(FGeometryCollection::GeometryGroup))
 			{
-				TManagedArray<float>* TargetAttributeArray = Target.GetAttributeArray<float>(Key.Attribute, Key.Group);
-				TargetAttributeArray->Fill(0.f);
-
-				TArray<FIntVector2> AlignedGeometry = FindSourceToTargetGeometryMap(SampleCollection, TargetCollection);
-				if (AlignedGeometry.Num() == TargetCollection.NumElements(FGeometryCollection::GeometryGroup))
+				TransferDataArray = PairedGeometryTransfer(Key, AlignedGeometry, Sample, Target);
+			}
+			else
+			{
+				TransferDataArray = NearestVertexTransfer(Key, Sample, Target);
+			}
+			if (bTransferSkinWeights)
+			{
+				GeometryCollection::Facades::FVertexBoneWeightsFacade SampleVertexBoneWeightsFacade(SampleCollection);
+				GeometryCollection::Facades::FVertexBoneWeightsFacade TargetVertexBoneWeightsFacade(TargetCollection);
+				const TManagedArray<TArray<int32>>* SampleBoneIndices = SampleVertexBoneWeightsFacade.FindBoneIndices();
+				const TManagedArray<TArray<float>>* SampleBoneWeights = SampleVertexBoneWeightsFacade.FindBoneWeights();
+				if (SampleBoneIndices && SampleBoneWeights)
 				{
-					PairedGeometryTransfer<float>(Key, AlignedGeometry, Sample, Target, SourceAttributeFloatArray, TargetAttributeArray);
+					//
+					// Compute the bone index mappings. This allows the transfer operator to retarget weights to the correct skeleton.
+					//
+					TMap<int32, int32> SampleBoneToTargetBone = FindSourceToTargetTransformMap(SampleCollection, TargetCollection);
+					for (int32 VertexIndex = 0; VertexIndex < Target.Vertex.Num(); VertexIndex++)
+					{
+						UE::Private::FTransferData TransferData = TransferDataArray[VertexIndex];
+						if (TransferData.SourceIndices.Num())
+						{
+							TMap<int32, float> BoneWeightBucket;
+							for (int32 SourceLocalIdx = 0; SourceLocalIdx < TransferData.SourceIndices.Num(); ++SourceLocalIdx)
+							{
+								for (int32 BoneLocalIdx = 0; BoneLocalIdx < (*SampleBoneIndices)[TransferData.SourceIndices[SourceLocalIdx]].Num(); ++BoneLocalIdx)
+								{
+									int32 BoneIndex = (*SampleBoneIndices)[TransferData.SourceIndices[SourceLocalIdx]][BoneLocalIdx];
+									float BoneWeight = (*SampleBoneWeights)[TransferData.SourceIndices[SourceLocalIdx]][BoneLocalIdx];
+									if (BoneWeightBucket.Contains(BoneIndex))
+									{
+										BoneWeightBucket[BoneIndex] += TransferData.SourceWeights[SourceLocalIdx] * BoneWeight;
+									}
+									else
+									{
+										BoneWeightBucket.Add(BoneIndex, TransferData.SourceWeights[SourceLocalIdx] * BoneWeight);
+									}
+								}
+							}
+
+							TArray<int32> VertexBoneIndex;
+							TArray<float> VertexBoneWeight;
+							for (const TPair<int32, float>& Pair : BoneWeightBucket)
+							{
+								VertexBoneIndex.Add(Pair.Key);
+								VertexBoneWeight.Add(Pair.Value);
+							}
+							TargetVertexBoneWeightsFacade.ModifyBoneWeight(VertexIndex, VertexBoneIndex, VertexBoneWeight);
+						}	
+					}
 				}
-				else
+			} //end if (bTransferSkinWeights)
+			else if (const TManagedArray<float>* SourceAttributeFloatArray = Sample.GetAttributeArray<float>(Key.Attribute, Key.Group))
+			{
+				TManagedArray<float>* TargetAttributeArray = Target.GetAttributeArray<float>(Key.Attribute, Key.Group);
+				TargetAttributeArray->Fill(0.f);
+				for (int32 VertexIndex = 0; VertexIndex < TargetAttributeArray->Num(); VertexIndex++)
 				{
-					NearestVertexTransfer<float>(Key, Sample, Target, SourceAttributeFloatArray, TargetAttributeArray);
+					UE::Private::FTransferData TransferData = TransferDataArray[VertexIndex];
+					if (TransferData.SourceIndices.Num())
+					{
+						for (int32 SourceLocalIdx = 0; SourceLocalIdx < TransferData.SourceIndices.Num(); ++SourceLocalIdx)
+						{
+							(*TargetAttributeArray)[VertexIndex] += (*SourceAttributeFloatArray)[TransferData.SourceIndices[SourceLocalIdx]]
+								* TransferData.SourceWeights[SourceLocalIdx] * TransferData.FalloffScale;
+						}
+					}
 				}
 			}
 			else if (const TManagedArray<FLinearColor>* SourceAttributeColorArray = Sample.GetAttributeArray<FLinearColor>(Key.Attribute, Key.Group))
 			{
 				TManagedArray<FLinearColor>* TargetAttributeArray = Target.GetAttributeArray<FLinearColor>(Key.Attribute, Key.Group);
 				TargetAttributeArray->Fill(FLinearColor(0, 0, 0, 0));
-
-				TArray<FIntVector2> AlignedGeometry = FindSourceToTargetGeometryMap(SampleCollection, TargetCollection);
-				if (AlignedGeometry.Num() == TargetCollection.NumElements(FGeometryCollection::GeometryGroup))
-				{
-					PairedGeometryTransfer<FLinearColor>(Key, AlignedGeometry, Sample, Target, SourceAttributeColorArray, TargetAttributeArray);
-				}
-				else
+				for (int32 VertexIndex = 0; VertexIndex < TargetAttributeArray->Num(); VertexIndex++)
 				{
-					NearestVertexTransfer<FLinearColor>(Key, Sample, Target, SourceAttributeColorArray, TargetAttributeArray);
+					UE::Private::FTransferData TransferData = TransferDataArray[VertexIndex];
+					if (TransferData.SourceIndices.Num())
+					{
+						for (int32 SourceLocalIdx = 0; SourceLocalIdx < TransferData.SourceIndices.Num(); ++SourceLocalIdx)
+						{
+							(*TargetAttributeArray)[VertexIndex] += (*SourceAttributeColorArray)[TransferData.SourceIndices[SourceLocalIdx]]
+								* TransferData.SourceWeights[SourceLocalIdx] * TransferData.FalloffScale;
+						}
+					}
 				}
 			}
 		}
@@ -139,14 +210,13 @@ void FGeometryCollectionTransferVertexAttributeNode::Evaluate(UE::Dataflow::FCon
 	}
 }
 
-
-TArray<FIntVector2> FGeometryCollectionTransferVertexAttributeNode::FindSourceToTargetGeometryMap(const FManagedArrayCollection& AttributeCollectionVal, const FManagedArrayCollection& CollectionVal) const
+TArray<FIntVector2> FGeometryCollectionTransferVertexAttributeNode::FindSourceToTargetGeometryMap(const FManagedArrayCollection& SourceCollection, const FManagedArrayCollection& TargetCollection) const
 {
 	TArray<FIntVector2> Mapping;
-	const TManagedArray<FString>* SourceName = AttributeCollectionVal.FindAttribute<FString>(FName("BoneName"), FTransformCollection::TransformGroup);
-	const TManagedArray<int32>* SourceGeometryGroup = AttributeCollectionVal.FindAttribute<int32>(FName("TransformToGeometryIndex"), FTransformCollection::TransformGroup);
-	const TManagedArray<FString>* TargetName = CollectionVal.FindAttribute<FString>(FName("BoneName"), FTransformCollection::TransformGroup);
-	const TManagedArray<int32>* TargetGeometryGroup = CollectionVal.FindAttribute<int32>(FName("TransformToGeometryIndex"), FTransformCollection::TransformGroup);
+	const TManagedArray<FString>* SourceName = SourceCollection.FindAttribute<FString>(FName("BoneName"), FTransformCollection::TransformGroup);
+	const TManagedArray<int32>* SourceGeometryGroup = SourceCollection.FindAttribute<int32>(FName("TransformToGeometryIndex"), FTransformCollection::TransformGroup);
+	const TManagedArray<FString>* TargetName = TargetCollection.FindAttribute<FString>(FName("BoneName"), FTransformCollection::TransformGroup);
+	const TManagedArray<int32>* TargetGeometryGroup = TargetCollection.FindAttribute<int32>(FName("TransformToGeometryIndex"), FTransformCollection::TransformGroup);
 	if (SourceName && SourceGeometryGroup && TargetName && TargetGeometryGroup)
 	{
 		for (int i = 0; i < SourceName->Num(); i++)
@@ -166,230 +236,253 @@ TArray<FIntVector2> FGeometryCollectionTransferVertexAttributeNode::FindSourceTo
 	return Mapping;
 }
 
-template<typename T>
-void FGeometryCollectionTransferVertexAttributeNode::PairedGeometryTransfer(FCollectionAttributeKey Key, const TArray<FIntVector2>& PairedGeometry,
-	const UE::Private::FTransferFacade& Sample, UE::Private::FTransferFacade& Target, const TManagedArray<T>* SourceAttributeArray, TManagedArray<T>* TargetAttributeArray) const
+TMap<int32, int32> FGeometryCollectionTransferVertexAttributeNode::FindSourceToTargetTransformMap(const FManagedArrayCollection& SourceCollection, const FManagedArrayCollection& TargetCollection) const
 {
-	//TargetAttributeArray has to be filled with empty values before (attribute types: float, FLinearColor)
-	if (SourceAttributeArray && TargetAttributeArray)
+	TMap<int32, int32> SourceIndexToTargetIndex;
+	TMap<FString, int32> TargetNameToIndex;
+	const TManagedArray<FString>* SourceName = SourceCollection.FindAttribute<FString>(FName("BoneName"), FTransformCollection::TransformGroup);
+	const TManagedArray<FString>* TargetName = TargetCollection.FindAttribute<FString>(FName("BoneName"), FTransformCollection::TransformGroup);
+	if (SourceName && TargetName)
 	{
-		Chaos::FReal SphereFullRadius;
-
-		if (SampleScale == EDataflowTransferVertexAttributeNodeSampleScale::Asset_Edge || SampleScale == EDataflowTransferVertexAttributeNodeSampleScale::Asset_Bound)
+		for (int32 j = 0; j < TargetName->Num(); j++)
 		{
-			// Build component space vertices for TargetCollection
-			TArray<FVector3f> ComponentSpaceFullTargetVertices; // size of num vertices of the geometry entry. 
-			BuildComponentSpaceVertices(Target.Transform.Get(), Target.Parent.Get(), Target.BoneMap.Get(), Target.Vertex.Get(),
-				0, Target.Vertex.Num(), ComponentSpaceFullTargetVertices);
-
-			// Build component space vertices for TargetCollection
-			TArray<FVector3f> ComponentSpaceFullVertices; // size of num vertices of the geometry entry
-			BuildComponentSpaceVertices(Sample.Transform.Get(), Sample.Parent.Get(), Sample.BoneMap.Get(), Sample.Vertex.Get(),
-				0, Sample.Vertex.Num(), ComponentSpaceFullVertices);
-			if (SampleScale == EDataflowTransferVertexAttributeNodeSampleScale::Asset_Edge)
+			TargetNameToIndex.Add((*TargetName)[j], j);
+		}
+		for (int32 i = 0; i < SourceName->Num(); i++)
+		{
+			if (int32* TargetIndex = TargetNameToIndex.Find((*SourceName)[i]))
 			{
-				SphereFullRadius = Chaos::FReal(EdgeMultiplier * FMath::Max(
-					MaxEdgeLength(ComponentSpaceFullTargetVertices, Target.Indices.Get(), 0, 0, Target.Indices.Num()),
-					MaxEdgeLength(ComponentSpaceFullVertices, Sample.Indices.Get(), 0, 0, Sample.Indices.Num())));
+				SourceIndexToTargetIndex.Add(i, *TargetIndex);
 			}
-			else if (SampleScale == EDataflowTransferVertexAttributeNodeSampleScale::Asset_Bound)
+		}
+	}
+
+	return SourceIndexToTargetIndex;
+}
+
+TArray<UE::Private::FTransferData> FGeometryCollectionTransferVertexAttributeNode::PairedGeometryTransfer(FCollectionAttributeKey Key, const TArray<FIntVector2>& PairedGeometry,
+	const UE::Private::FTransferFacade& Sample, const UE::Private::FTransferFacade& Target) const
+{
+	TArray<UE::Private::FTransferData> TransferDataArray;
+	TransferDataArray.SetNum(Target.Vertex.Num());
+	Chaos::FReal SphereFullRadius;
+
+	if (SampleScale == EDataflowTransferVertexAttributeNodeSampleScale::Asset_Edge || SampleScale == EDataflowTransferVertexAttributeNodeSampleScale::Asset_Bound)
+	{
+		// Build component space vertices for TargetCollection
+		TArray<FVector3f> ComponentSpaceFullTargetVertices; // size of num vertices of the geometry entry. 
+		BuildComponentSpaceVertices(Target.Transform.Get(), Target.Parent.Get(), Target.BoneMap.Get(), Target.Vertex.Get(),
+			0, Target.Vertex.Num(), ComponentSpaceFullTargetVertices);
+
+		// Build component space vertices for TargetCollection
+		TArray<FVector3f> ComponentSpaceFullVertices; // size of num vertices of the geometry entry
+		BuildComponentSpaceVertices(Sample.Transform.Get(), Sample.Parent.Get(), Sample.BoneMap.Get(), Sample.Vertex.Get(),
+			0, Sample.Vertex.Num(), ComponentSpaceFullVertices);
+		if (SampleScale == EDataflowTransferVertexAttributeNodeSampleScale::Asset_Edge)
+		{
+			SphereFullRadius = Chaos::FReal(EdgeMultiplier * FMath::Max(
+				MaxEdgeLength(ComponentSpaceFullTargetVertices, Target.Indices.Get(), 0, 0, Target.Indices.Num()),
+				MaxEdgeLength(ComponentSpaceFullVertices, Sample.Indices.Get(), 0, 0, Sample.Indices.Num())));
+		}
+		else if (SampleScale == EDataflowTransferVertexAttributeNodeSampleScale::Asset_Bound)
+		{
+			Chaos::TVec3<float> CoordMaxs(-FLT_MAX);
+			Chaos::TVec3<float> CoordMins(FLT_MAX);
+			for (int32 i = 0; i < ComponentSpaceFullVertices.Num(); i++)
 			{
-				Chaos::TVec3<float> CoordMaxs(-FLT_MAX);
-				Chaos::TVec3<float> CoordMins(FLT_MAX);
-				for (int32 i = 0; i < ComponentSpaceFullVertices.Num(); i++)
+				for (int32 j = 0; j < 3; j++)
 				{
-					for (int32 j = 0; j < 3; j++)
+					if (ComponentSpaceFullVertices[i][j] > CoordMaxs[j])
 					{
-						if (ComponentSpaceFullVertices[i][j] > CoordMaxs[j])
-						{
-							CoordMaxs[j] = ComponentSpaceFullVertices[i][j];
-						}
-						if (ComponentSpaceFullVertices[i][j] < CoordMins[j])
-						{
-							CoordMins[j] = ComponentSpaceFullVertices[i][j];
-						}
+						CoordMaxs[j] = ComponentSpaceFullVertices[i][j];
+					}
+					if (ComponentSpaceFullVertices[i][j] < CoordMins[j])
+					{
+						CoordMins[j] = ComponentSpaceFullVertices[i][j];
 					}
 				}
-				Chaos::TVec3<float> CoordDiff = (CoordMaxs - CoordMins) * BoundMultiplier;
-				SphereFullRadius = Chaos::FReal(FGenericPlatformMath::Min(CoordDiff[0], FGenericPlatformMath::Min(CoordDiff[1], CoordDiff[2])));
 			}
+			Chaos::TVec3<float> CoordDiff = (CoordMaxs - CoordMins) * BoundMultiplier;
+			SphereFullRadius = Chaos::FReal(FGenericPlatformMath::Min(CoordDiff[0], FGenericPlatformMath::Min(CoordDiff[1], CoordDiff[2])));
 		}
-		ParallelFor(PairedGeometry.Num(), [&](int32 Pdx)
+	}
+	ParallelFor(PairedGeometry.Num(), [&](int32 Pdx)
+	{
+		int32 AttributeGeometryIndex = PairedGeometry[Pdx][0];
+		int32 TargetGeometryIndex = PairedGeometry[Pdx][1];
+		if (ensure(0 <= AttributeGeometryIndex && AttributeGeometryIndex < Sample.VertexStart.Num()))
+		{
+			if (ensure(0 <= TargetGeometryIndex && TargetGeometryIndex < Target.VertexStart.Num()))
 			{
-				int32 AttributeGeometryIndex = PairedGeometry[Pdx][0];
-				int32 TargetGeometryIndex = PairedGeometry[Pdx][1];
-				if (ensure(0 <= AttributeGeometryIndex && AttributeGeometryIndex < Sample.VertexStart.Num()))
+				// Build component space vertices for TargetCollection
+				TArray<FVector3f> ComponentSpaceTargetVertices; // size of num vertices of the geometry entry. 
+				BuildComponentSpaceVertices(Target.Transform.Get(), Target.Parent.Get(), Target.BoneMap.Get(), Target.Vertex.Get(),
+					Target.VertexStart[TargetGeometryIndex], Target.VertexCount[TargetGeometryIndex], ComponentSpaceTargetVertices);
+
+				// Build component space vertices for SampleCollection
+				TArray<FVector3f> ComponentSpaceVertices; // size of num vertices of the geometry entry
+				BuildComponentSpaceVertices(Sample.Transform.Get(), Sample.Parent.Get(), Sample.BoneMap.Get(), Sample.Vertex.Get(),
+					Sample.VertexStart[AttributeGeometryIndex], Sample.VertexCount[AttributeGeometryIndex], ComponentSpaceVertices);
+
+				// build Sphere based BVH
+				Chaos::FReal SphereRadius = SphereFullRadius;
+				if (SampleScale == EDataflowTransferVertexAttributeNodeSampleScale::Component_Edge)
 				{
-					if (ensure(0 <= TargetGeometryIndex && TargetGeometryIndex < Target.VertexStart.Num()))
-					{
-						// Build component space vertices for TargetCollection
-						TArray<FVector3f> ComponentSpaceTargetVertices; // size of num vertices of the geometry entry. 
-						BuildComponentSpaceVertices(Target.Transform.Get(), Target.Parent.Get(), Target.BoneMap.Get(), Target.Vertex.Get(),
-							Target.VertexStart[TargetGeometryIndex], Target.VertexCount[TargetGeometryIndex], ComponentSpaceTargetVertices);
-
-						// Build component space vertices for SampleCollection
-						TArray<FVector3f> ComponentSpaceVertices; // size of num vertices of the geometry entry
-						BuildComponentSpaceVertices(Sample.Transform.Get(), Sample.Parent.Get(), Sample.BoneMap.Get(), Sample.Vertex.Get(),
-							Sample.VertexStart[AttributeGeometryIndex], Sample.VertexCount[AttributeGeometryIndex], ComponentSpaceVertices);
-
-						// build Sphere based BVH
-						Chaos::FReal SphereRadius = SphereFullRadius;
-						if (SampleScale == EDataflowTransferVertexAttributeNodeSampleScale::Component_Edge)
-						{
-							SphereRadius = Chaos::FReal(EdgeMultiplier * FMath::Max(
-								MaxEdgeLength(ComponentSpaceTargetVertices, Target.Indices.Get(), Target.VertexStart[TargetGeometryIndex], Target.FaceStart[TargetGeometryIndex], Target.FaceCount[TargetGeometryIndex]),
-								MaxEdgeLength(ComponentSpaceVertices, Sample.Indices.Get(), Sample.VertexStart[AttributeGeometryIndex], Sample.FaceStart[AttributeGeometryIndex], Sample.FaceCount[AttributeGeometryIndex])));
-						}
+					SphereRadius = Chaos::FReal(EdgeMultiplier * FMath::Max(
+						MaxEdgeLength(ComponentSpaceTargetVertices, Target.Indices.Get(), Target.VertexStart[TargetGeometryIndex], Target.FaceStart[TargetGeometryIndex], Target.FaceCount[TargetGeometryIndex]),
+						MaxEdgeLength(ComponentSpaceVertices, Sample.Indices.Get(), Sample.VertexStart[AttributeGeometryIndex], Sample.FaceStart[AttributeGeometryIndex], Sample.FaceCount[AttributeGeometryIndex])));
+				}
 
-						int32 TargetVertexStartVal = Target.VertexStart[TargetGeometryIndex];
-						int32 TargetVertexCountVal = Target.VertexCount[TargetGeometryIndex];
-						int32 VertexStartVal = Sample.VertexStart[AttributeGeometryIndex];
-						int32 FaceStartVal = Sample.FaceStart[AttributeGeometryIndex];
-						int32 FaceCountVal = Sample.FaceCount[AttributeGeometryIndex];
+				int32 TargetVertexStartVal = Target.VertexStart[TargetGeometryIndex];
+				int32 TargetVertexCountVal = Target.VertexCount[TargetGeometryIndex];
+				int32 VertexStartVal = Sample.VertexStart[AttributeGeometryIndex];
+				int32 FaceStartVal = Sample.FaceStart[AttributeGeometryIndex];
+				int32 FaceCountVal = Sample.FaceCount[AttributeGeometryIndex];
 
-						if (BoundingVolumeType == EDataflowTransferVertexAttributeNodeBoundingVolume::Triangle)
+				if (BoundingVolumeType == EDataflowTransferVertexAttributeNodeBoundingVolume::Triangle)
+				{
+					TArray<Chaos::TVec3<Chaos::FReal>> ComponentSpaceVerticesTVec3;
+					ComponentSpaceVerticesTVec3.SetNum(ComponentSpaceVertices.Num());
+					for (int32 SourceIndex = 0; SourceIndex < ComponentSpaceVerticesTVec3.Num(); SourceIndex++)
+					{
+						ComponentSpaceVerticesTVec3[SourceIndex] = Chaos::TVec3<Chaos::FReal>(ComponentSpaceVertices[SourceIndex]);
+					}
+					TConstArrayView<Chaos::TVec3<Chaos::FReal>> ConstComponentSpaceVertices(ComponentSpaceVerticesTVec3);
+					Chaos::FTriangleMesh TriangleMesh;
+					TArray<Chaos::TVec3<int32>> SourceElements;
+					SourceElements.SetNum(FaceCountVal);
+					for (int32 ElementIndex = 0; ElementIndex < FaceCountVal; ++ElementIndex)
+					{
+						FIntVector3 Element = Sample.Indices[FaceStartVal + ElementIndex];
+						SourceElements[ElementIndex] = Chaos::TVec3<int32>(Element[0] - VertexStartVal, Element[1] - VertexStartVal, Element[2] - VertexStartVal);
+					}
+					TriangleMesh.Init(SourceElements);
+					Chaos::FTriangleMesh::TSpatialHashType<Chaos::FReal> SpatialHash;
+					TriangleMesh.BuildSpatialHash(ConstComponentSpaceVertices, SpatialHash, SphereRadius);
+					for (int32 TargetIndex = 0; TargetIndex < TargetVertexCountVal; TargetIndex++)
+					{
+						TArray<Chaos::TTriangleCollisionPoint<Chaos::FReal>> Result;
+						if (TriangleMesh.PointClosestTriangleQuery(SpatialHash, ConstComponentSpaceVertices,
+							TargetIndex, Chaos::TVec3<Chaos::FReal>(ComponentSpaceTargetVertices[TargetIndex]), SphereRadius / 2.f, SphereRadius / 2.f,
+							[](const int32 PointIndex, const int32 TriangleIndex)->bool {return true; }, Result))
 						{
-							TArray<Chaos::TVec3<Chaos::FReal>> ComponentSpaceVerticesTVec3;
-							ComponentSpaceVerticesTVec3.SetNum(ComponentSpaceVertices.Num());
-							for (int32 SourceIndex = 0; SourceIndex < ComponentSpaceVerticesTVec3.Num(); SourceIndex++)
-							{
-								ComponentSpaceVerticesTVec3[SourceIndex] = Chaos::TVec3<Chaos::FReal>(ComponentSpaceVertices[SourceIndex]);
-							}
-							TConstArrayView<Chaos::TVec3<Chaos::FReal>> ConstComponentSpaceVertices(ComponentSpaceVerticesTVec3);
-							Chaos::FTriangleMesh TriangleMesh;
-							TArray<Chaos::TVec3<int32>> SourceElements;
-							SourceElements.SetNum(FaceCountVal);
-							for (int32 ElementIndex = 0; ElementIndex < FaceCountVal; ++ElementIndex)
-							{
-								FIntVector3 Element = Sample.Indices[FaceStartVal + ElementIndex];
-								SourceElements[ElementIndex] = Chaos::TVec3<int32>(Element[0] - VertexStartVal, Element[1] - VertexStartVal, Element[2] - VertexStartVal);
-							}
-							TriangleMesh.Init(SourceElements);
-							Chaos::FTriangleMesh::TSpatialHashType<Chaos::FReal> SpatialHash;
-							TriangleMesh.BuildSpatialHash(ConstComponentSpaceVertices, SpatialHash, SphereRadius);
-							for (int32 TargetIndex = 0; TargetIndex < TargetVertexCountVal; TargetIndex++)
+							for (const Chaos::TTriangleCollisionPoint<Chaos::FReal>& CollisionPoint : Result)
 							{
-								TArray<Chaos::TTriangleCollisionPoint<Chaos::FReal>> Result;
-								if (TriangleMesh.PointClosestTriangleQuery(SpatialHash, ConstComponentSpaceVertices,
-									TargetIndex, Chaos::TVec3<Chaos::FReal>(ComponentSpaceTargetVertices[TargetIndex]), SphereRadius / 2.f, SphereRadius / 2.f,
-									[](const int32 PointIndex, const int32 TriangleIndex)->bool {return true; }, Result))
+								float CurrentDistance = abs(CollisionPoint.Phi);
+								float TriRadius = FalloffThreshold * MaxEdgeLength(ComponentSpaceVertices, Sample.Indices.Get(), VertexStartVal, FaceStartVal + CollisionPoint.Indices[1], 1);
+								float FalloffScale = CalculateFalloffScale(Falloff, TriRadius, CurrentDistance);
+								if (!FMath::IsNearlyZero(FalloffScale))
 								{
-									for (const Chaos::TTriangleCollisionPoint<Chaos::FReal>& CollisionPoint : Result)
+									int32 TargetCandidateIndex = CollisionPoint.Indices[0] + TargetVertexStartVal;
+									TArray<int32> TransferIndices;
+									TArray<float> TransferWeights;
+									for (int32 k = 0; k < 3; k++)
 									{
-										float CurrentDistance = abs(CollisionPoint.Phi);
-										float TriRadius = FalloffThreshold * MaxEdgeLength(ComponentSpaceVertices, Sample.Indices.Get(), VertexStartVal, FaceStartVal + CollisionPoint.Indices[1], 1);
-										float FalloffScale = CalculateFalloffScale(Falloff, TriRadius, CurrentDistance);
-										if (!FMath::IsNearlyZero(FalloffScale))
-										{
-											int32 TargetCandidateIndex = CollisionPoint.Indices[0] + TargetVertexStartVal;
-											if (ensure(0 <= TargetCandidateIndex && TargetCandidateIndex < TargetAttributeArray->Num()))
-											{
-												for (int32 k = 0; k < 3; k++)
-												{
-													(*TargetAttributeArray)[TargetCandidateIndex] += FalloffScale * (CollisionPoint.Bary[k+1] * (*SourceAttributeArray)[Sample.Indices[FaceStartVal + CollisionPoint.Indices[1]][k]]);
-												}
-												break;
-											}
-										}
+										TransferIndices.Add(Sample.Indices[FaceStartVal + CollisionPoint.Indices[1]][k]);
+										TransferWeights.Add(CollisionPoint.Bary[k + 1]);
 									}
+									TransferDataArray[TargetCandidateIndex] = UE::Private::FTransferData(TransferIndices, TransferWeights, FalloffScale);
+									break;
 								}
 							}
 						}
-						else if (BoundingVolumeType == EDataflowTransferVertexAttributeNodeBoundingVolume::Vertex)
+					}
+				}
+				else if (BoundingVolumeType == EDataflowTransferVertexAttributeNodeBoundingVolume::Vertex)
+				{
+					TUniquePtr<BVH> VertexBVH(BuildParticleSphereBVH(ComponentSpaceTargetVertices, SphereRadius));
+					for (int32 i = 0; i < FaceCountVal; i++)
+					{
+						FIntVector3 ComponentTriangle = Sample.Indices[FaceStartVal + i] - FIntVector3(VertexStartVal);
+						TArray<int32> TargetVertexIntersection({});
+						TriangleToVertexIntersections(*VertexBVH, ComponentSpaceVertices, ComponentTriangle, TargetVertexIntersection);
+
+						for (int32 j = 0; j < TargetVertexIntersection.Num(); j++)
 						{
-							TUniquePtr<BVH> VertexBVH(BuildParticleSphereBVH(ComponentSpaceTargetVertices, SphereRadius));
-							for (int32 i = 0; i < FaceCountVal; i++)
+							Chaos::TVector<float, 3> Bary,
+								TriPos0(ComponentSpaceVertices[ComponentTriangle[0]]),
+								TriPos1(ComponentSpaceVertices[ComponentTriangle[1]]),
+								TriPos2(ComponentSpaceVertices[ComponentTriangle[2]]),
+								ParticlePos(ComponentSpaceTargetVertices[TargetVertexIntersection[j]]);
+
+							Chaos::TVector<float, 3> ClosestPoint = Chaos::FindClosestPointAndBaryOnTriangle(TriPos0, TriPos1, TriPos2, ParticlePos, Bary);
+							float CurrentDistance = (ParticlePos - ClosestPoint).Size();
+							float TriRadius = FalloffThreshold * MaxEdgeLength(ComponentSpaceVertices, Sample.Indices.Get(), VertexStartVal, FaceStartVal + i, 1);
+							float FalloffScale = CalculateFalloffScale(Falloff, TriRadius, CurrentDistance);
+							if (!FMath::IsNearlyZero(FalloffScale))
 							{
-								FIntVector3 ComponentTriangle = Sample.Indices[FaceStartVal + i] - FIntVector3(VertexStartVal);
-								TArray<int32> TargetVertexIntersection({});
-								TriangleToVertexIntersections(*VertexBVH, ComponentSpaceVertices, ComponentTriangle, TargetVertexIntersection);
-
-								for (int32 j = 0; j < TargetVertexIntersection.Num(); j++)
+								int32 TargetIndex = TargetVertexIntersection[j] + TargetVertexStartVal;
+								TArray<int32> TransferIndices;
+								TArray<float> TransferWeights;
+								for (int32 k = 0; k < 3; k++)
 								{
-									Chaos::TVector<float, 3> Bary,
-										TriPos0(ComponentSpaceVertices[ComponentTriangle[0]]),
-										TriPos1(ComponentSpaceVertices[ComponentTriangle[1]]),
-										TriPos2(ComponentSpaceVertices[ComponentTriangle[2]]),
-										ParticlePos(ComponentSpaceTargetVertices[TargetVertexIntersection[j]]);
-
-									Chaos::TVector<float, 3> ClosestPoint = Chaos::FindClosestPointAndBaryOnTriangle(TriPos0, TriPos1, TriPos2, ParticlePos, Bary);
-									float CurrentDistance = (ParticlePos - ClosestPoint).Size();
-									float TriRadius = FalloffThreshold * MaxEdgeLength(ComponentSpaceVertices, Sample.Indices.Get(), VertexStartVal, FaceStartVal + i, 1);
-									float FalloffScale = CalculateFalloffScale(Falloff, TriRadius, CurrentDistance);
-									if (!FMath::IsNearlyZero(FalloffScale))
-									{
-										int32 TargetIndex = TargetVertexIntersection[j] + TargetVertexStartVal;
-										if (ensure(0 <= TargetIndex && TargetIndex < TargetAttributeArray->Num()))
-										{
-											for (int32 k = 0; k < 3; k++)
-											{
-												(*TargetAttributeArray)[TargetIndex] += FalloffScale * (Bary[k] * (*SourceAttributeArray)[ComponentTriangle[k] + VertexStartVal]);
-											}
-											break;
-										}
-									}
+									TransferIndices.Add(ComponentTriangle[k] + VertexStartVal);
+									TransferWeights.Add(Bary[k]);
 								}
+								TransferDataArray[TargetIndex] = UE::Private::FTransferData(TransferIndices, TransferWeights, FalloffScale);
+								break;
 							}
 						}
 					}
 				}
-			});
-	}
+			}
+		}
+	});
+	return TransferDataArray;
 }
 
-template<typename T>
-void FGeometryCollectionTransferVertexAttributeNode::NearestVertexTransfer(FCollectionAttributeKey Key, const UE::Private::FTransferFacade& Sample, UE::Private::FTransferFacade& Target, const TManagedArray<T>* SourceAttributeArray, TManagedArray<T>* TargetAttributeArray) const
+TArray<UE::Private::FTransferData> FGeometryCollectionTransferVertexAttributeNode::NearestVertexTransfer(FCollectionAttributeKey Key, const UE::Private::FTransferFacade& Sample, const UE::Private::FTransferFacade& Target) const
 {
-	//TargetAttributeArray has to be filled with empty values before (attribute types: float, FLinearColor)
-	if (SourceAttributeArray && TargetAttributeArray)
+	TArray<UE::Private::FTransferData> TransferDataArray;
+	TransferDataArray.SetNum(Target.Vertex.Num());
+	// Build component space vertices for TargetCollection
+	TArray<FVector3f> ComponentSpaceTargetVertices;
+	BuildComponentSpaceVertices(Target.Transform.Get(), Target.Parent.Get(), Target.BoneMap.Get(), Target.Vertex.Get(), 0, Target.Vertex.Num(), ComponentSpaceTargetVertices);
+
+	// Build component space vertices for SourceCollection
+	TArray<FVector3f> ComponentSpaceVertices;
+	BuildComponentSpaceVertices(Sample.Transform.Get(), Sample.Parent.Get(), Sample.BoneMap.Get(), Sample.Vertex.Get(), 0, Sample.Vertex.Num(), ComponentSpaceVertices);
+
+	// build Sphere based BVH
+	Chaos::FReal SphereRadius = Chaos::FReal(EdgeMultiplier * FMath::Max(
+		MaxEdgeLength(ComponentSpaceTargetVertices, Target.Indices.Get(), 0, 0, Target.Indices.Num()),
+		MaxEdgeLength(ComponentSpaceVertices, Sample.Indices.Get(), 0, 0, Sample.Indices.Num())));
+	TUniquePtr<BVH> VertexBVH(BuildParticleSphereBVH(ComponentSpaceTargetVertices, SphereRadius));
+
+	for (int32 i = 0; i < Sample.Indices.Num(); i++)
 	{
-		// Build component space vertices for TargetCollection
-		TArray<FVector3f> ComponentSpaceTargetVertices;
-		BuildComponentSpaceVertices(Target.Transform.Get(), Target.Parent.Get(), Target.BoneMap.Get(), Target.Vertex.Get(), 0, Target.Vertex.Num(), ComponentSpaceTargetVertices);
-
-		// Build component space vertices for SourceCollection
-		TArray<FVector3f> ComponentSpaceVertices;
-		BuildComponentSpaceVertices(Sample.Transform.Get(), Sample.Parent.Get(), Sample.BoneMap.Get(), Sample.Vertex.Get(), 0, Sample.Vertex.Num(), ComponentSpaceVertices);
+		FIntVector3 Triangle = Sample.Indices[i];
+		TArray<int32> TargetVertexIntersection({});
+		TriangleToVertexIntersections(*VertexBVH, ComponentSpaceVertices, Triangle, TargetVertexIntersection);
 
-		// build Sphere based BVH
-		Chaos::FReal SphereRadius = Chaos::FReal(EdgeMultiplier * FMath::Max(
-			MaxEdgeLength(ComponentSpaceTargetVertices, Target.Indices.Get(), 0, 0, Target.Indices.Num()),
-			MaxEdgeLength(ComponentSpaceVertices, Sample.Indices.Get(), 0, 0, Sample.Indices.Num())));
-		TUniquePtr<BVH> VertexBVH(BuildParticleSphereBVH(ComponentSpaceTargetVertices, SphereRadius));
-
-		for (int32 i = 0; i < Sample.Indices.Num(); i++)
+		for (int32 j = 0; j < TargetVertexIntersection.Num(); j++)
 		{
-			FIntVector3 Triangle = Sample.Indices[i];
-			TArray<int32> TargetVertexIntersection({});
-			TriangleToVertexIntersections(*VertexBVH, ComponentSpaceVertices, Triangle, TargetVertexIntersection);
-
-			for (int32 j = 0; j < TargetVertexIntersection.Num(); j++)
+			Chaos::TVector<float, 3> Bary,
+				TriPos0(ComponentSpaceVertices[Sample.Indices[i][0]]),
+				TriPos1(ComponentSpaceVertices[Sample.Indices[i][1]]),
+				TriPos2(ComponentSpaceVertices[Sample.Indices[i][2]]),
+				ParticlePos(ComponentSpaceTargetVertices[TargetVertexIntersection[j]]);
+
+			Chaos::TVector<float, 3> ClosestPoint = Chaos::FindClosestPointAndBaryOnTriangle(TriPos0, TriPos1, TriPos2, ParticlePos, Bary);
+			float CurrentDistance = (ParticlePos - ClosestPoint).Size();
+			float TriRadius = FalloffThreshold * MaxEdgeLength(ComponentSpaceVertices, Sample.Indices.Get(), 0, i, 1);
+			float FalloffScale = CalculateFalloffScale(Falloff, TriRadius, CurrentDistance);
+			if (!FMath::IsNearlyZero(FalloffScale))
 			{
-				Chaos::TVector<float, 3> Bary,
-					TriPos0(ComponentSpaceVertices[Sample.Indices[i][0]]),
-					TriPos1(ComponentSpaceVertices[Sample.Indices[i][1]]),
-					TriPos2(ComponentSpaceVertices[Sample.Indices[i][2]]),
-					ParticlePos(ComponentSpaceTargetVertices[TargetVertexIntersection[j]]);
-
-				Chaos::TVector<float, 3> ClosestPoint = Chaos::FindClosestPointAndBaryOnTriangle(TriPos0, TriPos1, TriPos2, ParticlePos, Bary);
-				float CurrentDistance = (ParticlePos - ClosestPoint).Size();
-				float TriRadius = FalloffThreshold * MaxEdgeLength(ComponentSpaceVertices, Sample.Indices.Get(), 0, i, 1);
-				float FalloffScale = CalculateFalloffScale(Falloff, TriRadius, CurrentDistance);
-				if (!FMath::IsNearlyZero(FalloffScale))
+				int32 TargetIndex = TargetVertexIntersection[j];
+				TArray<int32> TransferIndices;
+				TArray<float> TransferWeights;
+				for (int32 k = 0; k < 3; k++)
 				{
-					int32 TargetIndex = TargetVertexIntersection[j];
-					if (ensure(0 <= TargetIndex && TargetIndex < TargetAttributeArray->Num()))
-					{
-						for (int32 k = 0; k < 3; k++)
-						{
-							(*TargetAttributeArray)[TargetIndex] += FalloffScale * (Bary[k] * (*SourceAttributeArray)[Sample.Indices[i][k]]);
-						}
-						break;
-					}
+					TransferIndices.Add(Sample.Indices[i][k]);
+					TransferWeights.Add(Bary[k]);
 				}
+				TransferDataArray[TargetIndex] = UE::Private::FTransferData(TransferIndices, TransferWeights, FalloffScale);
+				break;
 			}
 		}
 	}
+	return TransferDataArray;
 }
 
 
@@ -498,5 +591,19 @@ float FGeometryCollectionTransferVertexAttributeNode::CalculateFalloffScale(EDat
 	return 1.0;
 }
 
+void FGeometryCollectionSetKinematicVertexSelectionNode::Evaluate(UE::Dataflow::FContext& Context, const FDataflowOutput* Out) const
+{
+	if (Out->IsA(&Collection))
+	{
+		FManagedArrayCollection InCollection = GetValue(Context, &Collection);
+		if (IsConnected(&VertexSelection))
+		{
+			FDataflowVertexSelection VertexSelectionIn = GetValue<FDataflowVertexSelection>(Context, &VertexSelection);
+			GeometryCollection::Facades::FVertexBoneWeightsFacade VertexBoneWeightsFacade(InCollection);
+			VertexBoneWeightsFacade.SetVertexArrayKinematic(VertexSelectionIn.AsArray());
+		}
+		SetValue(Context, MoveTemp(InCollection), &Collection);
+	}
+}
 
 #undef LOCTEXT_NAMESPACE
diff --git a/Engine/Plugins/Experimental/GeometryCollectionPlugin/Source/GeometryCollectionNodes/Public/Dataflow/GeometryCollectionTransferVertexAttributeNode.h b/Engine/Plugins/Experimental/GeometryCollectionPlugin/Source/GeometryCollectionNodes/Public/Dataflow/GeometryCollectionTransferVertexAttributeNode.h
--- a/Engine/Plugins/Experimental/GeometryCollectionPlugin/Source/GeometryCollectionNodes/Public/Dataflow/GeometryCollectionTransferVertexAttributeNode.h
+++ b/Engine/Plugins/Experimental/GeometryCollectionPlugin/Source/GeometryCollectionNodes/Public/Dataflow/GeometryCollectionTransferVertexAttributeNode.h
@@ -7,6 +7,7 @@
 #include "Dataflow/DataflowNode.h"
 #include "Dataflow/DataflowConnectionTypes.h"
 #include "GeometryCollection/GeometryCollection.h"
+#include "Dataflow/DataflowSelection.h"
 
 #include "GeometryCollectionTransferVertexAttributeNode.generated.h"
 
@@ -18,8 +19,8 @@ namespace UE_DEPRECATED(5.5, "Use UE::Dafaflow instead.") Dataflow {}
 
 namespace UE::Private 
 {
-
 	class FTransferFacade;
+	struct FTransferData;
 };
 
 
@@ -91,9 +92,14 @@ public:
 	FManagedArrayCollection FromCollection;
 
 	/* The name of the vertex attribute to generate indices from. */
-	UPROPERTY(EditAnywhere, Category = "Dataflow", Meta = (DataflowInput, DataflowOutput, DisplayName = "AttributeKey", DataflowPassthrough = "AttributeKey"))
+	UPROPERTY(EditAnywhere, Category = "Dataflow", Meta = (DataflowInput, DataflowOutput, DisplayName = "AttributeKey", DataflowPassthrough = "AttributeKey",
+		EditCondition = "!bTransferSkinWeights", EditConditionHides))
 	FCollectionAttributeKey AttributeKey = FCollectionAttributeKey(FString(""), FString("Vertices"));
 
+	/* If kinematic skin weight attribute is transferred. */
+	UPROPERTY(EditAnywhere, Category = "Dataflow", Meta = (DisplayName = "TransferSkinWeights"))
+	bool bTransferSkinWeights = false;
+	
 	/* Bounding volume type for source assets[default: Triangle] */
 	UPROPERTY(EditAnywhere, Category = "Thresholds")
 	EDataflowTransferVertexAttributeNodeBoundingVolume BoundingVolumeType = EDataflowTransferVertexAttributeNodeBoundingVolume::Triangle;
@@ -134,7 +140,7 @@ public:
 private:
 
 	virtual void Evaluate(UE::Dataflow::FContext& Context, const FDataflowOutput* Out) const override;
-	
+
 	typedef Chaos::TSphere<Chaos::FReal, 3> SphereType;
 	typedef Chaos::TBoundingVolumeHierarchy<TArray<SphereType*>, TArray<int32>, Chaos::FReal, 3> BVH;
 
@@ -142,12 +148,11 @@ private:
 	// private helpers
 	//
 	TArray<FIntVector2> FindSourceToTargetGeometryMap(const FManagedArrayCollection& SourceCollection, const FManagedArrayCollection& TargetCollection) const;
+	TMap<int32, int32> FindSourceToTargetTransformMap(const FManagedArrayCollection& SourceCollection, const FManagedArrayCollection& TargetCollection) const;
 
-	template<typename T>
-	void PairedGeometryTransfer(FCollectionAttributeKey Key, const TArray<FIntVector2>& PairedGeometry, const UE::Private::FTransferFacade& Sample, UE::Private::FTransferFacade& Target, const TManagedArray<T>* SourceAttributeArray, TManagedArray<T>* TargetFloatArray) const;
-
-	template<typename T>
-	void NearestVertexTransfer(FCollectionAttributeKey Key, const UE::Private::FTransferFacade& Sample, UE::Private::FTransferFacade& Target, const TManagedArray<T>* SourceAttributeArray, TManagedArray<T>* TargetFloatArray) const;
+	TArray<UE::Private::FTransferData> PairedGeometryTransfer(FCollectionAttributeKey Key, const TArray<FIntVector2>& PairedGeometry,
+		const UE::Private::FTransferFacade& Sample, const UE::Private::FTransferFacade& Target) const;
+	TArray<UE::Private::FTransferData> NearestVertexTransfer(FCollectionAttributeKey Key, const UE::Private::FTransferFacade& Sample, const UE::Private::FTransferFacade& Target) const;
 	static float MaxEdgeLength(TArray<FVector3f>& Vert, const TManagedArray<FIntVector3>& Tri, int VertexOFfset, int TriStart, int TriCount);
 	static void BuildComponentSpaceVertices(const TManagedArray<FTransform3f>& LocalSpaceTransform, const TManagedArray<int32>& Parent, const TManagedArray<int32>& BoneMapArray, const TManagedArray<FVector3f>& VertexArray, int32 Start, int32 Count, TArray<FVector3f>& ComponentSpaceVertices);
 	static BVH* BuildParticleSphereBVH(const TArray<FVector3f>& Vertices, float Radius);
@@ -156,4 +161,34 @@ private:
 
 };
 
+/**
+ * Set VertexSelection to be kinematic. Note that kinematic particles need skin weights.
+ */
+USTRUCT(meta = (DataflowGeometryCollection))
+struct FGeometryCollectionSetKinematicVertexSelectionNode : public FDataflowNode
+{
+	GENERATED_USTRUCT_BODY()
+	DATAFLOW_NODE_DEFINE_INTERNAL(FGeometryCollectionSetKinematicVertexSelectionNode, "SetKinematicVertexSelection", "GeometryCollection", "Set Vertex Collection to be kinematic")
+	DATAFLOW_NODE_RENDER_TYPE("SurfaceRender", FGeometryCollection::StaticType(), "Collection")
 
+public:
+
+	UPROPERTY(Meta = (DataflowInput, DataflowOutput, DataflowPassthrough = "Collection"))
+	FManagedArrayCollection Collection;
+
+	/** Vertex Selection set to be kinematic */
+	UPROPERTY(meta = (DataflowInput, DisplayName = "VertexSelection"))
+	FDataflowVertexSelection VertexSelection;
+
+	FGeometryCollectionSetKinematicVertexSelectionNode(const UE::Dataflow::FNodeParameters& InParam, FGuid InGuid = FGuid::NewGuid())
+		: FDataflowNode(InParam, InGuid)
+	{
+		RegisterInputConnection(&Collection);
+		RegisterInputConnection(&VertexSelection);
+		RegisterOutputConnection(&Collection, &Collection);
+	}
+
+private:
+
+	virtual void Evaluate(UE::Dataflow::FContext& Context, const FDataflowOutput* Out) const override;
+};
\ No newline at end of file
diff --git a/Engine/Source/Runtime/Experimental/Chaos/Private/Chaos/Deformable/ChaosDeformableSolver.cpp b/Engine/Source/Runtime/Experimental/Chaos/Private/Chaos/Deformable/ChaosDeformableSolver.cpp
--- a/Engine/Source/Runtime/Experimental/Chaos/Private/Chaos/Deformable/ChaosDeformableSolver.cpp
+++ b/Engine/Source/Runtime/Experimental/Chaos/Private/Chaos/Deformable/ChaosDeformableSolver.cpp
@@ -88,6 +88,7 @@ namespace Chaos::Softs
 	FAutoConsoleVariableRef CVarDeformableDebugParamsDrawKinematicParticles(TEXT("p.Chaos.DebugDraw.Deformable.KinematicParticle"), GDeformableDebugParams.bDoDrawKinematicParticles, TEXT("Debug draw the deformables kinematic particles. [def: false]"));
 	FAutoConsoleVariableRef CVarDeformableDebugParamsDrawTransientKinematicParticles(TEXT("p.Chaos.DebugDraw.Deformable.TransientKinematicParticle"), GDeformableDebugParams.bDoDrawTransientKinematicParticles, TEXT("Debug draw the deformables transient kinematic particles. [def: false]"));
 	FAutoConsoleVariableRef CVarDeformableDebugParamsDrawRigidCollisionGeometry(TEXT("p.Chaos.DebugDraw.Deformable.RigidCollisionGeometry"), GDeformableDebugParams.bDoDrawRigidCollisionGeometry, TEXT("Debug draw the deformable solvers rigid collision geometry. [def: false]"));
+	FAutoConsoleVariableRef CVarDeformableDebugParamsDrawParticleRadius(TEXT("p.Chaos.DebugDraw.Deformable.ParticleRadius"), GDeformableDebugParams.ParticleRadius, TEXT("Drawn kinematic particle radius. [def: 5]"));
 
 	FDeformableXPBDCorotatedParams GDeformableXPBDCorotatedParams;
 	FAutoConsoleVariableRef CVarDeformableXPBDCorotatedBatchSize(TEXT("p.Chaos.Deformable.XPBDBatchSize"), GDeformableXPBDCorotatedParams.XPBDCorotatedBatchSize, TEXT("Batch size for physics parallel for. [def: 5]"));
@@ -285,7 +286,7 @@ namespace Chaos::Softs
 								{
 									int32 LocalIndex = CnstrTargets.GetIndex(i);
 									int32 ParticleIndex = Range[0] + LocalIndex;
-									Chaos::FDebugDrawQueue::GetInstance().DrawDebugPoint(ToDouble(Evolution->Particles().GetX(ParticleIndex)), FColor::Orange, false, -1.0f, 0, 5);
+									Chaos::FDebugDrawQueue::GetInstance().DrawDebugPoint(ToDouble(Evolution->Particles().GetX(ParticleIndex)), FColor::Orange, false, -1.0f, 0, GDeformableDebugParams.ParticleRadius);
 								}
 							}
 #endif
@@ -559,22 +560,14 @@ namespace Chaos::Softs
 
 		if (Property.bEnableKinematics)
 		{
+			GeometryCollection::Facades::FVertexBoneWeightsFacade VertexBoneWeightsFacade(Rest);
 			typedef GeometryCollection::Facades::FKinematicBindingFacade FKinematics;
 			FKinematics Kinematics(Rest);
-
-			// Add Kinematics Node
-			for (int i = Kinematics.NumKinematicBindings() - 1; i >= 0; i--)
+			for (int32 VertexIdx = 0; VertexIdx < VertexBoneWeightsFacade.NumVertices(); ++VertexIdx)
 			{
-				FKinematics::FBindingKey Key = Kinematics.GetKinematicBindingKey(i);
-
-				int32 BoneIndex = INDEX_NONE;
-				TArray<int32> BoundVerts;
-				TArray<float> BoundWeights;
-				Kinematics.GetBoneBindings(Key, BoneIndex, BoundVerts, BoundWeights);
-
-				for (int32 vdx : BoundVerts)
+				if (VertexBoneWeightsFacade.IsKinematicVertex(VertexIdx))
 				{
-					int32 ParticleIndex = Range[0] + vdx;
+					int32 ParticleIndex = Range[0] + VertexIdx;
 					Evolution->Particles().InvM(ParticleIndex) = 0.f;
 					Evolution->Particles().PAndInvM(ParticleIndex).InvM = 0.f;
 				}
@@ -833,7 +826,7 @@ namespace Chaos::Softs
 			{
 				auto T = ChaosTet(Tetrahedron[edx], Range[0]);
 				Chaos::FDebugDrawQueue::GetInstance().DrawDebugPoint(
-					DoubleVert(P.GetX(T[0])), FColor::Blue, false, -1.0f, 0, 5);
+					DoubleVert(P.GetX(T[0])), FColor::Blue, false, -1.0f, 0, GDeformableDebugParams.ParticleRadius);
 			}
 		}
 #endif
@@ -1340,7 +1333,7 @@ namespace Chaos::Softs
 						if (GDeformableDebugParams.IsDebugDrawingEnabled() && GDeformableDebugParams.bDoDrawKinematicParticles )
 						{
 							auto DoubleVert = [](FVector3f V) { return FVector3d(V.X, V.Y, V.Z); };
-							Chaos::FDebugDrawQueue::GetInstance().DrawDebugPoint(DoubleVert(MParticles.GetX(Index)), FColor::Red, false, -1.0f, 0, 5);
+							Chaos::FDebugDrawQueue::GetInstance().DrawDebugPoint(DoubleVert(MParticles.GetX(Index)), FColor::Red, false, -1.0f, 0, GDeformableDebugParams.ParticleRadius);
 						}
 #endif
 
diff --git a/Engine/Source/Runtime/Experimental/Chaos/Private/GeometryCollection/Facades/CollectionVertexBoneWeightsFacade.cpp b/Engine/Source/Runtime/Experimental/Chaos/Private/GeometryCollection/Facades/CollectionVertexBoneWeightsFacade.cpp
--- a/Engine/Source/Runtime/Experimental/Chaos/Private/GeometryCollection/Facades/CollectionVertexBoneWeightsFacade.cpp
+++ b/Engine/Source/Runtime/Experimental/Chaos/Private/GeometryCollection/Facades/CollectionVertexBoneWeightsFacade.cpp
@@ -14,12 +14,14 @@ namespace GeometryCollection::Facades
 	// Attributes
 	const FName FVertexBoneWeightsFacade::BoneWeightAttributeName = "BoneWeights";
 	const FName FVertexBoneWeightsFacade::BoneIndexAttributeName = "BoneWeightsIndex";
-
+	const FName FVertexBoneWeightsFacade::KinematicAttributeName = "Kinematic";
+	
 	FVertexBoneWeightsFacade::FVertexBoneWeightsFacade(FManagedArrayCollection& InCollection)
 		: ConstCollection(InCollection)
 		, Collection(&InCollection)
 		, BoneIndexAttribute(InCollection, BoneIndexAttributeName, FGeometryCollection::VerticesGroup, FTransformCollection::TransformGroup)
 		, BoneWeightAttribute(InCollection, BoneWeightAttributeName, FGeometryCollection::VerticesGroup, FTransformCollection::TransformGroup)
+		, KinematicAttribute(InCollection, KinematicAttributeName, FGeometryCollection::VerticesGroup)
 		, ParentAttribute(InCollection, FTransformCollection::ParentAttribute, FTransformCollection::TransformGroup)
 		, VerticesAttribute(InCollection, "Vertex", FGeometryCollection::VerticesGroup)
 	{
@@ -31,6 +33,7 @@ namespace GeometryCollection::Facades
 		, Collection(nullptr)
 		, BoneIndexAttribute(InCollection, BoneIndexAttributeName, FGeometryCollection::VerticesGroup, FTransformCollection::TransformGroup)
 		, BoneWeightAttribute(InCollection, BoneWeightAttributeName, FGeometryCollection::VerticesGroup, FTransformCollection::TransformGroup)
+		, KinematicAttribute(InCollection, KinematicAttributeName, FGeometryCollection::VerticesGroup)
 		, ParentAttribute(InCollection, FTransformCollection::ParentAttribute, FTransformCollection::TransformGroup)
 		, VerticesAttribute(InCollection, "Vertex", FGeometryCollection::VerticesGroup)
 	{
@@ -46,12 +49,13 @@ namespace GeometryCollection::Facades
 		check(!IsConst());
 		BoneIndexAttribute.Add();
 		BoneWeightAttribute.Add();
+		KinematicAttribute.AddAndFill(false);
 		ParentAttribute.Add();
 	}
 
 	bool FVertexBoneWeightsFacade::IsValid() const
 	{
-		return BoneIndexAttribute.IsValid() && BoneWeightAttribute.IsValid() && ParentAttribute.IsValid() && VerticesAttribute.IsValid();
+		return BoneIndexAttribute.IsValid() && BoneWeightAttribute.IsValid() && KinematicAttribute.IsValid() && ParentAttribute.IsValid() && VerticesAttribute.IsValid();
 	}
 
 
@@ -98,6 +102,27 @@ namespace GeometryCollection::Facades
 			}
 		}
 	}
+
+	void FVertexBoneWeightsFacade::SetVertexKinematic(int32 VertexIndex)
+	{
+		if (KinematicAttribute.IsValidIndex(VertexIndex))
+		{
+			KinematicAttribute.ModifyAt(VertexIndex, true);
+		}
+	}
+
+	void FVertexBoneWeightsFacade::SetVertexArrayKinematic(const TArray<int32>& VertexIndices)
+	{
+		TManagedArray<bool>& KinematicArray = KinematicAttribute.Modify();
+		for (const int32& VertexIndex : VertexIndices)
+		{
+			if (KinematicArray.IsValidIndex(VertexIndex))
+			{
+				KinematicArray[VertexIndex] = true;
+			}
+		}
+	}
+
 	//
 	//  Add Weights from Selection 
 	//
@@ -142,6 +167,7 @@ namespace GeometryCollection::Facades
 						int32 Vert = OutBoneVerts[Vdx]; float Weight = OutBoneWeights[Vdx];
 						if (0 <= Vert && Vert < NumVertices && !IndicesArray[Vert].Contains(Bone))
 						{
+							SetVertexKinematic(Vert);
 							int32 BoneIndex = IndicesArray[Vert].Find(Bone);
 							if (TotalWeights[Vert] + Weight <= 1.f + UE_KINDA_SMALL_NUMBER)
 							{
diff --git a/Engine/Source/Runtime/Experimental/Chaos/Public/Chaos/Deformable/ChaosDeformableSolverTypes.h b/Engine/Source/Runtime/Experimental/Chaos/Public/Chaos/Deformable/ChaosDeformableSolverTypes.h
--- a/Engine/Source/Runtime/Experimental/Chaos/Public/Chaos/Deformable/ChaosDeformableSolverTypes.h
+++ b/Engine/Source/Runtime/Experimental/Chaos/Public/Chaos/Deformable/ChaosDeformableSolverTypes.h
@@ -169,7 +169,7 @@ namespace Chaos::Softs
 		bool bDoDrawKinematicParticles = false;
 		bool bDoDrawTransientKinematicParticles = false;
 		bool bDoDrawRigidCollisionGeometry = false;
-
+		FSolverReal ParticleRadius = 5.f;
 
 		bool IsDebugDrawingEnabled()
 		{ 
diff --git a/Engine/Source/Runtime/Experimental/Chaos/Public/GeometryCollection/Facades/CollectionVertexBoneWeightsFacade.h b/Engine/Source/Runtime/Experimental/Chaos/Public/GeometryCollection/Facades/CollectionVertexBoneWeightsFacade.h
--- a/Engine/Source/Runtime/Experimental/Chaos/Public/GeometryCollection/Facades/CollectionVertexBoneWeightsFacade.h
+++ b/Engine/Source/Runtime/Experimental/Chaos/Public/GeometryCollection/Facades/CollectionVertexBoneWeightsFacade.h
@@ -14,15 +14,13 @@ namespace GeometryCollection::Facades
 	* FVertexBoneWeightsFacade
 	* 
 	* Defines common API for storing a vertex weights bound to a bone. This mapping is from the 
-	* the vertex to the bone index. The FSelectionFacad will store the mapping from the BoneIndex
-	* to the vertex. 
-	* Usage:
-	*    FVertexBoneWeightsFacade::AddBoneWeights(this, FSelectionFacade(this) );
+	* the vertex to the bone index. Kinematic array specifies whether vertices are considered kinematic. 
+	* Non-kinematic vertices can also have associated bone indices and weights.
 	* 
 	* Then arrays can be accessed later by:
 	*	const TManagedArray< TArray<int32> >* BoneIndices = FVertexBoneWeightsFacade::GetBoneIndices(this);
 	*	const TManagedArray< TArray<float> >* BoneWeights = FVertexBoneWeightsFacade::GetBoneWeights(this);
-	*
+	* 
 	* The following attributes are created on the collection:
 	* 
 	*	- FindAttribute<TArray<int32>>(FVertexSetInterface::IndexAttribute, FGeometryCollection::VerticesGroup);
@@ -36,7 +34,7 @@ namespace GeometryCollection::Facades
 		// Attributes
 		static CHAOS_API const FName BoneIndexAttributeName;
 		static CHAOS_API const FName BoneWeightAttributeName;
-
+		static CHAOS_API const FName KinematicAttributeName;
 		/**
 		* FVertexBoneWeightsFacade Constuctor
 		*/
@@ -61,12 +59,19 @@ namespace GeometryCollection::Facades
 		/** Modify bone weight based on the kinematic bindings. */
 		CHAOS_API void ModifyBoneWeight(int32 VertexIndex, TArray<int32> VertexBoneIndex, TArray<float> VertexBoneWeight);
 
+		/** Set vertex to be kinematic */
+		CHAOS_API void SetVertexKinematic(int32 VertexIndex);
+		CHAOS_API void SetVertexArrayKinematic(const TArray<int32>& VertexIndices);
+
 		/** Return the vertex bone indices from the collection. Null if not initialized.  */
 		const TManagedArray< TArray<int32> >* FindBoneIndices()  const { return BoneIndexAttribute.Find(); }
 		const TManagedArray< TArray<int32> >& GetBoneIndices() const { return BoneIndexAttribute.Get(); }
 
 		/** Return if the vertex is kinematic */
-		CHAOS_API bool IsKinematicVertex(int32 VertexIndex) const { return BoneIndexAttribute.IsValidIndex(VertexIndex) && BoneIndexAttribute[VertexIndex].Num(); };
+		CHAOS_API bool IsKinematicVertex(int32 VertexIndex) const { return KinematicAttribute.IsValidIndex(VertexIndex) && KinematicAttribute.Get()[VertexIndex] && BoneIndexAttribute.Get()[VertexIndex].Num();};
+
+		/** Return number of vertices */
+		CHAOS_API int32 NumVertices() const { return VerticesAttribute.Num(); };
 
 		/** Return the vertex bone weights from the collection. Null if not initialized. */
 		const TManagedArray< TArray<float> >* FindBoneWeights()  const { return BoneWeightAttribute.Find(); }
@@ -78,6 +83,7 @@ namespace GeometryCollection::Facades
 
 		TManagedArrayAccessor<TArray<int32>> BoneIndexAttribute;
 		TManagedArrayAccessor<TArray<float>> BoneWeightAttribute;
+		TManagedArrayAccessor<bool> KinematicAttribute;
 		TManagedArrayAccessor<int32> ParentAttribute;
 		TManagedArrayAccessor<FVector3f> VerticesAttribute;
 
diff --git a/Engine/Source/Runtime/Experimental/GeometryCollectionEngine/Private/GeometryCollection/GeometryCollectionEngineConversion.cpp b/Engine/Source/Runtime/Experimental/GeometryCollectionEngine/Private/GeometryCollection/GeometryCollectionEngineConversion.cpp
--- a/Engine/Source/Runtime/Experimental/GeometryCollectionEngine/Private/GeometryCollection/GeometryCollectionEngineConversion.cpp
+++ b/Engine/Source/Runtime/Experimental/GeometryCollectionEngine/Private/GeometryCollection/GeometryCollectionEngineConversion.cpp
@@ -39,7 +39,7 @@
 #include "StaticMeshOperations.h"
 #include "VertexConnectedComponents.h"
 #include "Util/ColorConstants.h"
-
+#include "GeometryCollection/Facades/CollectionVertexBoneWeightsFacade.h"
 
 DEFINE_LOG_CATEGORY_STATIC(UGeometryCollectionConversionLogging, Log, All);
 
@@ -1213,6 +1213,10 @@ bool FGeometryCollectionEngineConversion::AppendSkeletalMesh(const USkeletalMesh
 				RootIndex = TransformBaseIndex + BoneIndex;
 			}
 		}
+		GeometryCollection::Facades::FTransformSource TransformSourceFacade(*InManagedArrayCollection);
+		TSet<int32> Roots;
+		Roots.Add(RootIndex);
+		TransformSourceFacade.AddTransformSource(Skeleton->GetName(), Skeleton->GetGuid().ToString(), Roots);
 	}
 
 
@@ -1226,34 +1230,6 @@ bool FGeometryCollectionEngineConversion::AppendSkeletalMesh(const USkeletalMesh
 	GeometryCollectionEngineUtility::GenerateConnectedComponents(InSkeletalMesh, ComponentsSourceIndices,
 		SourceToTargetTriangleMap, SourceVertexToComponentMap, TriangleCount, VertexCount);
 
-
-	//
-	// Add the Triangles to the Geometry Collection
-	//.. ensure all component vertices are contigious in the array 
-	//
-	int NumVertices = GeometryCollection.NumElements(FGeometryCollection::VerticesGroup);
-	int IndicesBaseIndex = GeometryCollection.AddElements(TriangleCount, FGeometryCollection::FacesGroup);
-
-	int CurrentIndex = IndicesBaseIndex;
-	for (int ComponentIndex = 0; ComponentIndex < ComponentsSourceIndices.Num(); ComponentIndex++)
-	{
-		for(int32 TriangleIndex = 0; TriangleIndex< ComponentsSourceIndices[ComponentIndex].Num(); TriangleIndex++)
-		{
-			SourceToTargetTriangleMap[ComponentIndex][TriangleIndex][1] = CurrentIndex;
-
-			FIntVector& Triangle = ComponentsSourceIndices[ComponentIndex][TriangleIndex];
-			for (int k = 0; k < 3; k++)
-			{
-				Indices[CurrentIndex][k] = SourceVertexToComponentMap[Triangle[k]] + IndicesBaseIndex;
-			}
-
-			Visible[CurrentIndex] = true;
-			MaterialID[CurrentIndex] = 0;
-			MaterialIndex[CurrentIndex] = CurrentIndex;
-			CurrentIndex++;
-		}
-	}
-
 	//
 	// Vertex Attributes
 	//
@@ -1370,6 +1346,31 @@ bool FGeometryCollectionEngineConversion::AppendSkeletalMesh(const USkeletalMesh
 		}
 	}
 
+	//
+	// Add the Triangles to the Geometry Collection
+	//.. ensure all component vertices are contigious in the array 
+	//
+	int IndicesBaseIndex = GeometryCollection.AddElements(TriangleCount, FGeometryCollection::FacesGroup);
+
+	int CurrentIndex = IndicesBaseIndex;
+	for (int ComponentIndex = 0; ComponentIndex < ComponentsSourceIndices.Num(); ComponentIndex++)
+	{
+		for (int32 TriangleIndex = 0; TriangleIndex < ComponentsSourceIndices[ComponentIndex].Num(); TriangleIndex++)
+		{
+			SourceToTargetTriangleMap[ComponentIndex][TriangleIndex][1] = CurrentIndex;
+
+			FIntVector& Triangle = ComponentsSourceIndices[ComponentIndex][TriangleIndex];
+			for (int k = 0; k < 3; k++)
+			{
+				Indices[CurrentIndex][k] = SourceVertexToComponentMap[Triangle[k]] + VertexBaseIndex;
+			}
+
+			Visible[CurrentIndex] = true;
+			MaterialID[CurrentIndex] = 0;
+			MaterialIndex[CurrentIndex] = CurrentIndex;
+			CurrentIndex++;
+		}
+	}
 
 	// Geometry Group
 	TArray<int32> GeometryIndices;
@@ -1379,6 +1380,39 @@ bool FGeometryCollectionEngineConversion::AppendSkeletalMesh(const USkeletalMesh
 	GeometryCollection.RemoveElements(FGeometryCollection::GeometryGroup, GeometryIndices);
 	::GeometryCollection::AddGeometryProperties(&GeometryCollection);
 
+	//Extract skin weights
+	int32 NumVertices = GeometryCollection.NumElements(FGeometryCollection::VerticesGroup);
+	FSkeletalMeshRenderData* RenderData = InSkeletalMesh->GetResourceForRendering();
+	if (RenderData->LODRenderData.Num())
+	{
+		FSkeletalMeshLODRenderData* LODRenderData = &RenderData->LODRenderData[LODIndex];
+		const FSkinWeightVertexBuffer* SkinWeightVertexBuffer = LODRenderData->GetSkinWeightVertexBuffer();
+
+		TArray<int32> BlendedSkeletalMeshVertices;
+		TArray<TArray<float>> Weights;
+		Weights.SetNum(SkinWeightVertexBuffer->GetNumVertices());
+		FSkeletalMeshAttributes MeshAttribs(MeshDescription);
+		FSkinWeightsVertexAttributesRef VertexSkinWeights = MeshAttribs.GetVertexSkinWeights();
+		if (ensure(NumVertices == MeshDescription.Vertices().Num()))
+		{
+			GeometryCollection::Facades::FVertexBoneWeightsFacade VertexBoneWeightsFacade(GeometryCollection);
+			for (int32 VertexIndex = 0; VertexIndex < NumVertices; ++VertexIndex)
+			{
+				FVertexBoneWeights BoneWeights = VertexSkinWeights.Get(FVertexID(VertexIndex));
+				const int32 InfluenceCount = BoneWeights.Num();
+				TArray<int32> VertexBoneIndex;
+				TArray<float> VertexBoneWeight;
+				VertexBoneIndex.SetNum(InfluenceCount);
+				VertexBoneWeight.SetNum(InfluenceCount);
+				for (int32 InfluenceIndex = 0; InfluenceIndex < InfluenceCount; ++InfluenceIndex)
+				{
+					VertexBoneIndex[InfluenceIndex] = BoneWeights[InfluenceIndex].GetBoneIndex();
+					VertexBoneWeight[InfluenceIndex] = BoneWeights[InfluenceIndex].GetWeight();
+				}
+				VertexBoneWeightsFacade.ModifyBoneWeight(SourceVertexToComponentMap[VertexIndex], VertexBoneIndex, VertexBoneWeight);
+			}
+		}
+	}
 	GeometryCollection.CopyTo(InManagedArrayCollection);
 	return true;
 #else
