ChaosFlesh: fixed USkeleton refskeleton inconsistency, deprecated SkeletonToCollection with SkeletalMeshToCollection, added TransferVertexSkinWeights node

diff --git a/Engine/Plugins/Experimental/ChaosFlesh/Source/ChaosFleshNodes/Private/Dataflow/ChaosFleshKinematicMuscleAttachmentsNode.cpp b/Engine/Plugins/Experimental/ChaosFlesh/Source/ChaosFleshNodes/Private/Dataflow/ChaosFleshKinematicMuscleAttachmentsNode.cpp
--- a/Engine/Plugins/Experimental/ChaosFlesh/Source/ChaosFleshNodes/Private/Dataflow/ChaosFleshKinematicMuscleAttachmentsNode.cpp
+++ b/Engine/Plugins/Experimental/ChaosFlesh/Source/ChaosFleshNodes/Private/Dataflow/ChaosFleshKinematicMuscleAttachmentsNode.cpp
@@ -5,6 +5,7 @@
 #include "Engine/SkeletalMesh.h"
 #include "GeometryCollection/Facades/CollectionKinematicBindingFacade.h"
 #include "GeometryCollection/Facades/CollectionVertexBoneWeightsFacade.h"
+#include "GeometryCollection/Facades/CollectionTransformFacade.h"
 #include "Rendering/SkeletalMeshRenderData.h"
 #include "BoneWeights.h"
 
@@ -60,6 +61,8 @@ void FKinematicMuscleAttachmentsDataflowNode::Evaluate(UE::Dataflow::FContext& C
 						{
 							return;
 						}
+						GeometryCollection::Facades::FCollectionTransformFacade TransformFacade(InCollection);
+						const TMap<FString, int32> BoneNameToIndex = TransformFacade.BoneNameIndexMap();
 						auto DoubleVert = [](FVector3f V) { return FVector3d(V.X, V.Y, V.Z); };
 						for (int32 i = 0; i < BoundVerts.Num(); i++)
 						{
@@ -102,9 +105,9 @@ void FKinematicMuscleAttachmentsDataflowNode::Evaluate(UE::Dataflow::FContext& C
 								//get local coords of bound verts
 								typedef GeometryCollection::Facades::FKinematicBindingFacade FKinematics;
 								FKinematics Kinematics(InCollection); Kinematics.DefineSchema();
-								if (Kinematics.IsValid())
+								if (Kinematics.IsValid() && BoneNameToIndex.Contains(BoneName))
 								{
-									FKinematics::FBindingKey Binding = Kinematics.SetBoneBindings(BoneIndex, BoneBoundVerts[BoneIndex], BoneBoundWeights[BoneIndex]);
+									FKinematics::FBindingKey Binding = Kinematics.SetBoneBindings(BoneNameToIndex[BoneName], BoneBoundVerts[BoneIndex], BoneBoundWeights[BoneIndex]);
 									TManagedArray<TArray<FVector3f>>& LocalPos = InCollection.AddAttribute<TArray<FVector3f>>("LocalPosition", Binding.GroupName);
 									Kinematics.AddKinematicBinding(Binding);
 
diff --git a/Engine/Plugins/Experimental/ChaosFlesh/Source/ChaosFleshNodes/Private/Dataflow/ChaosFleshSetFleshBonePositionTargetBindingNode.cpp b/Engine/Plugins/Experimental/ChaosFlesh/Source/ChaosFleshNodes/Private/Dataflow/ChaosFleshSetFleshBonePositionTargetBindingNode.cpp
--- a/Engine/Plugins/Experimental/ChaosFlesh/Source/ChaosFleshNodes/Private/Dataflow/ChaosFleshSetFleshBonePositionTargetBindingNode.cpp
+++ b/Engine/Plugins/Experimental/ChaosFlesh/Source/ChaosFleshNodes/Private/Dataflow/ChaosFleshSetFleshBonePositionTargetBindingNode.cpp
@@ -8,6 +8,7 @@
 #include "GeometryCollection/Facades/CollectionKinematicBindingFacade.h"
 #include "GeometryCollection/Facades/CollectionPositionTargetFacade.h"
 #include "GeometryCollection/Facades/CollectionVertexBoneWeightsFacade.h"
+#include "GeometryCollection/Facades/CollectionTransformFacade.h"
 #include "Rendering/SkeletalMeshModel.h"
 #include "Rendering/SkeletalMeshRenderData.h"
 #include "BoneWeights.h"
@@ -28,13 +29,10 @@ void FSetFleshBonePositionTargetBindingDataflowNode::Evaluate(UE::Dataflow::FCon
 			const TManagedArray<FVector3f>* Vertices = InCollection.FindAttribute<FVector3f>("Vertex", FGeometryCollection::VerticesGroup);
 			const TManagedArray<FTransform3f>* Transform = InCollection.FindAttribute<FTransform3f>("Transform", FTransformCollection::TransformGroup);
 			const TManagedArray<FString>* TransformBoneName = InCollection.FindAttribute<FString>("BoneName", FTransformCollection::TransformGroup);
+			GeometryCollection::Facades::FCollectionTransformFacade TransformFacade(InCollection);
 			if (Indices && Vertices && Transform && TransformBoneName)
 			{
-				TMap<FString, int32> BoneNameIndexMap;
-				for (int32 BoneIndex = 0; BoneIndex < TransformBoneName->Num(); ++BoneIndex)
-				{
-					BoneNameIndexMap.Add((*TransformBoneName)[BoneIndex], BoneIndex);
-				}
+				const TMap<FString, int32> BoneNameIndexMap = TransformFacade.BoneNameIndexMap();
 				FSkeletalMeshRenderData* RenderData = BoneSkeletalMesh->GetResourceForRendering();
 				if (RenderData->LODRenderData.Num())
 				{
diff --git a/Engine/Plugins/Experimental/GeometryCollectionPlugin/Source/GeometryCollectionNodes/Private/Dataflow/GeometryCollectionNodesPlugin.cpp b/Engine/Plugins/Experimental/GeometryCollectionPlugin/Source/GeometryCollectionNodes/Private/Dataflow/GeometryCollectionNodesPlugin.cpp
--- a/Engine/Plugins/Experimental/GeometryCollectionPlugin/Source/GeometryCollectionNodes/Private/Dataflow/GeometryCollectionNodesPlugin.cpp
+++ b/Engine/Plugins/Experimental/GeometryCollectionPlugin/Source/GeometryCollectionNodes/Private/Dataflow/GeometryCollectionNodesPlugin.cpp
@@ -64,12 +64,16 @@ void IGeometryCollectionNodesPlugin::StartupModule()
 	DATAFLOW_NODE_REGISTER_CREATION_FACTORY(FSetVertexColorFromFloatArrayDataflowNode);
 	DATAFLOW_NODE_REGISTER_CREATION_FACTORY(FSetVertexColorFromVertexIndicesDataflowNode);
 	DATAFLOW_NODE_REGISTER_CREATION_FACTORY(FSetVertexColorFromVertexSelectionDataflowNode);
-	DATAFLOW_NODE_REGISTER_CREATION_FACTORY(FSkeletonToCollectionDataflowNode);
 	DATAFLOW_NODE_REGISTER_CREATION_FACTORY(FSkeletalMeshToCollectionDataflowNode);
 	DATAFLOW_NODE_REGISTER_CREATION_FACTORY(FGeometryCollectionTransferVertexAttributeNode);
+	DATAFLOW_NODE_REGISTER_CREATION_FACTORY(FGeometryCollectionTransferVertexSkinWeightsNode);
 	DATAFLOW_NODE_REGISTER_CREATION_FACTORY(FGeometryCollectionSetKinematicVertexSelectionNode);
 	DATAFLOW_NODE_REGISTER_CREATION_FACTORY(FStaticMeshToCollectionDataflowNode);
 	DATAFLOW_NODE_REGISTER_CREATION_FACTORY(FTriangleBoundaryIndicesNode);
+	// Deprecated nodes
+PRAGMA_DISABLE_DEPRECATION_WARNINGS
+	DATAFLOW_NODE_REGISTER_CREATION_FACTORY(FSkeletonToCollectionDataflowNode);
+PRAGMA_ENABLE_DEPRECATION_WARNINGS
 }
 
 void IGeometryCollectionNodesPlugin::ShutdownModule()
diff --git a/Engine/Plugins/Experimental/GeometryCollectionPlugin/Source/GeometryCollectionNodes/Private/Dataflow/GeometryCollectionSkeletalMeshToCollectionNode.cpp b/Engine/Plugins/Experimental/GeometryCollectionPlugin/Source/GeometryCollectionNodes/Private/Dataflow/GeometryCollectionSkeletalMeshToCollectionNode.cpp
--- a/Engine/Plugins/Experimental/GeometryCollectionPlugin/Source/GeometryCollectionNodes/Private/Dataflow/GeometryCollectionSkeletalMeshToCollectionNode.cpp
+++ b/Engine/Plugins/Experimental/GeometryCollectionPlugin/Source/GeometryCollectionNodes/Private/Dataflow/GeometryCollectionSkeletalMeshToCollectionNode.cpp
@@ -25,7 +25,7 @@ void FSkeletalMeshToCollectionDataflowNode::Evaluate(UE::Dataflow::FContext& Con
 		TObjectPtr<const USkeletalMesh> InSkeletalMesh = GetValue<TObjectPtr<const USkeletalMesh>>(Context, &SkeletalMesh);
 		if (InSkeletalMesh)
 		{
-			FGeometryCollectionEngineConversion::AppendSkeletalMesh(InSkeletalMesh, 0, FTransform::Identity, &OutCollection);
+			FGeometryCollectionEngineConversion::AppendSkeletalMesh(InSkeletalMesh, 0, FTransform::Identity, &OutCollection, /*bReindexMaterials = */ true, bImportTransformOnly);
 		}
 		SetValue(Context, FManagedArrayCollection(OutCollection), &Collection);
 	}
diff --git a/Engine/Plugins/Experimental/GeometryCollectionPlugin/Source/GeometryCollectionNodes/Private/Dataflow/GeometryCollectionTransferVertexAttributeNode.cpp b/Engine/Plugins/Experimental/GeometryCollectionPlugin/Source/GeometryCollectionNodes/Private/Dataflow/GeometryCollectionTransferVertexAttributeNode.cpp
--- a/Engine/Plugins/Experimental/GeometryCollectionPlugin/Source/GeometryCollectionNodes/Private/Dataflow/GeometryCollectionTransferVertexAttributeNode.cpp
+++ b/Engine/Plugins/Experimental/GeometryCollectionPlugin/Source/GeometryCollectionNodes/Private/Dataflow/GeometryCollectionTransferVertexAttributeNode.cpp
@@ -14,8 +14,10 @@
 #include UE_INLINE_GENERATED_CPP_BY_NAME(GeometryCollectionTransferVertexAttributeNode)
 #define LOCTEXT_NAMESPACE "FGeometryCollectionTransferVertexAttributeNode"
 
-namespace UE::Private {
-
+namespace UE::Private
+{
+	typedef Chaos::TSphere<Chaos::FReal, 3> SphereType;
+	typedef Chaos::TBoundingVolumeHierarchy<TArray<SphereType*>, TArray<int32>, Chaos::FReal, 3> BVH;
 	class FTransferFacade
 	{
 		const FManagedArrayCollection& ConstCollection;
@@ -88,509 +90,566 @@ namespace UE::Private {
 	{
 		TArray<int32> SourceIndices;
 		TArray<float> SourceWeights;
-		float FalloffScale;
-		FTransferData(): FalloffScale(0.f) {}
+		float FalloffScale = 0.f;
+		FTransferData() {}
 		FTransferData(TArray<int32> SourceIndicesIn, TArray<float> SourceWeightsIn, float FalloffScaleIn) :
 			SourceIndices(SourceIndicesIn), SourceWeights(SourceWeightsIn), FalloffScale(FalloffScaleIn) {}
 	};
-}
-
-void FGeometryCollectionTransferVertexAttributeNode::Evaluate(UE::Dataflow::FContext& Context, const FDataflowOutput* Out) const
-{
-	FCollectionAttributeKey Key = GetValue(Context, &AttributeKey, AttributeKey);
 
-	if (Out->IsA(&Collection))
+	struct FTransferProperties
 	{
-		FManagedArrayCollection TargetCollection = GetValue(Context, &Collection);
-		const FManagedArrayCollection& SampleCollection = GetValue(Context, &FromCollection);
-
-		UE::Private::FTransferFacade Target(TargetCollection);
-		const UE::Private::FTransferFacade Sample(SampleCollection);
+		EDataflowTransferVertexAttributeNodeBoundingVolume BoundingVolumeType = EDataflowTransferVertexAttributeNodeBoundingVolume::Dataflow_Max;
+		EDataflowTransferVertexAttributeNodeSourceScale SourceScale = EDataflowTransferVertexAttributeNodeSourceScale::Dataflow_Max;
+		EDataflowTransferVertexAttributeNodeFalloff Falloff = EDataflowTransferVertexAttributeNodeFalloff::Dataflow_Max;
+		float FalloffThreshold = 0.f;
+		float EdgeMultiplier = 0.f;
+		float BoundMultiplier = 0.f;
+	};
 
-		if (Target.IsValid() && Sample.IsValid())
-		{
-			TArray<FIntVector2> AlignedGeometry = FindSourceToTargetGeometryMap(SampleCollection, TargetCollection);
-			TArray<UE::Private::FTransferData> TransferDataArray; //store barycentric weight info
-			if (AlignedGeometry.Num() == SampleCollection.NumElements(FGeometryCollection::GeometryGroup))
+	static float MaxEdgeLength(TArray<FVector3f>& Vert, const TManagedArray<FIntVector3>& Tri, int32 VertexOffset, int32 TriStart, int32 TriCount)
+	{
+		auto TriInRange = [](const FIntVector3& T, int32 Max) {
+			for (int32 k = 0; k < 3; k++)
 			{
-				TransferDataArray = PairedGeometryTransfer(Key, AlignedGeometry, Sample, Target);
+				if (ensure(0 <= T[k] && T[k] < Max))
+				{
+					return true;
+				}
 			}
-			else
+			return false;
+			};
+
+		float Max = 0;
+		int32 TriStop = TriStart + TriCount;
+		for (int32 i = TriStart; i < TriStop; i++)
+		{
+			if (TriInRange(Tri[i] - FIntVector3(VertexOffset), Vert.Num()))
 			{
-				TransferDataArray = NearestVertexTransfer(Key, Sample, Target);
+				Max = FMath::Max(Max, (Vert[Tri[i][0] - VertexOffset] - Vert[Tri[i][1] - VertexOffset]).SquaredLength());
+				Max = FMath::Max(Max, (Vert[Tri[i][0] - VertexOffset] - Vert[Tri[i][2] - VertexOffset]).SquaredLength());
+				Max = FMath::Max(Max, (Vert[Tri[i][1] - VertexOffset] - Vert[Tri[i][2] - VertexOffset]).SquaredLength());
 			}
-			if (bTransferSkinWeights)
-			{
-				GeometryCollection::Facades::FVertexBoneWeightsFacade SampleVertexBoneWeightsFacade(SampleCollection);
-				GeometryCollection::Facades::FVertexBoneWeightsFacade TargetVertexBoneWeightsFacade(TargetCollection);
-				const TManagedArray<TArray<int32>>* SampleBoneIndices = SampleVertexBoneWeightsFacade.FindBoneIndices();
-				const TManagedArray<TArray<float>>* SampleBoneWeights = SampleVertexBoneWeightsFacade.FindBoneWeights();
-				if (SampleBoneIndices && SampleBoneWeights)
-				{
-					//
-					// Compute the bone index mappings. This allows the transfer operator to retarget weights to the correct skeleton.
-					//
-					TMap<int32, int32> SampleBoneToTargetBone = FindSourceToTargetTransformMap(SampleCollection, TargetCollection);
-					for (int32 VertexIndex = 0; VertexIndex < Target.Vertex.Num(); VertexIndex++)
-					{
-						UE::Private::FTransferData TransferData = TransferDataArray[VertexIndex];
-						if (TransferData.SourceIndices.Num())
-						{
-							TMap<int32, float> BoneWeightBucket;
-							for (int32 SourceLocalIdx = 0; SourceLocalIdx < TransferData.SourceIndices.Num(); ++SourceLocalIdx)
-							{
-								for (int32 BoneLocalIdx = 0; BoneLocalIdx < (*SampleBoneIndices)[TransferData.SourceIndices[SourceLocalIdx]].Num(); ++BoneLocalIdx)
-								{
-									int32 BoneIndex = (*SampleBoneIndices)[TransferData.SourceIndices[SourceLocalIdx]][BoneLocalIdx];
-									float BoneWeight = (*SampleBoneWeights)[TransferData.SourceIndices[SourceLocalIdx]][BoneLocalIdx];
-									if (BoneWeightBucket.Contains(BoneIndex))
-									{
-										BoneWeightBucket[BoneIndex] += TransferData.SourceWeights[SourceLocalIdx] * BoneWeight;
-									}
-									else
-									{
-										BoneWeightBucket.Add(BoneIndex, TransferData.SourceWeights[SourceLocalIdx] * BoneWeight);
-									}
-								}
-							}
+		}
+		return FMath::Sqrt(Max);
+	}
 
-							TArray<int32> VertexBoneIndex;
-							TArray<float> VertexBoneWeight;
-							for (const TPair<int32, float>& Pair : BoneWeightBucket)
-							{
-								VertexBoneIndex.Add(Pair.Key);
-								VertexBoneWeight.Add(Pair.Value);
-							}
-							TargetVertexBoneWeightsFacade.ModifyBoneWeight(VertexIndex, VertexBoneIndex, VertexBoneWeight);
-						}	
-					}
-				}
-			} //end if (bTransferSkinWeights)
-			else if (const TManagedArray<float>* SourceAttributeFloatArray = Sample.GetAttributeArray<float>(Key.Attribute, Key.Group))
+	static void BuildComponentSpaceVertices(const TManagedArray<FTransform3f>& LocalSpaceTransform, const TManagedArray<int32>& Parent, const TManagedArray<int32>& BoneMapArray,
+		const TManagedArray<FVector3f>& VertexArray, int32 Start, int32 Count, TArray<FVector3f>& ComponentSpaceVertices)
+	{
+		TArray<FTransform3f> ComponentTransform;
+		GeometryCollectionAlgo::GlobalMatrices(LocalSpaceTransform, Parent, ComponentTransform);
+
+		ComponentSpaceVertices.SetNumUninitialized(Count);
+		for (int32 i = 0; i < Count; ++i)
+		{
+			int32 j = i + Start;
+			if (0 < BoneMapArray[i] && BoneMapArray[i] < ComponentTransform.Num())
 			{
-				TManagedArray<float>* TargetAttributeArray = Target.GetAttributeArray<float>(Key.Attribute, Key.Group);
-				TargetAttributeArray->Fill(0.f);
-				for (int32 VertexIndex = 0; VertexIndex < TargetAttributeArray->Num(); VertexIndex++)
-				{
-					UE::Private::FTransferData TransferData = TransferDataArray[VertexIndex];
-					if (TransferData.SourceIndices.Num())
-					{
-						for (int32 SourceLocalIdx = 0; SourceLocalIdx < TransferData.SourceIndices.Num(); ++SourceLocalIdx)
-						{
-							(*TargetAttributeArray)[VertexIndex] += (*SourceAttributeFloatArray)[TransferData.SourceIndices[SourceLocalIdx]]
-								* TransferData.SourceWeights[SourceLocalIdx] * TransferData.FalloffScale;
-						}
-					}
-				}
+				ComponentSpaceVertices[i] = ComponentTransform[BoneMapArray[j]].TransformPosition(VertexArray[j]);
 			}
-			else if (const TManagedArray<FLinearColor>* SourceAttributeColorArray = Sample.GetAttributeArray<FLinearColor>(Key.Attribute, Key.Group))
+			else
 			{
-				TManagedArray<FLinearColor>* TargetAttributeArray = Target.GetAttributeArray<FLinearColor>(Key.Attribute, Key.Group);
-				TargetAttributeArray->Fill(FLinearColor(0, 0, 0, 0));
-				for (int32 VertexIndex = 0; VertexIndex < TargetAttributeArray->Num(); VertexIndex++)
-				{
-					UE::Private::FTransferData TransferData = TransferDataArray[VertexIndex];
-					if (TransferData.SourceIndices.Num())
-					{
-						for (int32 SourceLocalIdx = 0; SourceLocalIdx < TransferData.SourceIndices.Num(); ++SourceLocalIdx)
-						{
-							(*TargetAttributeArray)[VertexIndex] += (*SourceAttributeColorArray)[TransferData.SourceIndices[SourceLocalIdx]]
-								* TransferData.SourceWeights[SourceLocalIdx] * TransferData.FalloffScale;
-						}
-					}
-				}
+				ComponentSpaceVertices[i] = VertexArray[j];
 			}
 		}
-
-		SetValue(Context, MoveTemp(TargetCollection), &Collection);
 	}
-	else if (Out->IsA(&AttributeKey))
+
+	static BVH* BuildParticleSphereBVH(const TArray<FVector3f>& Vertices, float Radius)
 	{
-		SetValue(Context, MoveTemp(Key), &AttributeKey);
+		TArray<SphereType*> VertexSpherePtrs;
+		TArray<SphereType> VertexSpheres;
+		VertexSpheres.Init(SphereType(Chaos::FVec3f(0), Radius), Vertices.Num());
+		VertexSpherePtrs.SetNum(Vertices.Num());
+
+		for (int32 i = 0; i < Vertices.Num(); i++)
+		{
+			Chaos::FVec3f SphereCenter(Vertices[i]);
+			SphereType VertexSphere(SphereCenter, Radius);
+			VertexSpheres[i] = SphereType(SphereCenter, Radius);
+			VertexSpherePtrs[i] = &VertexSpheres[i];
+		}
+		return new BVH(VertexSpherePtrs);
 	}
-}
 
-TArray<FIntVector2> FGeometryCollectionTransferVertexAttributeNode::FindSourceToTargetGeometryMap(const FManagedArrayCollection& SourceCollection, const FManagedArrayCollection& TargetCollection) const
-{
-	TArray<FIntVector2> Mapping;
-	const TManagedArray<FString>* SourceName = SourceCollection.FindAttribute<FString>(FName("BoneName"), FTransformCollection::TransformGroup);
-	const TManagedArray<int32>* SourceGeometryGroup = SourceCollection.FindAttribute<int32>(FName("TransformToGeometryIndex"), FTransformCollection::TransformGroup);
-	const TManagedArray<FString>* TargetName = TargetCollection.FindAttribute<FString>(FName("BoneName"), FTransformCollection::TransformGroup);
-	const TManagedArray<int32>* TargetGeometryGroup = TargetCollection.FindAttribute<int32>(FName("TransformToGeometryIndex"), FTransformCollection::TransformGroup);
-	if (SourceName && SourceGeometryGroup && TargetName && TargetGeometryGroup)
+	static void TriangleToVertexIntersections(
+		const BVH& VertexBVH, const TArray<FVector3f>& ComponentSpaceVertices, const FIntVector3& Triangle, TArray<int32>& OutTargetVertexIntersection)
 	{
-		for (int i = 0; i < SourceName->Num(); i++)
+		OutTargetVertexIntersection.Empty();
+
+		TArray<int32> TargetVertexIntersection0 = VertexBVH.FindAllIntersections(ComponentSpaceVertices[Triangle[0]]);
+		TArray<int32> TargetVertexIntersection1 = VertexBVH.FindAllIntersections(ComponentSpaceVertices[Triangle[1]]);
+		TArray<int32> TargetVertexIntersection2 = VertexBVH.FindAllIntersections(ComponentSpaceVertices[Triangle[2]]);
+		TargetVertexIntersection0.Sort();
+		TargetVertexIntersection1.Sort();
+		TargetVertexIntersection2.Sort();
+
+		for (int32 k = 0; k < TargetVertexIntersection0.Num(); k++)
 		{
-			for (int j = 0; j < TargetName->Num(); j++)
+			if (TargetVertexIntersection1.Contains(TargetVertexIntersection0[k])
+				&& TargetVertexIntersection2.Contains(TargetVertexIntersection0[k]))
 			{
-				FString TestName = FString::Printf(TEXT("%s_Tet"), *(*SourceName)[i]);
-				if ((*TargetName)[j].StartsWith(TestName))
-				{
-					Mapping.Add(FIntVector2((*SourceGeometryGroup)[i], (*TargetGeometryGroup)[j]));
-					break;
-				}
+				OutTargetVertexIntersection.Emplace(TargetVertexIntersection0[k]);
 			}
 		}
 	}
 
-	return Mapping;
-}
 
-TMap<int32, int32> FGeometryCollectionTransferVertexAttributeNode::FindSourceToTargetTransformMap(const FManagedArrayCollection& SourceCollection, const FManagedArrayCollection& TargetCollection) const
-{
-	TMap<int32, int32> SourceIndexToTargetIndex;
-	TMap<FString, int32> TargetNameToIndex;
-	const TManagedArray<FString>* SourceName = SourceCollection.FindAttribute<FString>(FName("BoneName"), FTransformCollection::TransformGroup);
-	const TManagedArray<FString>* TargetName = TargetCollection.FindAttribute<FString>(FName("BoneName"), FTransformCollection::TransformGroup);
-	if (SourceName && TargetName)
+	static float CalculateFalloffScale(EDataflowTransferVertexAttributeNodeFalloff FalloffSetting, float Threshold, float Distance)
 	{
-		for (int32 j = 0; j < TargetName->Num(); j++)
+		float Denominator = 1.0;
+		if (Distance > Threshold && !FMath::IsNearlyZero(Threshold))
 		{
-			TargetNameToIndex.Add((*TargetName)[j], j);
+			Denominator = Distance / Threshold;
 		}
-		for (int32 i = 0; i < SourceName->Num(); i++)
+		switch (FalloffSetting)
 		{
-			if (int32* TargetIndex = TargetNameToIndex.Find((*SourceName)[i]))
-			{
-				SourceIndexToTargetIndex.Add(i, *TargetIndex);
-			}
+		case EDataflowTransferVertexAttributeNodeFalloff::Linear:
+			return 1. / Denominator;
+		case EDataflowTransferVertexAttributeNodeFalloff::Squared:
+			return 1. / FMath::Square(Denominator);
 		}
+		return 1.0;
 	}
 
-	return SourceIndexToTargetIndex;
-}
-
-TArray<UE::Private::FTransferData> FGeometryCollectionTransferVertexAttributeNode::PairedGeometryTransfer(FCollectionAttributeKey Key, const TArray<FIntVector2>& PairedGeometry,
-	const UE::Private::FTransferFacade& Sample, const UE::Private::FTransferFacade& Target) const
-{
-	TArray<UE::Private::FTransferData> TransferDataArray;
-	TransferDataArray.SetNum(Target.Vertex.Num());
-	Chaos::FReal SphereFullRadius;
-
-	if (SampleScale == EDataflowTransferVertexAttributeNodeSampleScale::Asset_Edge || SampleScale == EDataflowTransferVertexAttributeNodeSampleScale::Asset_Bound)
+	static TArray<FIntVector2> FindSourceToTargetGeometryMap(const FManagedArrayCollection& SourceCollection, const FManagedArrayCollection& TargetCollection)
 	{
-		// Build component space vertices for TargetCollection
-		TArray<FVector3f> ComponentSpaceFullTargetVertices; // size of num vertices of the geometry entry. 
-		BuildComponentSpaceVertices(Target.Transform.Get(), Target.Parent.Get(), Target.BoneMap.Get(), Target.Vertex.Get(),
-			0, Target.Vertex.Num(), ComponentSpaceFullTargetVertices);
-
-		// Build component space vertices for TargetCollection
-		TArray<FVector3f> ComponentSpaceFullVertices; // size of num vertices of the geometry entry
-		BuildComponentSpaceVertices(Sample.Transform.Get(), Sample.Parent.Get(), Sample.BoneMap.Get(), Sample.Vertex.Get(),
-			0, Sample.Vertex.Num(), ComponentSpaceFullVertices);
-		if (SampleScale == EDataflowTransferVertexAttributeNodeSampleScale::Asset_Edge)
-		{
-			SphereFullRadius = Chaos::FReal(EdgeMultiplier * FMath::Max(
-				MaxEdgeLength(ComponentSpaceFullTargetVertices, Target.Indices.Get(), 0, 0, Target.Indices.Num()),
-				MaxEdgeLength(ComponentSpaceFullVertices, Sample.Indices.Get(), 0, 0, Sample.Indices.Num())));
-		}
-		else if (SampleScale == EDataflowTransferVertexAttributeNodeSampleScale::Asset_Bound)
+		TArray<FIntVector2> Mapping;
+		const TManagedArray<FString>* const SourceName = SourceCollection.FindAttribute<FString>(FName("BoneName"), FTransformCollection::TransformGroup);
+		const TManagedArray<int32>* const SourceGeometryGroup = SourceCollection.FindAttribute<int32>(FName("TransformToGeometryIndex"), FTransformCollection::TransformGroup);
+		const TManagedArray<FString>* const TargetName = TargetCollection.FindAttribute<FString>(FName("BoneName"), FTransformCollection::TransformGroup);
+		const TManagedArray<int32>* const TargetGeometryGroup = TargetCollection.FindAttribute<int32>(FName("TransformToGeometryIndex"), FTransformCollection::TransformGroup);
+		if (SourceName && SourceGeometryGroup && TargetName && TargetGeometryGroup)
 		{
-			Chaos::TVec3<float> CoordMaxs(-FLT_MAX);
-			Chaos::TVec3<float> CoordMins(FLT_MAX);
-			for (int32 i = 0; i < ComponentSpaceFullVertices.Num(); i++)
+			for (int32 i = 0; i < SourceName->Num(); ++i)
 			{
-				for (int32 j = 0; j < 3; j++)
+				for (int32 j = 0; j < TargetName->Num(); ++j)
 				{
-					if (ComponentSpaceFullVertices[i][j] > CoordMaxs[j])
-					{
-						CoordMaxs[j] = ComponentSpaceFullVertices[i][j];
-					}
-					if (ComponentSpaceFullVertices[i][j] < CoordMins[j])
+					const FString TestName = FString::Printf(TEXT("%s_Tet"), *(*SourceName)[i]);
+					if ((*TargetName)[j].StartsWith(TestName))
 					{
-						CoordMins[j] = ComponentSpaceFullVertices[i][j];
+						Mapping.Add(FIntVector2((*SourceGeometryGroup)[i], (*TargetGeometryGroup)[j]));
+						break;
 					}
 				}
 			}
-			Chaos::TVec3<float> CoordDiff = (CoordMaxs - CoordMins) * BoundMultiplier;
-			SphereFullRadius = Chaos::FReal(FGenericPlatformMath::Min(CoordDiff[0], FGenericPlatformMath::Min(CoordDiff[1], CoordDiff[2])));
 		}
+
+		return MoveTemp(Mapping);
 	}
-	ParallelFor(PairedGeometry.Num(), [&](int32 Pdx)
+
+	static TMap<int32, int32> FindSourceToTargetTransformMap(const FManagedArrayCollection& SourceCollection, const FManagedArrayCollection& TargetCollection)
 	{
-		int32 AttributeGeometryIndex = PairedGeometry[Pdx][0];
-		int32 TargetGeometryIndex = PairedGeometry[Pdx][1];
-		if (ensure(0 <= AttributeGeometryIndex && AttributeGeometryIndex < Sample.VertexStart.Num()))
+		TMap<int32, int32> SourceIndexToTargetIndex;
+		TMap<FString, int32> TargetNameToIndex;
+		const TManagedArray<FString>* const SourceName = SourceCollection.FindAttribute<FString>(FName("BoneName"), FTransformCollection::TransformGroup);
+		const TManagedArray<FString>* const TargetName = TargetCollection.FindAttribute<FString>(FName("BoneName"), FTransformCollection::TransformGroup);
+		if (SourceName && TargetName)
 		{
-			if (ensure(0 <= TargetGeometryIndex && TargetGeometryIndex < Target.VertexStart.Num()))
+			for (int32 j = 0; j < TargetName->Num(); j++)
+			{
+				TargetNameToIndex.Add((*TargetName)[j], j);
+			}
+			for (int32 i = 0; i < SourceName->Num(); i++)
 			{
-				// Build component space vertices for TargetCollection
-				TArray<FVector3f> ComponentSpaceTargetVertices; // size of num vertices of the geometry entry. 
-				BuildComponentSpaceVertices(Target.Transform.Get(), Target.Parent.Get(), Target.BoneMap.Get(), Target.Vertex.Get(),
-					Target.VertexStart[TargetGeometryIndex], Target.VertexCount[TargetGeometryIndex], ComponentSpaceTargetVertices);
-
-				// Build component space vertices for SampleCollection
-				TArray<FVector3f> ComponentSpaceVertices; // size of num vertices of the geometry entry
-				BuildComponentSpaceVertices(Sample.Transform.Get(), Sample.Parent.Get(), Sample.BoneMap.Get(), Sample.Vertex.Get(),
-					Sample.VertexStart[AttributeGeometryIndex], Sample.VertexCount[AttributeGeometryIndex], ComponentSpaceVertices);
-
-				// build Sphere based BVH
-				Chaos::FReal SphereRadius = SphereFullRadius;
-				if (SampleScale == EDataflowTransferVertexAttributeNodeSampleScale::Component_Edge)
+				if (int32* TargetIndex = TargetNameToIndex.Find((*SourceName)[i]))
 				{
-					SphereRadius = Chaos::FReal(EdgeMultiplier * FMath::Max(
-						MaxEdgeLength(ComponentSpaceTargetVertices, Target.Indices.Get(), Target.VertexStart[TargetGeometryIndex], Target.FaceStart[TargetGeometryIndex], Target.FaceCount[TargetGeometryIndex]),
-						MaxEdgeLength(ComponentSpaceVertices, Sample.Indices.Get(), Sample.VertexStart[AttributeGeometryIndex], Sample.FaceStart[AttributeGeometryIndex], Sample.FaceCount[AttributeGeometryIndex])));
+					SourceIndexToTargetIndex.Add(i, *TargetIndex);
 				}
+			}
+		}
+
+		return MoveTemp(SourceIndexToTargetIndex);
+	}
 
-				int32 TargetVertexStartVal = Target.VertexStart[TargetGeometryIndex];
-				int32 TargetVertexCountVal = Target.VertexCount[TargetGeometryIndex];
-				int32 VertexStartVal = Sample.VertexStart[AttributeGeometryIndex];
-				int32 FaceStartVal = Sample.FaceStart[AttributeGeometryIndex];
-				int32 FaceCountVal = Sample.FaceCount[AttributeGeometryIndex];
+	static TArray<FTransferData> PairedGeometryTransfer(const TArray<FIntVector2>& PairedGeometry,
+		const FTransferFacade& Source, const FTransferFacade& Target, const FTransferProperties& TransferProperties)
+	{
+		TArray<FTransferData> TransferDataArray;
+		TransferDataArray.SetNum(Target.Vertex.Num());
+		Chaos::FReal SphereFullRadius;
 
-				if (BoundingVolumeType == EDataflowTransferVertexAttributeNodeBoundingVolume::Triangle)
+		if (TransferProperties.SourceScale == EDataflowTransferVertexAttributeNodeSourceScale::Asset_Edge || TransferProperties.SourceScale == EDataflowTransferVertexAttributeNodeSourceScale::Asset_Bound)
+		{
+			// Build component space vertices for TargetCollection
+			TArray<FVector3f> ComponentSpaceFullTargetVertices; // size of num vertices of the geometry entry. 
+			BuildComponentSpaceVertices(Target.Transform.Get(), Target.Parent.Get(), Target.BoneMap.Get(), Target.Vertex.Get(),
+				0, Target.Vertex.Num(), ComponentSpaceFullTargetVertices);
+
+			// Build component space vertices for TargetCollection
+			TArray<FVector3f> ComponentSpaceFullVertices; // size of num vertices of the geometry entry
+			BuildComponentSpaceVertices(Source.Transform.Get(), Source.Parent.Get(), Source.BoneMap.Get(), Source.Vertex.Get(),
+				0, Source.Vertex.Num(), ComponentSpaceFullVertices);
+			if (TransferProperties.SourceScale == EDataflowTransferVertexAttributeNodeSourceScale::Asset_Edge)
+			{
+				SphereFullRadius = Chaos::FReal(TransferProperties.EdgeMultiplier * FMath::Max(
+					MaxEdgeLength(ComponentSpaceFullTargetVertices, Target.Indices.Get(), 0, 0, Target.Indices.Num()),
+					MaxEdgeLength(ComponentSpaceFullVertices, Source.Indices.Get(), 0, 0, Source.Indices.Num())));
+			}
+			else if (TransferProperties.SourceScale == EDataflowTransferVertexAttributeNodeSourceScale::Asset_Bound)
+			{
+				Chaos::FVec3f CoordMaxs(-FLT_MAX);
+				Chaos::FVec3f CoordMins(FLT_MAX);
+				for (int32 i = 0; i < ComponentSpaceFullVertices.Num(); i++)
 				{
-					TArray<Chaos::TVec3<Chaos::FReal>> ComponentSpaceVerticesTVec3;
-					ComponentSpaceVerticesTVec3.SetNum(ComponentSpaceVertices.Num());
-					for (int32 SourceIndex = 0; SourceIndex < ComponentSpaceVerticesTVec3.Num(); SourceIndex++)
-					{
-						ComponentSpaceVerticesTVec3[SourceIndex] = Chaos::TVec3<Chaos::FReal>(ComponentSpaceVertices[SourceIndex]);
-					}
-					TConstArrayView<Chaos::TVec3<Chaos::FReal>> ConstComponentSpaceVertices(ComponentSpaceVerticesTVec3);
-					Chaos::FTriangleMesh TriangleMesh;
-					TArray<Chaos::TVec3<int32>> SourceElements;
-					SourceElements.SetNum(FaceCountVal);
-					for (int32 ElementIndex = 0; ElementIndex < FaceCountVal; ++ElementIndex)
+					for (int32 j = 0; j < 3; j++)
 					{
-						FIntVector3 Element = Sample.Indices[FaceStartVal + ElementIndex];
-						SourceElements[ElementIndex] = Chaos::TVec3<int32>(Element[0] - VertexStartVal, Element[1] - VertexStartVal, Element[2] - VertexStartVal);
+						if (ComponentSpaceFullVertices[i][j] > CoordMaxs[j])
+						{
+							CoordMaxs[j] = ComponentSpaceFullVertices[i][j];
+						}
+						if (ComponentSpaceFullVertices[i][j] < CoordMins[j])
+						{
+							CoordMins[j] = ComponentSpaceFullVertices[i][j];
+						}
 					}
-					TriangleMesh.Init(SourceElements);
-					Chaos::FTriangleMesh::TSpatialHashType<Chaos::FReal> SpatialHash;
-					TriangleMesh.BuildSpatialHash(ConstComponentSpaceVertices, SpatialHash, SphereRadius);
-					for (int32 TargetIndex = 0; TargetIndex < TargetVertexCountVal; TargetIndex++)
+				}
+				const Chaos::FVec3f CoordDiff = (CoordMaxs - CoordMins) * TransferProperties.BoundMultiplier;
+				SphereFullRadius = FMath::Min3(CoordDiff[0], CoordDiff[1], CoordDiff[2]);
+			}
+		}
+		ParallelFor(PairedGeometry.Num(), [&](int32 Pdx)
+			{
+				const int32 AttributeGeometryIndex = PairedGeometry[Pdx][0];
+				const int32 TargetGeometryIndex = PairedGeometry[Pdx][1];
+				if (ensure(0 <= AttributeGeometryIndex && AttributeGeometryIndex < Source.VertexStart.Num()))
+				{
+					if (ensure(0 <= TargetGeometryIndex && TargetGeometryIndex < Target.VertexStart.Num()))
 					{
-						TArray<Chaos::TTriangleCollisionPoint<Chaos::FReal>> Result;
-						if (TriangleMesh.PointClosestTriangleQuery(SpatialHash, ConstComponentSpaceVertices,
-							TargetIndex, Chaos::TVec3<Chaos::FReal>(ComponentSpaceTargetVertices[TargetIndex]), SphereRadius / 2.f, SphereRadius / 2.f,
-							[](const int32 PointIndex, const int32 TriangleIndex)->bool {return true; }, Result))
+						// Build component space vertices for TargetCollection
+						TArray<FVector3f> ComponentSpaceTargetVertices; // size of num vertices of the geometry entry. 
+						BuildComponentSpaceVertices(Target.Transform.Get(), Target.Parent.Get(), Target.BoneMap.Get(), Target.Vertex.Get(),
+							Target.VertexStart[TargetGeometryIndex], Target.VertexCount[TargetGeometryIndex], ComponentSpaceTargetVertices);
+
+						// Build component space vertices for SourceCollection
+						TArray<FVector3f> ComponentSpaceVertices; // size of num vertices of the geometry entry
+						BuildComponentSpaceVertices(Source.Transform.Get(), Source.Parent.Get(), Source.BoneMap.Get(), Source.Vertex.Get(),
+							Source.VertexStart[AttributeGeometryIndex], Source.VertexCount[AttributeGeometryIndex], ComponentSpaceVertices);
+
+						// build Sphere based BVH
+						Chaos::FReal SphereRadius = SphereFullRadius;
+						if (TransferProperties.SourceScale == EDataflowTransferVertexAttributeNodeSourceScale::Component_Edge)
 						{
-							for (const Chaos::TTriangleCollisionPoint<Chaos::FReal>& CollisionPoint : Result)
+							SphereRadius = Chaos::FReal(TransferProperties.EdgeMultiplier * FMath::Max(
+								MaxEdgeLength(ComponentSpaceTargetVertices, Target.Indices.Get(), Target.VertexStart[TargetGeometryIndex], Target.FaceStart[TargetGeometryIndex], Target.FaceCount[TargetGeometryIndex]),
+								MaxEdgeLength(ComponentSpaceVertices, Source.Indices.Get(), Source.VertexStart[AttributeGeometryIndex], Source.FaceStart[AttributeGeometryIndex], Source.FaceCount[AttributeGeometryIndex])));
+						}
+
+						const int32 TargetVertexStartVal = Target.VertexStart[TargetGeometryIndex];
+						const int32 TargetVertexCountVal = Target.VertexCount[TargetGeometryIndex];
+						const int32 VertexStartVal = Source.VertexStart[AttributeGeometryIndex];
+						const int32 FaceStartVal = Source.FaceStart[AttributeGeometryIndex];
+						const int32 FaceCountVal = Source.FaceCount[AttributeGeometryIndex];
+
+						if (TransferProperties.BoundingVolumeType == EDataflowTransferVertexAttributeNodeBoundingVolume::Triangle)
+						{
+							TArray<Chaos::TVec3<Chaos::FReal>> ComponentSpaceVerticesTVec3;
+							ComponentSpaceVerticesTVec3.SetNum(ComponentSpaceVertices.Num());
+							for (int32 SourceIndex = 0; SourceIndex < ComponentSpaceVerticesTVec3.Num(); SourceIndex++)
+							{
+								ComponentSpaceVerticesTVec3[SourceIndex] = Chaos::TVec3<Chaos::FReal>(ComponentSpaceVertices[SourceIndex]);
+							}
+							TConstArrayView<Chaos::TVec3<Chaos::FReal>> ConstComponentSpaceVertices(ComponentSpaceVerticesTVec3);
+							Chaos::FTriangleMesh TriangleMesh;
+							TArray<Chaos::TVec3<int32>> SourceElements;
+							SourceElements.SetNum(FaceCountVal);
+							for (int32 ElementIndex = 0; ElementIndex < FaceCountVal; ++ElementIndex)
 							{
-								float CurrentDistance = abs(CollisionPoint.Phi);
-								float TriRadius = FalloffThreshold * MaxEdgeLength(ComponentSpaceVertices, Sample.Indices.Get(), VertexStartVal, FaceStartVal + CollisionPoint.Indices[1], 1);
-								float FalloffScale = CalculateFalloffScale(Falloff, TriRadius, CurrentDistance);
-								if (!FMath::IsNearlyZero(FalloffScale))
+								FIntVector3 Element = Source.Indices[FaceStartVal + ElementIndex];
+								SourceElements[ElementIndex] = Chaos::TVec3<int32>(Element[0] - VertexStartVal, Element[1] - VertexStartVal, Element[2] - VertexStartVal);
+							}
+							TriangleMesh.Init(SourceElements);
+							Chaos::FTriangleMesh::TSpatialHashType<Chaos::FReal> SpatialHash;
+							TriangleMesh.BuildSpatialHash(ConstComponentSpaceVertices, SpatialHash, SphereRadius);
+							for (int32 TargetIndex = 0; TargetIndex < TargetVertexCountVal; TargetIndex++)
+							{
+								TArray<Chaos::TTriangleCollisionPoint<Chaos::FReal>> Result;
+								if (TriangleMesh.PointClosestTriangleQuery(SpatialHash, ConstComponentSpaceVertices,
+									TargetIndex, Chaos::TVec3<Chaos::FReal>(ComponentSpaceTargetVertices[TargetIndex]), SphereRadius / 2.f, SphereRadius / 2.f,
+									[](const int32 PointIndex, const int32 TriangleIndex)->bool {return true; }, Result))
 								{
-									int32 TargetCandidateIndex = CollisionPoint.Indices[0] + TargetVertexStartVal;
-									TArray<int32> TransferIndices;
-									TArray<float> TransferWeights;
-									for (int32 k = 0; k < 3; k++)
+									for (const Chaos::TTriangleCollisionPoint<Chaos::FReal>& CollisionPoint : Result)
 									{
-										TransferIndices.Add(Sample.Indices[FaceStartVal + CollisionPoint.Indices[1]][k]);
-										TransferWeights.Add(CollisionPoint.Bary[k + 1]);
+										const float CurrentDistance = abs(CollisionPoint.Phi);
+										const float TriRadius = TransferProperties.FalloffThreshold * MaxEdgeLength(ComponentSpaceVertices, Source.Indices.Get(), VertexStartVal, FaceStartVal + CollisionPoint.Indices[1], 1);
+										const float FalloffScale = CalculateFalloffScale(TransferProperties.Falloff, TriRadius, CurrentDistance);
+										if (!FMath::IsNearlyZero(FalloffScale))
+										{
+											const int32 TargetCandidateIndex = CollisionPoint.Indices[0] + TargetVertexStartVal;
+											TArray<int32> TransferIndices;
+											TArray<float> TransferWeights;
+											for (int32 k = 0; k < 3; k++)
+											{
+												TransferIndices.Add(Source.Indices[FaceStartVal + CollisionPoint.Indices[1]][k]);
+												TransferWeights.Add(CollisionPoint.Bary[k + 1]);
+											}
+											TransferDataArray[TargetCandidateIndex] = FTransferData(TransferIndices, TransferWeights, FalloffScale);
+											break;
+										}
 									}
-									TransferDataArray[TargetCandidateIndex] = UE::Private::FTransferData(TransferIndices, TransferWeights, FalloffScale);
-									break;
 								}
 							}
 						}
-					}
-				}
-				else if (BoundingVolumeType == EDataflowTransferVertexAttributeNodeBoundingVolume::Vertex)
-				{
-					TUniquePtr<BVH> VertexBVH(BuildParticleSphereBVH(ComponentSpaceTargetVertices, SphereRadius));
-					for (int32 i = 0; i < FaceCountVal; i++)
-					{
-						FIntVector3 ComponentTriangle = Sample.Indices[FaceStartVal + i] - FIntVector3(VertexStartVal);
-						TArray<int32> TargetVertexIntersection({});
-						TriangleToVertexIntersections(*VertexBVH, ComponentSpaceVertices, ComponentTriangle, TargetVertexIntersection);
-
-						for (int32 j = 0; j < TargetVertexIntersection.Num(); j++)
+						else if (TransferProperties.BoundingVolumeType == EDataflowTransferVertexAttributeNodeBoundingVolume::Vertex)
 						{
-							Chaos::TVector<float, 3> Bary,
-								TriPos0(ComponentSpaceVertices[ComponentTriangle[0]]),
-								TriPos1(ComponentSpaceVertices[ComponentTriangle[1]]),
-								TriPos2(ComponentSpaceVertices[ComponentTriangle[2]]),
-								ParticlePos(ComponentSpaceTargetVertices[TargetVertexIntersection[j]]);
-
-							Chaos::TVector<float, 3> ClosestPoint = Chaos::FindClosestPointAndBaryOnTriangle(TriPos0, TriPos1, TriPos2, ParticlePos, Bary);
-							float CurrentDistance = (ParticlePos - ClosestPoint).Size();
-							float TriRadius = FalloffThreshold * MaxEdgeLength(ComponentSpaceVertices, Sample.Indices.Get(), VertexStartVal, FaceStartVal + i, 1);
-							float FalloffScale = CalculateFalloffScale(Falloff, TriRadius, CurrentDistance);
-							if (!FMath::IsNearlyZero(FalloffScale))
+							TUniquePtr<BVH> VertexBVH(BuildParticleSphereBVH(ComponentSpaceTargetVertices, SphereRadius));
+							for (int32 i = 0; i < FaceCountVal; i++)
 							{
-								int32 TargetIndex = TargetVertexIntersection[j] + TargetVertexStartVal;
-								TArray<int32> TransferIndices;
-								TArray<float> TransferWeights;
-								for (int32 k = 0; k < 3; k++)
+								FIntVector3 ComponentTriangle = Source.Indices[FaceStartVal + i] - FIntVector3(VertexStartVal);
+								TArray<int32> TargetVertexIntersection({});
+								TriangleToVertexIntersections(*VertexBVH, ComponentSpaceVertices, ComponentTriangle, TargetVertexIntersection);
+
+								for (int32 j = 0; j < TargetVertexIntersection.Num(); j++)
 								{
-									TransferIndices.Add(ComponentTriangle[k] + VertexStartVal);
-									TransferWeights.Add(Bary[k]);
+									Chaos::FVec3f Bary,
+										TriPos0(ComponentSpaceVertices[ComponentTriangle[0]]),
+										TriPos1(ComponentSpaceVertices[ComponentTriangle[1]]),
+										TriPos2(ComponentSpaceVertices[ComponentTriangle[2]]),
+										ParticlePos(ComponentSpaceTargetVertices[TargetVertexIntersection[j]]);
+
+									Chaos::FVec3f ClosestPoint = Chaos::FindClosestPointAndBaryOnTriangle(TriPos0, TriPos1, TriPos2, ParticlePos, Bary);
+									float CurrentDistance = (ParticlePos - ClosestPoint).Size();
+									float TriRadius = TransferProperties.FalloffThreshold * MaxEdgeLength(ComponentSpaceVertices, Source.Indices.Get(), VertexStartVal, FaceStartVal + i, 1);
+									float FalloffScale = CalculateFalloffScale(TransferProperties.Falloff, TriRadius, CurrentDistance);
+									if (!FMath::IsNearlyZero(FalloffScale))
+									{
+										int32 TargetIndex = TargetVertexIntersection[j] + TargetVertexStartVal;
+										TArray<int32> TransferIndices;
+										TArray<float> TransferWeights;
+										for (int32 k = 0; k < 3; k++)
+										{
+											TransferIndices.Add(ComponentTriangle[k] + VertexStartVal);
+											TransferWeights.Add(Bary[k]);
+										}
+										TransferDataArray[TargetIndex] = FTransferData(TransferIndices, TransferWeights, FalloffScale);
+										break;
+									}
 								}
-								TransferDataArray[TargetIndex] = UE::Private::FTransferData(TransferIndices, TransferWeights, FalloffScale);
-								break;
 							}
 						}
 					}
 				}
-			}
-		}
-	});
-	return TransferDataArray;
-}
+			});
+		return MoveTemp(TransferDataArray);
+	}
 
-TArray<UE::Private::FTransferData> FGeometryCollectionTransferVertexAttributeNode::NearestVertexTransfer(FCollectionAttributeKey Key, const UE::Private::FTransferFacade& Sample, const UE::Private::FTransferFacade& Target) const
-{
-	TArray<UE::Private::FTransferData> TransferDataArray;
-	TransferDataArray.SetNum(Target.Vertex.Num());
-	// Build component space vertices for TargetCollection
-	TArray<FVector3f> ComponentSpaceTargetVertices;
-	BuildComponentSpaceVertices(Target.Transform.Get(), Target.Parent.Get(), Target.BoneMap.Get(), Target.Vertex.Get(), 0, Target.Vertex.Num(), ComponentSpaceTargetVertices);
-
-	// Build component space vertices for SourceCollection
-	TArray<FVector3f> ComponentSpaceVertices;
-	BuildComponentSpaceVertices(Sample.Transform.Get(), Sample.Parent.Get(), Sample.BoneMap.Get(), Sample.Vertex.Get(), 0, Sample.Vertex.Num(), ComponentSpaceVertices);
-
-	// build Sphere based BVH
-	Chaos::FReal SphereRadius = Chaos::FReal(EdgeMultiplier * FMath::Max(
-		MaxEdgeLength(ComponentSpaceTargetVertices, Target.Indices.Get(), 0, 0, Target.Indices.Num()),
-		MaxEdgeLength(ComponentSpaceVertices, Sample.Indices.Get(), 0, 0, Sample.Indices.Num())));
-	TUniquePtr<BVH> VertexBVH(BuildParticleSphereBVH(ComponentSpaceTargetVertices, SphereRadius));
-
-	for (int32 i = 0; i < Sample.Indices.Num(); i++)
+	static TArray<FTransferData> NearestVertexTransfer(const FTransferFacade& Source, const FTransferFacade& Target, const FTransferProperties& TransferProperties)
 	{
-		FIntVector3 Triangle = Sample.Indices[i];
-		TArray<int32> TargetVertexIntersection({});
-		TriangleToVertexIntersections(*VertexBVH, ComponentSpaceVertices, Triangle, TargetVertexIntersection);
+		TArray<FTransferData> TransferDataArray;
+		TransferDataArray.SetNum(Target.Vertex.Num());
+		// Build component space vertices for TargetCollection
+		TArray<FVector3f> ComponentSpaceTargetVertices;
+		BuildComponentSpaceVertices(Target.Transform.Get(), Target.Parent.Get(), Target.BoneMap.Get(), Target.Vertex.Get(), 0, Target.Vertex.Num(), ComponentSpaceTargetVertices);
 
-		for (int32 j = 0; j < TargetVertexIntersection.Num(); j++)
+		// Build component space vertices for SourceCollection
+		TArray<FVector3f> ComponentSpaceVertices;
+		BuildComponentSpaceVertices(Source.Transform.Get(), Source.Parent.Get(), Source.BoneMap.Get(), Source.Vertex.Get(), 0, Source.Vertex.Num(), ComponentSpaceVertices);
+
+		// build Sphere based BVH
+		Chaos::FReal SphereRadius = Chaos::FReal(TransferProperties.EdgeMultiplier * FMath::Max(
+			MaxEdgeLength(ComponentSpaceTargetVertices, Target.Indices.Get(), 0, 0, Target.Indices.Num()),
+			MaxEdgeLength(ComponentSpaceVertices, Source.Indices.Get(), 0, 0, Source.Indices.Num())));
+		TUniquePtr<BVH> VertexBVH(BuildParticleSphereBVH(ComponentSpaceTargetVertices, SphereRadius));
+
+		for (int32 i = 0; i < Source.Indices.Num(); i++)
 		{
-			Chaos::TVector<float, 3> Bary,
-				TriPos0(ComponentSpaceVertices[Sample.Indices[i][0]]),
-				TriPos1(ComponentSpaceVertices[Sample.Indices[i][1]]),
-				TriPos2(ComponentSpaceVertices[Sample.Indices[i][2]]),
-				ParticlePos(ComponentSpaceTargetVertices[TargetVertexIntersection[j]]);
-
-			Chaos::TVector<float, 3> ClosestPoint = Chaos::FindClosestPointAndBaryOnTriangle(TriPos0, TriPos1, TriPos2, ParticlePos, Bary);
-			float CurrentDistance = (ParticlePos - ClosestPoint).Size();
-			float TriRadius = FalloffThreshold * MaxEdgeLength(ComponentSpaceVertices, Sample.Indices.Get(), 0, i, 1);
-			float FalloffScale = CalculateFalloffScale(Falloff, TriRadius, CurrentDistance);
-			if (!FMath::IsNearlyZero(FalloffScale))
+			FIntVector3 Triangle = Source.Indices[i];
+			TArray<int32> TargetVertexIntersection({});
+			TriangleToVertexIntersections(*VertexBVH, ComponentSpaceVertices, Triangle, TargetVertexIntersection);
+
+			for (int32 j = 0; j < TargetVertexIntersection.Num(); j++)
 			{
-				int32 TargetIndex = TargetVertexIntersection[j];
-				TArray<int32> TransferIndices;
-				TArray<float> TransferWeights;
-				for (int32 k = 0; k < 3; k++)
+				Chaos::FVec3f Bary,
+					TriPos0(ComponentSpaceVertices[Source.Indices[i][0]]),
+					TriPos1(ComponentSpaceVertices[Source.Indices[i][1]]),
+					TriPos2(ComponentSpaceVertices[Source.Indices[i][2]]),
+					ParticlePos(ComponentSpaceTargetVertices[TargetVertexIntersection[j]]);
+
+				Chaos::FVec3f ClosestPoint = Chaos::FindClosestPointAndBaryOnTriangle(TriPos0, TriPos1, TriPos2, ParticlePos, Bary);
+				float CurrentDistance = (ParticlePos - ClosestPoint).Size();
+				float TriRadius = TransferProperties.FalloffThreshold * MaxEdgeLength(ComponentSpaceVertices, Source.Indices.Get(), 0, i, 1);
+				float FalloffScale = CalculateFalloffScale(TransferProperties.Falloff, TriRadius, CurrentDistance);
+				if (!FMath::IsNearlyZero(FalloffScale))
 				{
-					TransferIndices.Add(Sample.Indices[i][k]);
-					TransferWeights.Add(Bary[k]);
+					int32 TargetIndex = TargetVertexIntersection[j];
+					TArray<int32> TransferIndices;
+					TArray<float> TransferWeights;
+					for (int32 k = 0; k < 3; k++)
+					{
+						TransferIndices.Add(Source.Indices[i][k]);
+						TransferWeights.Add(Bary[k]);
+					}
+					TransferDataArray[TargetIndex] = FTransferData(TransferIndices, TransferWeights, FalloffScale);
+					break;
 				}
-				TransferDataArray[TargetIndex] = UE::Private::FTransferData(TransferIndices, TransferWeights, FalloffScale);
-				break;
 			}
 		}
+		return MoveTemp(TransferDataArray);
 	}
-	return TransferDataArray;
-}
-
 
-float FGeometryCollectionTransferVertexAttributeNode::MaxEdgeLength(TArray<FVector3f>& Vert, const TManagedArray<FIntVector3>& Tri, int VertexOffset, int TriStart, int TriCount)
-{
-	auto TriInRange = [](const FIntVector3& T, int Max) {
-		for (int k = 0; k < 3; k++)
+	static TArray<FTransferData> ComputeTransferData(const FManagedArrayCollection& SourceCollection, const FManagedArrayCollection& TargetCollection, const FTransferProperties& TransferProperties)
+	{
+		TArray<FTransferData> TransferDataArray; //store barycentric weight info
+		const FTransferFacade Target(TargetCollection);
+		const FTransferFacade Source(SourceCollection);
+		if (Target.IsValid() && Source.IsValid())
 		{
-			if (ensure(0 <= T[k] && T[k] < Max))
+			TArray<FIntVector2> AlignedGeometry = FindSourceToTargetGeometryMap(SourceCollection, TargetCollection);
+			if (AlignedGeometry.Num() == SourceCollection.NumElements(FGeometryCollection::GeometryGroup))
 			{
-				return true;
+				TransferDataArray = PairedGeometryTransfer(AlignedGeometry, Source, Target, TransferProperties);
+			}
+			else
+			{
+				TransferDataArray = NearestVertexTransfer(Source, Target, TransferProperties);
 			}
 		}
-		return false;
-		};
-
-	float Max = 0;
-	int TriStop = TriStart + TriCount;
-	for (int i = TriStart; i < TriStop; i++)
-	{
-		if (TriInRange(Tri[i] - FIntVector3(VertexOffset), Vert.Num()))
-		{
-			Max = FMath::Max(Max, (Vert[Tri[i][0] - VertexOffset] - Vert[Tri[i][1] - VertexOffset]).SquaredLength());
-			Max = FMath::Max(Max, (Vert[Tri[i][0] - VertexOffset] - Vert[Tri[i][2] - VertexOffset]).SquaredLength());
-			Max = FMath::Max(Max, (Vert[Tri[i][1] - VertexOffset] - Vert[Tri[i][2] - VertexOffset]).SquaredLength());
-		}
+		return MoveTemp(TransferDataArray);
 	}
-	return FMath::Sqrt(Max);
 }
 
-void FGeometryCollectionTransferVertexAttributeNode::BuildComponentSpaceVertices(const TManagedArray<FTransform3f>& LocalSpaceTransform, const TManagedArray<int32>& Parent, const TManagedArray<int32>& BoneMapArray,
-	const TManagedArray<FVector3f>& VertexArray, int32 Start, int32 Count, TArray<FVector3f>& ComponentSpaceVertices)
+void FGeometryCollectionTransferVertexAttributeNode::Evaluate(UE::Dataflow::FContext& Context, const FDataflowOutput* Out) const
 {
-	TArray<FTransform3f> ComponentTransform;
-	GeometryCollectionAlgo::GlobalMatrices(LocalSpaceTransform, Parent, ComponentTransform);
+	FCollectionAttributeKey Key = GetValue(Context, &AttributeKey, AttributeKey);
 
-	ComponentSpaceVertices.SetNumUninitialized(Count);
-	for (int i = 0; i < Count; i++)
+	if (Out->IsA(&Collection))
 	{
-		int j = i + Start;
-		if (0 < (BoneMapArray)[i] && (BoneMapArray)[i] < ComponentTransform.Num())
-		{
-			ComponentSpaceVertices[i] = ComponentTransform[(BoneMapArray)[j]].TransformPosition(VertexArray[j]);
-		}
-		else
+		FManagedArrayCollection TargetCollection = GetValue(Context, &Collection);
+		const FManagedArrayCollection& SourceCollection = GetValue(Context, &FromCollection);
+		UE::Private::FTransferFacade Target(TargetCollection);
+		const UE::Private::FTransferFacade Source(SourceCollection);
+
+		const UE::Private::FTransferProperties TransferProperties = {
+			BoundingVolumeType,
+			SourceScale,
+			Falloff,
+			FalloffThreshold,
+			EdgeMultiplier,
+			BoundMultiplier
+		};
+
+		if (Target.IsValid() && Source.IsValid())
 		{
-			ComponentSpaceVertices[i] = VertexArray[j];
+			const TArray<UE::Private::FTransferData> TransferDataArray = ComputeTransferData(SourceCollection, TargetCollection, TransferProperties);
+			if (const TManagedArray<float>* SourceAttributeFloatArray = Source.GetAttributeArray<float>(Key.Attribute, Key.Group))
+			{
+				TManagedArray<float>* TargetAttributeArray = Target.GetAttributeArray<float>(Key.Attribute, Key.Group);
+				TargetAttributeArray->Fill(0.f);
+				for (int32 VertexIndex = 0; VertexIndex < TargetAttributeArray->Num(); VertexIndex++)
+				{
+					const UE::Private::FTransferData& TransferData = TransferDataArray[VertexIndex];
+					if (TransferData.SourceIndices.Num())
+					{
+						for (int32 SourceLocalIdx = 0; SourceLocalIdx < TransferData.SourceIndices.Num(); ++SourceLocalIdx)
+						{
+							(*TargetAttributeArray)[VertexIndex] += (*SourceAttributeFloatArray)[TransferData.SourceIndices[SourceLocalIdx]]
+								* TransferData.SourceWeights[SourceLocalIdx] * TransferData.FalloffScale;
+						}
+					}
+				}
+			}
+			else if (const TManagedArray<FLinearColor>* SourceAttributeColorArray = Source.GetAttributeArray<FLinearColor>(Key.Attribute, Key.Group))
+			{
+				TManagedArray<FLinearColor>* TargetAttributeArray = Target.GetAttributeArray<FLinearColor>(Key.Attribute, Key.Group);
+				TargetAttributeArray->Fill(FLinearColor(0, 0, 0, 0));
+				for (int32 VertexIndex = 0; VertexIndex < TargetAttributeArray->Num(); VertexIndex++)
+				{
+					const UE::Private::FTransferData& TransferData = TransferDataArray[VertexIndex];
+					if (TransferData.SourceIndices.Num())
+					{
+						for (int32 SourceLocalIdx = 0; SourceLocalIdx < TransferData.SourceIndices.Num(); ++SourceLocalIdx)
+						{
+							(*TargetAttributeArray)[VertexIndex] += (*SourceAttributeColorArray)[TransferData.SourceIndices[SourceLocalIdx]]
+								* TransferData.SourceWeights[SourceLocalIdx] * TransferData.FalloffScale;
+						}
+					}
+				}
+			}
 		}
-	}
-}
-
-FGeometryCollectionTransferVertexAttributeNode::BVH* FGeometryCollectionTransferVertexAttributeNode::BuildParticleSphereBVH(const TArray<FVector3f>& Vertices, float Radius)
-{
-	TArray<SphereType*> VertexSpherePtrs;
-	TArray<SphereType> VertexSpheres;
-	VertexSpheres.Init(SphereType(Chaos::FVec3f(0), Radius), Vertices.Num());
-	VertexSpherePtrs.SetNum(Vertices.Num());
 
-	for (int32 i = 0; i < Vertices.Num(); i++)
+		SetValue(Context, MoveTemp(TargetCollection), &Collection);
+	}
+	else if (Out->IsA(&AttributeKey))
 	{
-		Chaos::FVec3f SphereCenter(Vertices[i]);
-		SphereType VertexSphere(SphereCenter, Radius);
-		VertexSpheres[i] = SphereType(SphereCenter, Radius);
-		VertexSpherePtrs[i] = &VertexSpheres[i];
+		SetValue(Context, MoveTemp(Key), &AttributeKey);
 	}
-	return new BVH(VertexSpherePtrs);
 }
 
-void FGeometryCollectionTransferVertexAttributeNode::TriangleToVertexIntersections(
-	const BVH& VertexBVH, const TArray<FVector3f>& ComponentSpaceVertices, const FIntVector3& Triangle, TArray<int32>& OutTargetVertexIntersection)
+void FGeometryCollectionTransferVertexSkinWeightsNode::Evaluate(UE::Dataflow::FContext& Context, const FDataflowOutput* Out) const
 {
-	OutTargetVertexIntersection.Empty();
-
-	TArray<int32> TargetVertexIntersection0 = VertexBVH.FindAllIntersections(ComponentSpaceVertices[Triangle[0]]);
-	TArray<int32> TargetVertexIntersection1 = VertexBVH.FindAllIntersections(ComponentSpaceVertices[Triangle[1]]);
-	TArray<int32> TargetVertexIntersection2 = VertexBVH.FindAllIntersections(ComponentSpaceVertices[Triangle[2]]);
-	TargetVertexIntersection0.Sort();
-	TargetVertexIntersection1.Sort();
-	TargetVertexIntersection2.Sort();
-
-	for (int32 k = 0; k < TargetVertexIntersection0.Num(); k++)
+	if (Out->IsA(&Collection))
 	{
-		if (TargetVertexIntersection1.Contains(TargetVertexIntersection0[k])
-			&& TargetVertexIntersection2.Contains(TargetVertexIntersection0[k]))
+		FManagedArrayCollection TargetCollection = GetValue(Context, &Collection);
+		const FManagedArrayCollection& SourceCollection = GetValue(Context, &FromCollection);
+		UE::Private::FTransferFacade Target(TargetCollection);
+		const UE::Private::FTransferFacade Source(SourceCollection);
+
+		const UE::Private::FTransferProperties TransferProperties = {
+			BoundingVolumeType,
+			SourceScale,
+			Falloff,
+			FalloffThreshold,
+			EdgeMultiplier,
+			BoundMultiplier
+		};
+
+		if (Target.IsValid() && Source.IsValid())
 		{
-			OutTargetVertexIntersection.Emplace(TargetVertexIntersection0[k]);
-		}
-	}
-}
+			const TArray<UE::Private::FTransferData> TransferDataArray = ComputeTransferData(SourceCollection, TargetCollection, TransferProperties);
 
+			GeometryCollection::Facades::FVertexBoneWeightsFacade SourceVertexBoneWeightsFacade(SourceCollection);
+			GeometryCollection::Facades::FVertexBoneWeightsFacade TargetVertexBoneWeightsFacade(TargetCollection);
+			const TManagedArray<TArray<int32>>* SourceBoneIndices = SourceVertexBoneWeightsFacade.FindBoneIndices();
+			const TManagedArray<TArray<float>>* SourceBoneWeights = SourceVertexBoneWeightsFacade.FindBoneWeights();
+			if (SourceBoneIndices && SourceBoneWeights)
+			{
+				//
+				// Compute the bone index mappings. This allows the transfer operator to retarget weights to the correct skeleton.
+				//
+				const TMap<int32, int32> SourceBoneToTargetBone = UE::Private::FindSourceToTargetTransformMap(SourceCollection, TargetCollection);
+				for (int32 VertexIndex = 0; VertexIndex < Target.Vertex.Num(); VertexIndex++)
+				{
+					const UE::Private::FTransferData& TransferData = TransferDataArray[VertexIndex];
+					if (TransferData.SourceIndices.Num())
+					{
+						TMap<int32, float> BoneWeightBucket;
+						for (int32 SourceLocalIdx = 0; SourceLocalIdx < TransferData.SourceIndices.Num(); ++SourceLocalIdx)
+						{
+							for (int32 BoneLocalIdx = 0; BoneLocalIdx < (*SourceBoneIndices)[TransferData.SourceIndices[SourceLocalIdx]].Num(); ++BoneLocalIdx)
+							{
+								const int32 SourceBoneIndex = (*SourceBoneIndices)[TransferData.SourceIndices[SourceLocalIdx]][BoneLocalIdx];
+								if (SourceBoneToTargetBone.Contains(SourceBoneIndex))
+								{
+									const int32 TargetBoneIndex = SourceBoneToTargetBone[SourceBoneIndex];
+									const float BoneWeight = (*SourceBoneWeights)[TransferData.SourceIndices[SourceLocalIdx]][BoneLocalIdx];
+									if (BoneWeightBucket.Contains(TargetBoneIndex))
+									{
+										BoneWeightBucket[TargetBoneIndex] += TransferData.SourceWeights[SourceLocalIdx] * BoneWeight;
+									}
+									else
+									{
+										BoneWeightBucket.Add(TargetBoneIndex, TransferData.SourceWeights[SourceLocalIdx] * BoneWeight);
+									}
+								}
+							}
+						}
 
-float FGeometryCollectionTransferVertexAttributeNode::CalculateFalloffScale(EDataflowTransferVertexAttributeNodeFalloff FalloffSetting, float Threshold, float Distance)
-{
-	float Denominator = 1.0;
-	if (Distance > Threshold && !FMath::IsNearlyZero(Threshold))
-	{
-		Denominator = Distance / Threshold;
-	}
-	switch (FalloffSetting)
-	{
-	case EDataflowTransferVertexAttributeNodeFalloff::Linear:
-		return 1. / Denominator;
-	case EDataflowTransferVertexAttributeNodeFalloff::Squared:
-		return 1. / FMath::Square(Denominator);
+						TArray<int32> VertexBoneIndex;
+						TArray<float> VertexBoneWeight;
+						VertexBoneIndex.Reserve(BoneWeightBucket.Num());
+						VertexBoneWeight.Reserve(BoneWeightBucket.Num());
+						for (const TPair<int32, float>& Pair : BoneWeightBucket)
+						{
+							VertexBoneIndex.Add(Pair.Key);
+							VertexBoneWeight.Add(Pair.Value);
+						}
+						TargetVertexBoneWeightsFacade.ModifyBoneWeight(VertexIndex, VertexBoneIndex, VertexBoneWeight);
+					}
+				}
+			}
+		}
+		SetValue(Context, MoveTemp(TargetCollection), &Collection);
 	}
-	return 1.0;
 }
-
 void FGeometryCollectionSetKinematicVertexSelectionNode::Evaluate(UE::Dataflow::FContext& Context, const FDataflowOutput* Out) const
 {
 	if (Out->IsA(&Collection))
diff --git a/Engine/Plugins/Experimental/GeometryCollectionPlugin/Source/GeometryCollectionNodes/Public/Dataflow/GeometryCollectionSkeletalMeshToCollectionNode.h b/Engine/Plugins/Experimental/GeometryCollectionPlugin/Source/GeometryCollectionNodes/Public/Dataflow/GeometryCollectionSkeletalMeshToCollectionNode.h
--- a/Engine/Plugins/Experimental/GeometryCollectionPlugin/Source/GeometryCollectionNodes/Public/Dataflow/GeometryCollectionSkeletalMeshToCollectionNode.h
+++ b/Engine/Plugins/Experimental/GeometryCollectionPlugin/Source/GeometryCollectionNodes/Public/Dataflow/GeometryCollectionSkeletalMeshToCollectionNode.h
@@ -33,6 +33,9 @@ public:
 	UPROPERTY(meta = (DataflowOutput, DisplayName = "Collection"))
 	FManagedArrayCollection Collection;
 
+	UPROPERTY(EditAnywhere, Category = "Dataflow", meta = (DisplayName = "Transforms Only"))
+	bool bImportTransformOnly = false;
+
 	FSkeletalMeshToCollectionDataflowNode(const UE::Dataflow::FNodeParameters& InParam, FGuid InGuid = FGuid::NewGuid())
 		: FDataflowNode(InParam, InGuid)
 	{
diff --git a/Engine/Plugins/Experimental/GeometryCollectionPlugin/Source/GeometryCollectionNodes/Public/Dataflow/GeometryCollectionSkeletonToCollectionNode.h b/Engine/Plugins/Experimental/GeometryCollectionPlugin/Source/GeometryCollectionNodes/Public/Dataflow/GeometryCollectionSkeletonToCollectionNode.h
--- a/Engine/Plugins/Experimental/GeometryCollectionPlugin/Source/GeometryCollectionNodes/Public/Dataflow/GeometryCollectionSkeletonToCollectionNode.h
+++ b/Engine/Plugins/Experimental/GeometryCollectionPlugin/Source/GeometryCollectionNodes/Public/Dataflow/GeometryCollectionSkeletonToCollectionNode.h
@@ -16,8 +16,8 @@ namespace UE_DEPRECATED(5.5, "Use UE::Dafaflow instead.") Dataflow {}
 
 class USkeletalMesh;
 
-USTRUCT(meta = (DataflowGeometryCollection))
-struct FSkeletonToCollectionDataflowNode : public FDataflowNode
+USTRUCT(meta = (DataflowGeometryCollection, Deprecated = "5.5"))
+struct UE_DEPRECATED(5.5, "Use SkeletalMeshToCollection instead") FSkeletonToCollectionDataflowNode : public FDataflowNode
 {
 	GENERATED_USTRUCT_BODY()
 	DATAFLOW_NODE_DEFINE_INTERNAL(FSkeletonToCollectionDataflowNode, "SkeletonToCollection", "GeometryCollection", "")
diff --git a/Engine/Plugins/Experimental/GeometryCollectionPlugin/Source/GeometryCollectionNodes/Public/Dataflow/GeometryCollectionTransferVertexAttributeNode.h b/Engine/Plugins/Experimental/GeometryCollectionPlugin/Source/GeometryCollectionNodes/Public/Dataflow/GeometryCollectionTransferVertexAttributeNode.h
--- a/Engine/Plugins/Experimental/GeometryCollectionPlugin/Source/GeometryCollectionNodes/Public/Dataflow/GeometryCollectionTransferVertexAttributeNode.h
+++ b/Engine/Plugins/Experimental/GeometryCollectionPlugin/Source/GeometryCollectionNodes/Public/Dataflow/GeometryCollectionTransferVertexAttributeNode.h
@@ -17,13 +17,6 @@ namespace Dataflow = UE::Dataflow;
 namespace UE_DEPRECATED(5.5, "Use UE::Dafaflow instead.") Dataflow {}
 #endif
 
-namespace UE::Private 
-{
-	class FTransferFacade;
-	struct FTransferData;
-};
-
-
 UENUM(BlueprintType)
 enum class EDataflowTransferVertexAttributeNodeFalloff : uint8
 {
@@ -41,7 +34,7 @@ enum class EDataflowTransferVertexAttributeNodeFalloff : uint8
 };
 
 UENUM(BlueprintType)
-enum class EDataflowTransferVertexAttributeNodeSampleScale : uint8
+enum class EDataflowTransferVertexAttributeNodeSourceScale : uint8
 {
 	/** Bounding volume hierarchy cell size based on max edge length of each geometry group*/
 	Component_Edge  UMETA(DisplayName = "Component Max Edge"),
@@ -72,7 +65,7 @@ enum class EDataflowTransferVertexAttributeNodeBoundingVolume : uint8
 };
 
 /**
- * Transfer float properties from a sample collection to a target collection.
+ * Transfer float properties from a source collection to a target collection.
  */
 USTRUCT(meta = (DataflowGeometryCollection))
 struct FGeometryCollectionTransferVertexAttributeNode : public FDataflowNode
@@ -87,18 +80,13 @@ public:
 	UPROPERTY(Meta = (DataflowInput, DataflowOutput, DataflowPassthrough = "Collection"))
 	FManagedArrayCollection Collection;
 
-	/* Sample collection to transfer vertex attribute from. */
+	/* Source collection to transfer vertex attribute from. */
 	UPROPERTY(Meta = (DataflowInput, DisplayName = "FromCollection"))
 	FManagedArrayCollection FromCollection;
 
 	/* The name of the vertex attribute to generate indices from. */
-	UPROPERTY(EditAnywhere, Category = "Dataflow", Meta = (DataflowInput, DataflowOutput, DisplayName = "AttributeKey", DataflowPassthrough = "AttributeKey",
-		EditCondition = "!bTransferSkinWeights", EditConditionHides))
+	UPROPERTY(EditAnywhere, Category = "Dataflow", Meta = (DataflowInput, DataflowOutput, DisplayName = "AttributeKey", DataflowPassthrough = "AttributeKey"))
 	FCollectionAttributeKey AttributeKey = FCollectionAttributeKey(FString(""), FString("Vertices"));
-
-	/* If kinematic skin weight attribute is transferred. */
-	UPROPERTY(EditAnywhere, Category = "Dataflow", Meta = (DisplayName = "TransferSkinWeights"))
-	bool bTransferSkinWeights = false;
 	
 	/* Bounding volume type for source assets[default: Triangle] */
 	UPROPERTY(EditAnywhere, Category = "Thresholds")
@@ -106,23 +94,23 @@ public:
 
 	/* Bounding volume hierarchy cell size for neighboring vertices to transfer into[default: Asset] */
 	UPROPERTY(EditAnywhere, Category = "Thresholds")
-	EDataflowTransferVertexAttributeNodeSampleScale SampleScale = EDataflowTransferVertexAttributeNodeSampleScale::Asset_Bound;
+	EDataflowTransferVertexAttributeNodeSourceScale SourceScale = EDataflowTransferVertexAttributeNodeSourceScale::Asset_Bound;
 
-	/* Falloff of sample value based on distance from sample triangle[default: Squared] */
+	/* Falloff of source value based on distance from source triangle[default: Squared] */
 	UPROPERTY(EditAnywhere, Category = "Thresholds")
 	EDataflowTransferVertexAttributeNodeFalloff Falloff = EDataflowTransferVertexAttributeNodeFalloff::None;
 
-	/* Threshold based on distance from sample triangle.Values sampled past the threshold will falloff.[Defaults to 1 percent of triangle size(0.01)] */
+	/* Threshold based on distance from source triangle.Values past the threshold will falloff.[Defaults to 1 percent of triangle size(0.01)] */
 	UPROPERTY(EditAnywhere, Category = "Thresholds")
 	float FalloffThreshold = 0.01f;
 
 	/* Edge multiplier for the Bounding Volume Hierarchy(BVH) target's particle search radius. */
-	UPROPERTY(EditAnywhere, Category = "Thresholds", meta = (EditCondition = "SampleScale == EDataflowTransferVertexAttributeNodeSampleScale::Asset_Edge || SampleScale == EDataflowTransferVertexAttributeNodeSampleScale::Component_Edge", EditConditionHides))
+	UPROPERTY(EditAnywhere, Category = "Thresholds", meta = (EditCondition = "SourceScale == EDataflowTransferVertexAttributeNodeSourceScale::Asset_Edge || SourceScale == EDataflowTransferVertexAttributeNodeSourceScale::Component_Edge", EditConditionHides))
 	float EdgeMultiplier = 0.5f;
 
 	/* Max bound multiplier for the Bounding Volume Hierarchy(BVH) target's particle search radius. */
 	UPROPERTY(EditAnywhere, Category = "Thresholds", meta = (
-		EditCondition = "SampleScale == EDataflowTransferVertexAttributeNodeSampleScale::Asset_Bound",
+		EditCondition = "SourceScale == EDataflowTransferVertexAttributeNodeSourceScale::Asset_Bound",
 		EditConditionHides))
 	float BoundMultiplier = 0.01f;
 
@@ -140,25 +128,65 @@ public:
 private:
 
 	virtual void Evaluate(UE::Dataflow::FContext& Context, const FDataflowOutput* Out) const override;
+};
+
+/**
+ * Transfer skin weights from a source collection to a target collection.
+ */
+USTRUCT(meta = (DataflowGeometryCollection))
+struct FGeometryCollectionTransferVertexSkinWeightsNode : public FDataflowNode
+{
+	GENERATED_USTRUCT_BODY()
+	DATAFLOW_NODE_DEFINE_INTERNAL(FGeometryCollectionTransferVertexSkinWeightsNode, "TransferVertexSkinWeights", "GeometryCollection", "Transfer vertex skin weights from the Source Collection to the Target Collection")
+	DATAFLOW_NODE_RENDER_TYPE("SurfaceRender", FGeometryCollection::StaticType(), "Collection")
+
+public:
+
+	/* Target collection to transfer vertex attribute to. */
+	UPROPERTY(Meta = (DataflowInput, DataflowOutput, DataflowPassthrough = "Collection"))
+	FManagedArrayCollection Collection;
+
+	/* Source collection to transfer vertex attribute from. */
+	UPROPERTY(Meta = (DataflowInput, DisplayName = "FromCollection"))
+	FManagedArrayCollection FromCollection;
+
+	/* Bounding volume type for source assets[default: Triangle] */
+	UPROPERTY(EditAnywhere, Category = "Thresholds")
+	EDataflowTransferVertexAttributeNodeBoundingVolume BoundingVolumeType = EDataflowTransferVertexAttributeNodeBoundingVolume::Triangle;
 
-	typedef Chaos::TSphere<Chaos::FReal, 3> SphereType;
-	typedef Chaos::TBoundingVolumeHierarchy<TArray<SphereType*>, TArray<int32>, Chaos::FReal, 3> BVH;
+	/* Bounding volume hierarchy cell size for neighboring vertices to transfer into[default: Asset] */
+	UPROPERTY(EditAnywhere, Category = "Thresholds")
+	EDataflowTransferVertexAttributeNodeSourceScale SourceScale = EDataflowTransferVertexAttributeNodeSourceScale::Asset_Bound;
 
-	//
-	// private helpers
-	//
-	TArray<FIntVector2> FindSourceToTargetGeometryMap(const FManagedArrayCollection& SourceCollection, const FManagedArrayCollection& TargetCollection) const;
-	TMap<int32, int32> FindSourceToTargetTransformMap(const FManagedArrayCollection& SourceCollection, const FManagedArrayCollection& TargetCollection) const;
+	/* Falloff of source value based on distance from source triangle[default: Squared] */
+	UPROPERTY(EditAnywhere, Category = "Thresholds")
+	EDataflowTransferVertexAttributeNodeFalloff Falloff = EDataflowTransferVertexAttributeNodeFalloff::None;
 
-	TArray<UE::Private::FTransferData> PairedGeometryTransfer(FCollectionAttributeKey Key, const TArray<FIntVector2>& PairedGeometry,
-		const UE::Private::FTransferFacade& Sample, const UE::Private::FTransferFacade& Target) const;
-	TArray<UE::Private::FTransferData> NearestVertexTransfer(FCollectionAttributeKey Key, const UE::Private::FTransferFacade& Sample, const UE::Private::FTransferFacade& Target) const;
-	static float MaxEdgeLength(TArray<FVector3f>& Vert, const TManagedArray<FIntVector3>& Tri, int VertexOFfset, int TriStart, int TriCount);
-	static void BuildComponentSpaceVertices(const TManagedArray<FTransform3f>& LocalSpaceTransform, const TManagedArray<int32>& Parent, const TManagedArray<int32>& BoneMapArray, const TManagedArray<FVector3f>& VertexArray, int32 Start, int32 Count, TArray<FVector3f>& ComponentSpaceVertices);
-	static BVH* BuildParticleSphereBVH(const TArray<FVector3f>& Vertices, float Radius);
-	static void TriangleToVertexIntersections(const BVH& VertexBVH, const TArray<FVector3f>& ComponentSpaceVertices, const FIntVector3& Triangle, TArray<int32>& OutTargetVertexIntersection);
-	static float CalculateFalloffScale(EDataflowTransferVertexAttributeNodeFalloff FalloffSetting, float Threshold, float Distance);
+	/* Threshold based on distance from source triangle.Values past the threshold will falloff.[Defaults to 1 percent of triangle size(0.01)] */
+	UPROPERTY(EditAnywhere, Category = "Thresholds")
+	float FalloffThreshold = 0.01f;
 
+	/* Edge multiplier for the Bounding Volume Hierarchy(BVH) target's particle search radius. */
+	UPROPERTY(EditAnywhere, Category = "Thresholds", meta = (EditCondition = "SourceScale == EDataflowTransferVertexAttributeNodeSourceScale::Asset_Edge || SourceScale == EDataflowTransferVertexAttributeNodeSourceScale::Component_Edge", EditConditionHides))
+	float EdgeMultiplier = 0.5f;
+
+	/* Max bound multiplier for the Bounding Volume Hierarchy(BVH) target's particle search radius. */
+	UPROPERTY(EditAnywhere, Category = "Thresholds", meta = (
+		EditCondition = "SourceScale == EDataflowTransferVertexAttributeNodeSourceScale::Asset_Bound",
+		EditConditionHides))
+	float BoundMultiplier = 0.01f;
+
+	FGeometryCollectionTransferVertexSkinWeightsNode(const UE::Dataflow::FNodeParameters& InParam, FGuid InGuid = FGuid::NewGuid())
+		: FDataflowNode(InParam, InGuid)
+	{
+		RegisterInputConnection(&Collection);
+		RegisterInputConnection(&FromCollection);
+		RegisterOutputConnection(&Collection, &Collection);
+	}
+
+private:
+
+	virtual void Evaluate(UE::Dataflow::FContext& Context, const FDataflowOutput* Out) const override;
 };
 
 /**
diff --git a/Engine/Source/Runtime/Experimental/Chaos/Private/GeometryCollection/Facades/CollectionTransformFacade.cpp b/Engine/Source/Runtime/Experimental/Chaos/Private/GeometryCollection/Facades/CollectionTransformFacade.cpp
--- a/Engine/Source/Runtime/Experimental/Chaos/Private/GeometryCollection/Facades/CollectionTransformFacade.cpp
+++ b/Engine/Source/Runtime/Experimental/Chaos/Private/GeometryCollection/Facades/CollectionTransformFacade.cpp
@@ -49,6 +49,18 @@ namespace GeometryCollection::Facades
 		return FString();
 	}
 
+	TMap<FString, int32> FCollectionTransformFacade::BoneNameIndexMap() const
+	{
+		TMap<FString, int32> BoneNameToIndex;
+		if (BoneNameAttribute.IsValid())
+		{
+			for (int32 BoneIdx = 0; BoneIdx < BoneNameAttribute.Num(); ++BoneIdx)
+			{
+				BoneNameToIndex.Add(BoneNameAttribute[BoneIdx], BoneIdx);
+			}
+		}
+		return MoveTemp(BoneNameToIndex);
+	}
 
 	TArray<int32> FCollectionTransformFacade::GetRootIndices() const
 	{
diff --git a/Engine/Source/Runtime/Experimental/Chaos/Public/Chaos/Deformable/GaussSeidelWeakConstraints.h b/Engine/Source/Runtime/Experimental/Chaos/Public/Chaos/Deformable/GaussSeidelWeakConstraints.h
--- a/Engine/Source/Runtime/Experimental/Chaos/Public/Chaos/Deformable/GaussSeidelWeakConstraints.h
+++ b/Engine/Source/Runtime/Experimental/Chaos/Public/Chaos/Deformable/GaussSeidelWeakConstraints.h
@@ -701,16 +701,18 @@ namespace Chaos::Softs
 			Resize(InitialWCSize + Particles.Size());
 			std::atomic<int32> ConstraintIndex(InitialWCSize);
 			const TArray<TVec3<int32>>& Elements = TriangleMesh.GetSurfaceElements();
+			const float HalfRadius = DetectRadius / 2;
 			PhysicsParallelFor(SurfaceVertices.Num(),
-				[this, &Spatial, &Particles, &SurfaceVertices, &ConstraintIndex, &TriangleMesh, &Elements, &DetectRadius, &ComponentIndex, &PositionTargetStiffness, &UseAnisotropicSpring](int32 i)
+				[this, &Spatial, &Particles, &SurfaceVertices, &ConstraintIndex, &TriangleMesh, &Elements, &HalfRadius, &ComponentIndex, &PositionTargetStiffness, &UseAnisotropicSpring](int32 i)
 				{
 					const int32 Index = SurfaceVertices[i];
 					TArray< TTriangleCollisionPoint<FSolverReal>> Result;
 					//PointProximityQuery
-					if (TriangleMesh.PointClosestTriangleQuery(Spatial, static_cast<const TArrayView<const FSolverVec3>&>(Particles.XArray()), Index, Particles.GetX(Index), DetectRadius, DetectRadius,
+					if (TriangleMesh.PointClosestTriangleQuery(Spatial, static_cast<const TArrayView<const FSolverVec3>&>(Particles.XArray()), Index, Particles.GetX(Index), HalfRadius, HalfRadius,
 						[this, &ComponentIndex, &Elements](const int32 PointIndex, const int32 TriangleIndex)->bool
 						{
-							return ComponentIndex[PointIndex] != ComponentIndex[Elements[TriangleIndex][0]];
+							//Skip particles that are bound in initial springs
+							return ComponentIndex[PointIndex] != ComponentIndex[Elements[TriangleIndex][0]] && (!NoCollisionWCIncidentElements.IsValidIndex(PointIndex) || NoCollisionWCIncidentElements[PointIndex].Num() == 0);
 						},
 						Result))
 					{
@@ -749,19 +751,20 @@ namespace Chaos::Softs
 		}
 
 		template<typename SpatialAccelerator>
-		void CollisionDetectionSpatialHashInComponent(const FSolverParticles& Particles, const TArray<int32>& SurfaceVertices, const FTriangleMesh& TriangleMesh, const TMap<int32, TSet<int32>>& ExcludeMap, const SpatialAccelerator& Spatial, float DetectRadius = 1.f, float PositionTargetStiffness = 10000.f, bool UseAnisotropicSpring = true)
+		void CollisionDetectionSpatialHashInComponent(const FSolverParticles& Particles, const TArray<int32>& SurfaceVertices, const FTriangleMesh& TriangleMesh, const TMap<int32, TSet<int32>>& ExcludeMap, const SpatialAccelerator& Spatial, float DetectRadius = 0.f, float PositionTargetStiffness = 10000.f, bool UseAnisotropicSpring = true)
 		{
 			TRACE_CPUPROFILER_EVENT_SCOPE(STAT_ChaosGaussSeidelWeakConstraintsCollisionDetectionSpatialHashInComponent);
 			Resize(InitialWCSize + Particles.Size());
 			std::atomic<int32> ConstraintIndex(InitialWCSize);
 			const TArray<TVec3<int32>>& Elements = TriangleMesh.GetSurfaceElements();
+			const float HalfRadius = DetectRadius/2;
 			PhysicsParallelFor(SurfaceVertices.Num(),
-				[this, &Spatial, &Particles, &SurfaceVertices, &ConstraintIndex, &TriangleMesh, &ExcludeMap, &Elements, &DetectRadius, &PositionTargetStiffness, &UseAnisotropicSpring](int32 i)
+				[this, &Spatial, &Particles, &SurfaceVertices, &ConstraintIndex, &TriangleMesh, &ExcludeMap, &Elements, &HalfRadius, &PositionTargetStiffness, &UseAnisotropicSpring](int32 i)
 				{
 					const int32 Index = SurfaceVertices[i];
 					TArray< TTriangleCollisionPoint<FSolverReal>> Result;
 					//PointProximityQuery
-					if (TriangleMesh.PointClosestTriangleQuery(Spatial, static_cast<const TArrayView<const FSolverVec3>&>(Particles.XArray()), Index, Particles.GetX(Index), DetectRadius, DetectRadius,
+					if (TriangleMesh.PointClosestTriangleQuery(Spatial, static_cast<const TArrayView<const FSolverVec3>&>(Particles.XArray()), Index, Particles.GetX(Index), HalfRadius, HalfRadius,
 						[this, &Elements, &ExcludeMap](const int32 PointIndex, const int32 TriangleIndex)->bool
 						{	
 							return  !(ExcludeMap.Find(PointIndex) && ExcludeMap[PointIndex].Contains(TriangleIndex));
diff --git a/Engine/Source/Runtime/Experimental/Chaos/Public/GeometryCollection/Facades/CollectionTransformFacade.h b/Engine/Source/Runtime/Experimental/Chaos/Public/GeometryCollection/Facades/CollectionTransformFacade.h
--- a/Engine/Source/Runtime/Experimental/Chaos/Public/GeometryCollection/Facades/CollectionTransformFacade.h
+++ b/Engine/Source/Runtime/Experimental/Chaos/Public/GeometryCollection/Facades/CollectionTransformFacade.h
@@ -80,6 +80,8 @@ namespace GeometryCollection::Facades
 		/** Get a bone name from the index if the facade has the attribute defined. */
 		CHAOS_API FString BoneName(int32 Index) const;
 
+		/** Get a TMap from bone name to bone index if the facade has the attribute defined. */
+		CHAOS_API TMap<FString, int32> BoneNameIndexMap() const;
 
 		/** Builds a FMatrix from all the components */
 		static CHAOS_API FMatrix BuildMatrix(const FVector& Translate,
diff --git a/Engine/Source/Runtime/Experimental/GeometryCollectionEngine/Private/GeometryCollection/GeometryCollectionEngineConversion.cpp b/Engine/Source/Runtime/Experimental/GeometryCollectionEngine/Private/GeometryCollection/GeometryCollectionEngineConversion.cpp
--- a/Engine/Source/Runtime/Experimental/GeometryCollectionEngine/Private/GeometryCollection/GeometryCollectionEngineConversion.cpp
+++ b/Engine/Source/Runtime/Experimental/GeometryCollectionEngine/Private/GeometryCollection/GeometryCollectionEngineConversion.cpp
@@ -1137,7 +1137,7 @@ void FGeometryCollectionEngineConversion::AppendGeometryCollection(const UGeomet
 
 
 bool FGeometryCollectionEngineConversion::AppendSkeletalMesh(const USkeletalMesh* InSkeletalMesh, int32 MaterialStartIndex, const FTransform& SkeletalMeshTransform, 
-	FManagedArrayCollection* InManagedArrayCollection, bool bReindexMaterials)
+	FManagedArrayCollection* InManagedArrayCollection, bool bReindexMaterials, bool bImportTransformOnly)
 {
 	//UE_LOG(UGeometryCollectionConversionLogging, Log, TEXT("FGeometryCollectionEngineConversion::AppendSkeletalMesh()"));
 #if WITH_EDITOR
@@ -1147,13 +1147,6 @@ bool FGeometryCollectionEngineConversion::AppendSkeletalMesh(const USkeletalMesh
 		return false;
 	}
 
-
-	FMeshDescription MeshDescription;
-	if (!InSkeletalMesh->CloneMeshDescription(LODIndex, MeshDescription))
-	{
-		return false;
-	}
-
 	FGeometryCollection GeometryCollection;
 
 	// Transform Attributes
@@ -1163,18 +1156,6 @@ bool FGeometryCollectionEngineConversion::AppendSkeletalMesh(const USkeletalMesh
 	TManagedArray<FLinearColor>& BoneColor = GeometryCollection.ModifyAttribute<FLinearColor>("BoneColor", FTransformCollection::TransformGroup);
 	TManagedArray<FString>& BoneName = GeometryCollection.ModifyAttribute<FString>("BoneName", FTransformCollection::TransformGroup);
 	TManagedArray<int32>& SimulationType = GeometryCollection.ModifyAttribute<int32>("SimulationType", FTransformCollection::TransformGroup);
-	// Vertices Attributes
-	TManagedArray<FVector3f>& Vertex = GeometryCollection.ModifyAttribute<FVector3f>("Vertex", FGeometryCollection::VerticesGroup);
-	TManagedArray<FVector3f>& Normal = GeometryCollection.ModifyAttribute<FVector3f>("Normal", FGeometryCollection::VerticesGroup);
-	TManagedArray<FLinearColor>& Color = GeometryCollection.ModifyAttribute<FLinearColor>("Color", FGeometryCollection::VerticesGroup);
-	TManagedArray<FVector3f>& TangentU = GeometryCollection.ModifyAttribute<FVector3f>("TangentU", FGeometryCollection::VerticesGroup);
-	TManagedArray<FVector3f>& TangentV = GeometryCollection.ModifyAttribute<FVector3f>("TangentV", FGeometryCollection::VerticesGroup);
-	TManagedArray<int32>& BoneMap = GeometryCollection.ModifyAttribute<int32>("BoneMap", FGeometryCollection::VerticesGroup);
-	// Index Attributes
-	TManagedArray<FIntVector>& Indices = GeometryCollection.ModifyAttribute<FIntVector>("Indices", FGeometryCollection::FacesGroup);
-	TManagedArray<bool>& Visible = GeometryCollection.ModifyAttribute<bool>("Visible", FGeometryCollection::FacesGroup);
-	TManagedArray<int32>& MaterialIndex = GeometryCollection.ModifyAttribute<int32>("MaterialIndex", FGeometryCollection::FacesGroup);
-	TManagedArray<int32>& MaterialID = GeometryCollection.ModifyAttribute<int32>("MaterialID", FGeometryCollection::FacesGroup);
 
 	//
 	// Convert the transform hierarchy
@@ -1183,7 +1164,7 @@ bool FGeometryCollectionEngineConversion::AppendSkeletalMesh(const USkeletalMesh
 	int32 TransformBaseIndex = INDEX_NONE;
 	const USkeleton* Skeleton = InSkeletalMesh->GetSkeleton();
 	const TArray<FTransform>& RestArray = Skeleton->GetRefLocalPoses();
-	const FReferenceSkeleton& ReferenceSkeleton = Skeleton->GetReferenceSkeleton();
+	const FReferenceSkeleton& ReferenceSkeleton = InSkeletalMesh->GetRefSkeleton();
 
 	if (ReferenceSkeleton.GetNum())
 	{
@@ -1219,6 +1200,30 @@ bool FGeometryCollectionEngineConversion::AppendSkeletalMesh(const USkeletalMesh
 		TransformSourceFacade.AddTransformSource(Skeleton->GetName(), Skeleton->GetGuid().ToString(), Roots);
 	}
 
+	if (bImportTransformOnly)
+	{
+		GeometryCollection.CopyTo(InManagedArrayCollection);
+		return true;
+	}
+
+	FMeshDescription MeshDescription;
+	if (!InSkeletalMesh->CloneMeshDescription(LODIndex, MeshDescription))
+	{
+		return false;
+	}
+
+	// Vertices Attributes
+	TManagedArray<FVector3f>& Vertex = GeometryCollection.ModifyAttribute<FVector3f>("Vertex", FGeometryCollection::VerticesGroup);
+	TManagedArray<FVector3f>& Normal = GeometryCollection.ModifyAttribute<FVector3f>("Normal", FGeometryCollection::VerticesGroup);
+	TManagedArray<FLinearColor>& Color = GeometryCollection.ModifyAttribute<FLinearColor>("Color", FGeometryCollection::VerticesGroup);
+	TManagedArray<FVector3f>& TangentU = GeometryCollection.ModifyAttribute<FVector3f>("TangentU", FGeometryCollection::VerticesGroup);
+	TManagedArray<FVector3f>& TangentV = GeometryCollection.ModifyAttribute<FVector3f>("TangentV", FGeometryCollection::VerticesGroup);
+	TManagedArray<int32>& BoneMap = GeometryCollection.ModifyAttribute<int32>("BoneMap", FGeometryCollection::VerticesGroup);
+	// Index Attributes
+	TManagedArray<FIntVector>& Indices = GeometryCollection.ModifyAttribute<FIntVector>("Indices", FGeometryCollection::FacesGroup);
+	TManagedArray<bool>& Visible = GeometryCollection.ModifyAttribute<bool>("Visible", FGeometryCollection::FacesGroup);
+	TManagedArray<int32>& MaterialIndex = GeometryCollection.ModifyAttribute<int32>("MaterialIndex", FGeometryCollection::FacesGroup);
+	TManagedArray<int32>& MaterialID = GeometryCollection.ModifyAttribute<int32>("MaterialID", FGeometryCollection::FacesGroup);
 
 	//
 	// Identify disconnected geoemtry
@@ -1491,7 +1496,8 @@ void FGeometryCollectionEngineConversion::AppendSkeletalMesh(const USkeletalMesh
 		if (FGeometryCollection* GeometryCollection = GeometryCollectionPtr.Get())
 		{
 			int32 MaterialStart = GeometryCollectionObject->Materials.Num();
-			if (AppendSkeletalMesh(SkeletalMesh, MaterialStart, SkeletalMeshTransform, GeometryCollection, bReindexMaterials))
+			constexpr bool bImportTransformOnly = false;
+			if (AppendSkeletalMesh(SkeletalMesh, MaterialStart, SkeletalMeshTransform, GeometryCollection, bReindexMaterials, bImportTransformOnly))
 			{
 				AppendSkeletalMeshMaterials(SkeletalMesh, SkeletalMeshComponent, GeometryCollectionObject);
 			}
@@ -1562,7 +1568,8 @@ bool FGeometryCollectionEngineConversion::AppendGeometryCollectionSourceNoMateri
 				StartMaterialIndex,
 				GeometryCollectionSource.LocalTransform,
 				&GeometryCollectionInOut,
-				ReindexMaterials
+				ReindexMaterials,
+				false
 				);
 			return true;
 		}
diff --git a/Engine/Source/Runtime/Experimental/GeometryCollectionEngine/Public/GeometryCollection/GeometryCollectionEngineConversion.h b/Engine/Source/Runtime/Experimental/GeometryCollectionEngine/Public/GeometryCollection/GeometryCollectionEngineConversion.h
--- a/Engine/Source/Runtime/Experimental/GeometryCollectionEngine/Public/GeometryCollection/GeometryCollectionEngineConversion.h
+++ b/Engine/Source/Runtime/Experimental/GeometryCollectionEngine/Public/GeometryCollection/GeometryCollectionEngineConversion.h
@@ -120,7 +120,7 @@ public:
 	*  @param SkeletalMeshTransform : Mesh transform.
 	*  @param GeometryCollection    : Collection to append the mesh into.
 	*/
-	static GEOMETRYCOLLECTIONENGINE_API bool AppendSkeletalMesh(const USkeletalMesh* SkeletalMesh, int32 MaterialStartIndex, const FTransform& SkeletalMeshTransform, FManagedArrayCollection* InCollection, bool bReindexMaterials = true);
+	static GEOMETRYCOLLECTIONENGINE_API bool AppendSkeletalMesh(const USkeletalMesh* SkeletalMesh, int32 MaterialStartIndex, const FTransform& SkeletalMeshTransform, FManagedArrayCollection* InCollection, bool bReindexMaterials = true, bool bImportTransformOnly = false);
 
 	/**
 	*  Appends a skeleton mesh to a GeometryCollection.
