CIS - Fixed issue 621501: Compile errors in UObjectGlobals.h and PointerIsConvertibleFromTo.h

#rnx

diff --git a/Engine/Plugins/Experimental/GeometryCollectionPlugin/Source/GeometryCollectionNodes/Private/Dataflow/GeometryCollectionSkeletonToCollectionNode.cpp b/Engine/Plugins/Experimental/GeometryCollectionPlugin/Source/GeometryCollectionNodes/Private/Dataflow/GeometryCollectionSkeletonToCollectionNode.cpp
--- a/Engine/Plugins/Experimental/GeometryCollectionPlugin/Source/GeometryCollectionNodes/Private/Dataflow/GeometryCollectionSkeletonToCollectionNode.cpp
+++ b/Engine/Plugins/Experimental/GeometryCollectionPlugin/Source/GeometryCollectionNodes/Private/Dataflow/GeometryCollectionSkeletonToCollectionNode.cpp
@@ -2,6 +2,7 @@
 
 #include "Dataflow/GeometryCollectionSkeletonToCollectionNode.h"
 
+#include "Animation/Skeleton.h"
 #include "Engine/SkeletalMesh.h"
 #include "GeometryCollection/GeometryCollectionObject.h"
 #include "GeometryCollection/ManagedArrayCollection.h"
