ChaosFlesh: support transform merging when appending collection
-added conversion from geometry selection to vertex selection
-cleared group elements in NewCopy() to clear out ConvexProperties group
-fixed index bug in GeometryCollection::ReorderTransformElements

diff --git a/Engine/Plugins/Experimental/ChaosFlesh/Source/ChaosFleshNodes/Private/Dataflow/ChaosFleshAppendTetrahedralCollectionNode.cpp b/Engine/Plugins/Experimental/ChaosFlesh/Source/ChaosFleshNodes/Private/Dataflow/ChaosFleshAppendTetrahedralCollectionNode.cpp
--- a/Engine/Plugins/Experimental/ChaosFlesh/Source/ChaosFleshNodes/Private/Dataflow/ChaosFleshAppendTetrahedralCollectionNode.cpp
+++ b/Engine/Plugins/Experimental/ChaosFlesh/Source/ChaosFleshNodes/Private/Dataflow/ChaosFleshAppendTetrahedralCollectionNode.cpp
@@ -4,6 +4,7 @@
 
 #include "ChaosFlesh/FleshCollection.h"
 #include "Dataflow/DataflowNodeFactory.h"
+#include "GeometryCollection/Facades/CollectionTransformFacade.h"
 
 void FAppendTetrahedralCollectionDataflowNode::Evaluate(UE::Dataflow::FContext& Context, const FDataflowOutput* Out) const
 {
@@ -16,17 +17,53 @@ void FAppendTetrahedralCollectionDataflowNode::Evaluate(UE::Dataflow::FContext&
 		{
 			GeometryGroupGuidsLocal1 = GuidArray1->GetConstArray();
 		}
-		if (InCollection2)
-		{
-			InCollection1->AppendCollection(*InCollection2);
-		}		
+		GeometryCollection::Facades::FCollectionTransformFacade InTransformFacade1(*InCollection1);
+		GeometryCollection::Facades::FCollectionTransformFacade InTransformFacade2(*InCollection2);
+		const TMap<FString, int32> InBoneNameMap = InTransformFacade1.BoneNameIndexMap();
+		const int32 InNumTransform2 = InTransformFacade2.Num();
+		const int32 InNumTransform1 = InTransformFacade1.Num();
+		//Append
+		InCollection1->AppendCollection(*InCollection2);
+
 		if (const TManagedArray<FString>* GuidArray2 = InCollection2->FindAttribute<FString>("Guid", FGeometryCollection::GeometryGroup))
 		{
 			GeometryGroupGuidsLocal2 = GuidArray2->GetConstArray();
 		}
-		SetValue<const FManagedArrayCollection&>(Context, *InCollection1, &Collection1);
+		if (bMergeTransform)
+		{
+			//Reorder and delete transform
+			const TManagedArray<FString>* OutBoneNames = InTransformFacade1.FindBoneNames();
+			TArray<int32> MergeRemapIndex;
+			TArray<int32> SortedMergeList;
+			for (int32 Idx = 0; Idx < InNumTransform2; ++Idx)
+			{
+				const FString& BoneName = (*OutBoneNames)[Idx];
+				if (Idx < InNumTransform2 && InBoneNameMap.Contains(BoneName))
+				{
+					SortedMergeList.Add(Idx);
+					MergeRemapIndex.Add(InBoneNameMap[BoneName] + InNumTransform2); //SKM collection is appended to the front
+				}
+			}
+			InCollection1->MergeElements(FTransformCollection::TransformGroup, SortedMergeList, MergeRemapIndex);
+		}
+		const int32 NumGeometries = InCollection1->NumElements(FGeometryCollection::GeometryGroup);
+		const int32 NumGeometries2 = InCollection2->NumElements(FGeometryCollection::GeometryGroup);
+		FDataflowGeometrySelection InGeometrySelection1;
+		FDataflowGeometrySelection InGeometrySelection2;
+		InGeometrySelection1.Initialize(NumGeometries, false);
+		InGeometrySelection2.Initialize(NumGeometries, false);
+		for (int32 GeometryIdx = 0; GeometryIdx < NumGeometries2; ++GeometryIdx)
+		{
+			InGeometrySelection2.SetSelected(GeometryIdx);
+		}
+		for (int32 GeometryIdx = NumGeometries2; GeometryIdx < NumGeometries; ++GeometryIdx)
+		{
+			InGeometrySelection1.SetSelected(GeometryIdx);
+		}
+		SetValue<const FManagedArrayCollection&>(Context, MoveTemp(*InCollection1), &Collection1);
+		SetValue<FDataflowGeometrySelection>(Context, MoveTemp(InGeometrySelection1), &GeometrySelection1);
+		SetValue<FDataflowGeometrySelection>(Context, MoveTemp(InGeometrySelection2), &GeometrySelection2);
 		SetValue(Context, MoveTemp(GeometryGroupGuidsLocal1), &GeometryGroupGuidsOut1);
 		SetValue(Context, MoveTemp(GeometryGroupGuidsLocal2), &GeometryGroupGuidsOut2);
 	}
-}
-
+}
\ No newline at end of file
diff --git a/Engine/Plugins/Experimental/ChaosFlesh/Source/ChaosFleshNodes/Private/Dataflow/ChaosFleshSetVertexTrianglePositionTargetBindingNode.cpp b/Engine/Plugins/Experimental/ChaosFlesh/Source/ChaosFleshNodes/Private/Dataflow/ChaosFleshSetVertexTrianglePositionTargetBindingNode.cpp
--- a/Engine/Plugins/Experimental/ChaosFlesh/Source/ChaosFleshNodes/Private/Dataflow/ChaosFleshSetVertexTrianglePositionTargetBindingNode.cpp
+++ b/Engine/Plugins/Experimental/ChaosFlesh/Source/ChaosFleshNodes/Private/Dataflow/ChaosFleshSetVertexTrianglePositionTargetBindingNode.cpp
@@ -146,7 +146,6 @@ void FSetVertexTrianglePositionTargetBindingDataflowNode::Evaluate(UE::Dataflow:
 										DataPackage.Stiffness += DataPackage.SourceWeights[k] * PositionTargetStiffness * (*Mass)[DataPackage.SourceIndex[k]];
 									}
 									DataPackage.Stiffness += DataPackage.TargetWeights[0] * PositionTargetStiffness * (*Mass)[DataPackage.TargetIndex[0]];
-									DataPackage.Stiffness /= 2.f;
 								}
 								else
 								{
diff --git a/Engine/Plugins/Experimental/ChaosFlesh/Source/ChaosFleshNodes/Public/Dataflow/ChaosFleshAppendTetrahedralCollectionNode.h b/Engine/Plugins/Experimental/ChaosFlesh/Source/ChaosFleshNodes/Public/Dataflow/ChaosFleshAppendTetrahedralCollectionNode.h
--- a/Engine/Plugins/Experimental/ChaosFlesh/Source/ChaosFleshNodes/Public/Dataflow/ChaosFleshAppendTetrahedralCollectionNode.h
+++ b/Engine/Plugins/Experimental/ChaosFlesh/Source/ChaosFleshNodes/Public/Dataflow/ChaosFleshAppendTetrahedralCollectionNode.h
@@ -9,9 +9,13 @@
 #include "GeometryCollection/ManagedArray.h"
 #include "GeometryCollection/ManagedArrayCollection.h"
 #include "GeometryCollection/GeometryCollection.h"
+#include "Dataflow/DataflowSelection.h"
 
 #include "ChaosFleshAppendTetrahedralCollectionNode.generated.h"
 
+/**
+ * Append another Tetrahedral Collection to this collection. All attributes will be copied.
+ */
 USTRUCT(meta = (DataflowFlesh))
 struct FAppendTetrahedralCollectionDataflowNode : public FDataflowNode
 {
@@ -25,11 +29,21 @@ public:
 	UPROPERTY(meta = (DataflowInput, DisplayName = "Collection2"));
 	FManagedArrayCollection Collection2;
 
-	UPROPERTY(meta = (DataflowOutput, DisplayName = "GeometryGroupIndicesOut1"))
-		TArray<FString> GeometryGroupGuidsOut1;
+	/* if transforms from Collection2 are merged into Collection1 by matching transform name */
+	UPROPERTY(EditAnywhere, Category = "Dataflow", meta = (DisplayName = "Merge Transform"));
+	bool bMergeTransform = false;
 
-	UPROPERTY(meta = (DataflowOutput, DisplayName = "GeometryGroupIndicesOut2"))
-		TArray<FString> GeometryGroupGuidsOut2;
+	UPROPERTY(meta = (DataflowOutput))
+	FDataflowGeometrySelection GeometrySelection1;
+
+	UPROPERTY(meta = (DataflowOutput))
+	FDataflowGeometrySelection GeometrySelection2;
+
+	UPROPERTY(meta = (DataflowOutput, DisplayName = "GeometryGroupGuids1"))
+	TArray<FString> GeometryGroupGuidsOut1;
+
+	UPROPERTY(meta = (DataflowOutput, DisplayName = "GeometryGroupGuids2"))
+	TArray<FString> GeometryGroupGuidsOut2;
 
 	FAppendTetrahedralCollectionDataflowNode(const UE::Dataflow::FNodeParameters& InParam, FGuid InGuid = FGuid::NewGuid())
 		: FDataflowNode(InParam, InGuid)
@@ -39,7 +53,9 @@ public:
 		RegisterOutputConnection(&Collection1, &Collection1);
 		RegisterOutputConnection(&GeometryGroupGuidsOut1);
 		RegisterOutputConnection(&GeometryGroupGuidsOut2);
+		RegisterOutputConnection(&GeometrySelection1);
+		RegisterOutputConnection(&GeometrySelection2);
 	}
 
 	virtual void Evaluate(UE::Dataflow::FContext& Context, const FDataflowOutput* Out) const override;
-};
+};
\ No newline at end of file
diff --git a/Engine/Plugins/Experimental/GeometryCollectionPlugin/Source/GeometryCollectionNodes/Private/Dataflow/GeometryCollectionSelectionNodes.cpp b/Engine/Plugins/Experimental/GeometryCollectionPlugin/Source/GeometryCollectionNodes/Private/Dataflow/GeometryCollectionSelectionNodes.cpp
--- a/Engine/Plugins/Experimental/GeometryCollectionPlugin/Source/GeometryCollectionNodes/Private/Dataflow/GeometryCollectionSelectionNodes.cpp
+++ b/Engine/Plugins/Experimental/GeometryCollectionPlugin/Source/GeometryCollectionNodes/Private/Dataflow/GeometryCollectionSelectionNodes.cpp
@@ -71,7 +71,7 @@ namespace UE::Dataflow
 		DATAFLOW_NODE_REGISTER_CREATION_FACTORY(FCollectionVertexSelectionByPercentageDataflowNode);
 		DATAFLOW_NODE_REGISTER_CREATION_FACTORY(FCollectionVertexSelectionSetOperationDataflowNode);
 		DATAFLOW_NODE_REGISTER_CREATION_FACTORY(FCollectionSelectionByAttrDataflowNode);
-
+		DATAFLOW_NODE_REGISTER_CREATION_FACTORY(FGeometrySelectionToVertexSelectionDataflowNode);
 		// GeometryCollection|Selection
 		DATAFLOW_NODE_REGISTER_CREATION_FACTORY_NODE_COLORS_BY_CATEGORY("GeometryCollection|Selection", FLinearColor(1.f, 1.f, 0.05f), CDefaultNodeBodyTintColor);
 	}
@@ -375,10 +375,10 @@ void FCollectionTransformSelectionCustomDataflowNode::Evaluate(UE::Dataflow::FCo
 
 			const FString InBoneIndices = GetValue<FString>(Context, &BoneIndicies);
 
-			TArray<FString> Indicies;
-			InBoneIndices.ParseIntoArray(Indicies, TEXT(" "), true);
+			TArray<FString> Indices;
+			InBoneIndices.ParseIntoArray(Indices, TEXT(" "), true);
 
-			for (FString IndexStr : Indicies)
+			for (FString IndexStr : Indices)
 			{
 				if (IndexStr.IsNumeric())
 				{
@@ -888,12 +888,12 @@ void FCollectionVertexSelectionCustomDataflowNode::Evaluate(UE::Dataflow::FConte
 			FDataflowVertexSelection NewVertexSelection;
 			NewVertexSelection.Initialize(NumVertices, false);
 
-			const FString InVertexIndicies = GetValue<FString>(Context, &VertexIndicies);
+			const FString InVertexIndices = GetValue<FString>(Context, &VertexIndicies);
 
-			TArray<FString> Indicies;
-			InVertexIndicies.ParseIntoArray(Indicies, TEXT(" "), true);
+			TArray<FString> Indices;
+			InVertexIndices.ParseIntoArray(Indices, TEXT(" "), true);
 
-			for (FString IndexStr : Indicies)
+			for (FString IndexStr : Indices)
 			{
 				if (IndexStr.IsNumeric())
 				{
@@ -938,16 +938,16 @@ void FCollectionFaceSelectionCustomDataflowNode::Evaluate(UE::Dataflow::FContext
 			FDataflowFaceSelection NewFaceSelection;
 			NewFaceSelection.Initialize(NumFaces, false);
 
-			const FString InFaceIndicies = GetValue<FString>(Context, &FaceIndicies);
+			const FString InFaceIndices = GetValue<FString>(Context, &FaceIndicies);
 
-			TArray<FString> Indicies;
-			InFaceIndicies.ParseIntoArray(Indicies, TEXT(" "), true);
+			TArray<FString> Indices;
+			InFaceIndices.ParseIntoArray(Indices, TEXT(" "), true);
 
-			for (FString IndexStr : Indicies)
+			for (FString& IndexStr : Indices)
 			{
 				if (IndexStr.IsNumeric())
 				{
-					int32 Index = FCString::Atoi(*IndexStr);
+					const int32 Index = FCString::Atoi(*IndexStr);
 					if (Index >= 0 && Index < NumFaces)
 					{
 						NewFaceSelection.SetSelected(Index);
@@ -955,8 +955,7 @@ void FCollectionFaceSelectionCustomDataflowNode::Evaluate(UE::Dataflow::FContext
 					else
 					{
 						// ERROR: INVALID INDEX
-						FString ErrorStr = "Invalid specified index found.";
-						UE_LOG(LogTemp, Error, TEXT("[Dataflow ERROR] %s"), *ErrorStr);
+						UE_LOG(LogTemp, Error, TEXT("[Dataflow ERROR] Invalid specified index found."));
 					}
 				}
 			}
@@ -1327,3 +1326,63 @@ void FCollectionSelectionByAttrDataflowNode::Evaluate(UE::Dataflow::FContext& Co
 	}
 }
 
+void FGeometrySelectionToVertexSelectionDataflowNode::Evaluate(UE::Dataflow::FContext& Context, const FDataflowOutput* Out) const
+{
+	if (Out->IsA<FDataflowVertexSelection>(&VertexSelection))
+	{
+		const FManagedArrayCollection& InCollection = GetValue<FManagedArrayCollection>(Context, &Collection);
+		const int32 NumVertices = InCollection.NumElements(FGeometryCollection::VerticesGroup);
+		const int32 NumGeometries = InCollection.NumElements(FGeometryCollection::GeometryGroup);
+		
+		FDataflowVertexSelection InVertexSelection;
+		InVertexSelection.Initialize(NumVertices, false);
+		const TManagedArray<int32>* VertexStart = InCollection.FindAttributeTyped<int32>("VertexStart", FGeometryCollection::GeometryGroup);
+		const TManagedArray<int32>* VertexCount = InCollection.FindAttributeTyped<int32>("VertexCount", FGeometryCollection::GeometryGroup);
+		TArray<int32> InGeometryIndexArray;
+		if (IsConnected(&GeometrySelection))
+		{
+			InGeometryIndexArray = GetValue<FDataflowGeometrySelection>(Context, &GeometrySelection).AsArray();
+		}
+		else
+		{
+			const FString InGeometryIndices = GetValue<FString>(Context, &GeometryIndices);
+			TArray<FString> Indices;
+			InGeometryIndices.ParseIntoArray(Indices, TEXT(" "), true);
+			for (FString IndexStr : Indices)
+			{
+				if (IndexStr.IsNumeric())
+				{
+					int32 Index = FCString::Atoi(*IndexStr);
+					if (Index >= 0 && Index < NumGeometries)
+					{
+						InGeometryIndexArray.Add(Index);
+					}
+					else
+					{
+						// ERROR: INVALID INDEX
+						FString ErrorStr = "Invalid geometry index found.";
+						UE_LOG(LogTemp, Error, TEXT("[Dataflow ERROR] %s"), *ErrorStr);
+					}
+				}
+			}
+		}
+		if (VertexStart && VertexCount)
+		{
+			TArray<int32> VertexIndices;
+			for (int32 GeometryIdx : InGeometryIndexArray)
+			{
+				if (ensure(VertexStart->IsValidIndex(GeometryIdx)))
+				{
+					const int32 Start = (*VertexStart)[GeometryIdx];
+					const int32 Count = (*VertexCount)[GeometryIdx];
+					for (int32 VertexIdx = Start; VertexIdx < Start + Count; ++VertexIdx)
+					{
+						VertexIndices.Add(VertexIdx);
+					}
+				}
+			}
+			InVertexSelection.SetFromArray(VertexIndices);
+		}
+		SetValue(Context, MoveTemp(InVertexSelection), &VertexSelection);
+	}
+}
\ No newline at end of file
diff --git a/Engine/Plugins/Experimental/GeometryCollectionPlugin/Source/GeometryCollectionNodes/Public/Dataflow/GeometryCollectionSelectionNodes.h b/Engine/Plugins/Experimental/GeometryCollectionPlugin/Source/GeometryCollectionNodes/Public/Dataflow/GeometryCollectionSelectionNodes.h
--- a/Engine/Plugins/Experimental/GeometryCollectionPlugin/Source/GeometryCollectionNodes/Public/Dataflow/GeometryCollectionSelectionNodes.h
+++ b/Engine/Plugins/Experimental/GeometryCollectionPlugin/Source/GeometryCollectionNodes/Public/Dataflow/GeometryCollectionSelectionNodes.h
@@ -37,7 +37,7 @@ public:
 	UPROPERTY(meta = (DataflowInput, DataflowOutput, DataflowPassthrough = "Collection", DataflowIntrinsic))
 	FManagedArrayCollection Collection;
 
-	/** Array of the selected bone indicies */
+	/** Array of the selected bone indices */
 	UPROPERTY(meta = (DataflowOutput, DisplayName = "TransformSelection"))
 	FDataflowTransformSelection TransformSelection;
 
@@ -86,15 +86,15 @@ public:
 	UPROPERTY(EditAnywhere, Category = "Compare");
 	ESetOperationEnum Operation = ESetOperationEnum::Dataflow_SetOperation_AND;
 
-	/** Array of the selected bone indicies */
+	/** Array of the selected bone indices */
 	UPROPERTY(meta = (DataflowInput, DisplayName = "TransformSelectionA", DataflowIntrinsic))
 	FDataflowTransformSelection TransformSelectionA;
 
-	/** Array of the selected bone indicies */
+	/** Array of the selected bone indices */
 	UPROPERTY(meta = (DataflowInput, DisplayName = "TransformSelectionB", DataflowIntrinsic))
 	FDataflowTransformSelection TransformSelectionB;
 
-	/** Array of the selected bone indicies after operation*/
+	/** Array of the selected bone indices after operation*/
 	UPROPERTY(meta = (DataflowOutput, DisplayName = "TransformSelection", DataflowPassthrough = "TransformSelectionA"))
 	FDataflowTransformSelection TransformSelection;
 
@@ -123,7 +123,7 @@ struct FCollectionTransformSelectionInfoDataflowNode : public FDataflowNode
 	DATAFLOW_NODE_DEFINE_INTERNAL(FCollectionTransformSelectionInfoDataflowNode, "CollectionTransformSelectionInfo", "GeometryCollection|Selection|Transform", "")
 
 public:
-	/** Array of the selected bone indicies */
+	/** Array of the selected bone indices */
 	UPROPERTY(meta = (DataflowInput, DisplayName = "TransformSelection", DataflowIntrinsic))
 	FDataflowTransformSelection TransformSelection;
 
@@ -164,7 +164,7 @@ public:
 	UPROPERTY(meta = (DataflowInput, DataflowOutput, DataflowPassthrough = "Collection", DataflowIntrinsic))
 	FManagedArrayCollection Collection;
 
-	/** Array of the selected bone indicies */
+	/** Array of the selected bone indices */
 	UPROPERTY(meta = (DataflowOutput, DisplayName = "TransformSelection"))
 	FDataflowTransformSelection TransformSelection;
 
@@ -193,7 +193,7 @@ struct FCollectionTransformSelectionInvertDataflowNode : public FDataflowNode
 	DATAFLOW_NODE_DEFINE_INTERNAL(FCollectionTransformSelectionInvertDataflowNode, "CollectionTransformSelectInvert", "GeometryCollection|Selection|Transform", "")
 
 public:
-	/** Array of the selected bone indicies */
+	/** Array of the selected bone indices */
 	UPROPERTY(meta = (DataflowInput, DataflowOutput, DisplayName = "TransformSelection", DataflowPassthrough = "TransformSelection", DataflowIntrinsic))
 	FDataflowTransformSelection TransformSelection;
 
@@ -237,7 +237,7 @@ public:
 	UPROPERTY(meta = (DataflowInput, DataflowOutput, DataflowPassthrough = "Collection", DataflowIntrinsic))
 	FManagedArrayCollection Collection;
 
-	/** Array of the selected bone indicies */
+	/** Array of the selected bone indices */
 	UPROPERTY(meta = (DataflowOutput, DisplayName = "TransformSelection"))
 	FDataflowTransformSelection TransformSelection;
 
@@ -273,7 +273,7 @@ public:
 	UPROPERTY(meta = (DataflowInput, DataflowOutput, DataflowPassthrough = "Collection", DataflowIntrinsic))
 	FManagedArrayCollection Collection;
 
-	/** Array of the selected bone indicies */
+	/** Array of the selected bone indices */
 	UPROPERTY(meta = (DataflowOutput, DisplayName = "TransformSelection"))
 	FDataflowTransformSelection TransformSelection;
 
@@ -309,7 +309,7 @@ public:
 
 	/** Space separated list of bone indices to specify the selection */
 	UPROPERTY(EditAnywhere, Category = "Selection", meta=(DisplayName="Bone Indices"))
-	FString BoneIndicies = FString();
+	FString BoneIndicies = FString(); //Fix typo for v2
 
 	/** Array of the selected bone indices */
 	UPROPERTY(meta = (DataflowOutput, DisplayName = "TransformSelection"))
@@ -379,7 +379,7 @@ struct FCollectionTransformSelectionParentDataflowNode : public FDataflowNode
 	DATAFLOW_NODE_DEFINE_INTERNAL(FCollectionTransformSelectionParentDataflowNode, "CollectionTransformSelectParent", "GeometryCollection|Selection|Transform", "")
 
 public:
-	/** Array of the selected bone indicies */
+	/** Array of the selected bone indices */
 	UPROPERTY(meta = (DataflowInput, DataflowOutput, DisplayName = "TransformSelection", DataflowPassthrough = "TransformSelection", DataflowIntrinsic))
 	FDataflowTransformSelection TransformSelection;
 
@@ -413,7 +413,7 @@ struct FCollectionTransformSelectionByPercentageDataflowNode : public FDataflowN
 	DATAFLOW_NODE_DEFINE_INTERNAL(FCollectionTransformSelectionByPercentageDataflowNode, "CollectionTransformSelectByPercentage", "GeometryCollection|Selection|Transform", "")
 
 public:
-	/** Array of the selected bone indicies */
+	/** Array of the selected bone indices */
 	UPROPERTY(meta = (DataflowInput, DataflowOutput, DisplayName = "TransformSelection", DataflowPassthrough = "TransformSelection", DataflowIntrinsic))
 	FDataflowTransformSelection TransformSelection;
 
@@ -456,7 +456,7 @@ struct FCollectionTransformSelectionChildrenDataflowNode : public FDataflowNode
 	DATAFLOW_NODE_DEFINE_INTERNAL(FCollectionTransformSelectionChildrenDataflowNode, "CollectionTransformSelectChildren", "GeometryCollection|Selection|Transform", "")
 
 public:
-	/** Array of the selected bone indicies */
+	/** Array of the selected bone indices */
 	UPROPERTY(meta = (DataflowInput, DataflowOutput, DisplayName = "TransformSelection", DataflowPassthrough = "TransformSelection", DataflowIntrinsic))
 	FDataflowTransformSelection TransformSelection;
 
@@ -490,7 +490,7 @@ struct FCollectionTransformSelectionSiblingsDataflowNode : public FDataflowNode
 	DATAFLOW_NODE_DEFINE_INTERNAL(FCollectionTransformSelectionSiblingsDataflowNode, "CollectionTransformSelectSiblings", "GeometryCollection|Selection|Transform", "")
 
 public:
-	/** Array of the selected bone indicies */
+	/** Array of the selected bone indices */
 	UPROPERTY(meta = (DataflowInput, DataflowOutput, DisplayName = "TransformSelection", DataflowPassthrough = "TransformSelection", DataflowIntrinsic))
 	FDataflowTransformSelection TransformSelection;
 
@@ -601,7 +601,7 @@ struct FCollectionTransformSelectionContactDataflowNode : public FDataflowNode
 	DATAFLOW_NODE_DEFINE_INTERNAL(FCollectionTransformSelectionContactDataflowNode, "CollectionTransformSelectContact", "GeometryCollection|Selection|Transform", "")
 
 public:
-	/** Array of the selected bone indicies */
+	/** Array of the selected bone indices */
 	UPROPERTY(meta = (DataflowInput, DataflowOutput, DisplayName = "TransformSelection", DataflowPassthrough = "TransformSelection", DataflowIntrinsic))
 	FDataflowTransformSelection TransformSelection;
 
@@ -643,7 +643,7 @@ public:
 	UPROPERTY(meta = (DataflowInput, DataflowOutput, DataflowPassthrough = "Collection", DataflowIntrinsic))
 	FManagedArrayCollection Collection;
 
-	/** Array of the selected bone indicies */
+	/** Array of the selected bone indices */
 	UPROPERTY(meta = (DataflowOutput, DisplayName = "TransformSelection"))
 	FDataflowTransformSelection TransformSelection;
 
@@ -678,7 +678,7 @@ public:
 	UPROPERTY(meta = (DataflowInput, DataflowOutput, DataflowPassthrough = "Collection", DataflowIntrinsic))
 	FManagedArrayCollection Collection;
 
-	/** Array of the selected bone indicies */
+	/** Array of the selected bone indices */
 	UPROPERTY(meta = (DataflowOutput, DisplayName = "TransformSelection"))
 	FDataflowTransformSelection TransformSelection;
 
@@ -711,7 +711,7 @@ public:
 	UPROPERTY(meta = (DataflowInput, DataflowOutput, DataflowPassthrough = "Collection", DataflowIntrinsic))
 	FManagedArrayCollection Collection;
 
-	/** Array of the selected bone indicies */
+	/** Array of the selected bone indices */
 	UPROPERTY(meta = (DataflowOutput, DisplayName = "TransformSelection"))
 	FDataflowTransformSelection TransformSelection;
 
@@ -827,7 +827,7 @@ public:
 	UPROPERTY(EditAnywhere, Category = "Size")
 	bool bUseRelativeSize = true;
 
-	/** Array of the selected bone indicies */
+	/** Array of the selected bone indices */
 	UPROPERTY(meta = (DataflowOutput, DisplayName = "TransformSelection"))
 	FDataflowTransformSelection TransformSelection;
 
@@ -878,7 +878,7 @@ public:
 	UPROPERTY(EditAnywhere, Category = "Volume")
 	bool bInclusive = true;
 
-	/** Array of the selected bone indicies */
+	/** Array of the selected bone indices */
 	UPROPERTY(meta = (DataflowOutput, DisplayName = "TransformSelection"))
 	FDataflowTransformSelection TransformSelection;
 
@@ -943,7 +943,7 @@ public:
 	UPROPERTY(EditAnywhere, Category = "Select", meta = (EditCondition = "Type == ESelectSubjectTypeEnum::Dataflow_SelectSubjectType_Vertices"))
 	bool bAllVerticesMustContainedInBox = true;
 
-	/** Array of the selected bone indicies */
+	/** Array of the selected bone indices */
 	UPROPERTY(meta = (DataflowOutput, DisplayName = "TransformSelection"))
 	FDataflowTransformSelection TransformSelection;
 
@@ -994,7 +994,7 @@ public:
 	UPROPERTY(EditAnywhere, Category = "Select", meta = (EditCondition = "Type == ESelectSubjectTypeEnum::Dataflow_SelectSubjectType_Vertices"))
 	bool bAllVerticesMustContainedInSphere = true;
 
-	/** Array of the selected bone indicies */
+	/** Array of the selected bone indices */
 	UPROPERTY(meta = (DataflowOutput, DisplayName = "TransformSelection"))
 	FDataflowTransformSelection TransformSelection;
 
@@ -1053,7 +1053,7 @@ public:
 	UPROPERTY(EditAnywhere, Category = "Attribute")
 	bool bInclusive = true;
 
-	/** Array of the selected bone indicies */
+	/** Array of the selected bone indices */
 	UPROPERTY(meta = (DataflowOutput, DisplayName = "TransformSelection"))
 	FDataflowTransformSelection TransformSelection;
 
@@ -1112,7 +1112,7 @@ public:
 	UPROPERTY(EditAnywhere, Category = "Attribute")
 	bool bInclusive = true;
 
-	/** Transform selection including the new indicies */
+	/** Transform selection including the new indices */
 	UPROPERTY(meta = (DataflowOutput, DisplayName = "TransformSelection"))
 	FDataflowTransformSelection TransformSelection;
 
@@ -1148,11 +1148,11 @@ public:
 	UPROPERTY(meta = (DataflowInput, DataflowOutput, DataflowPassthrough = "Collection", DataflowIntrinsic))
 	FManagedArrayCollection Collection;
 
-	/** Space separated list of vertex indicies to specify the selection */
+	/** Space separated list of vertex indices to specify the selection */
 	UPROPERTY(EditAnywhere, Category = "Selection")
-	FString VertexIndicies = FString();
+	FString VertexIndicies = FString(); //Fix typo for v2
 
-	/** Vertex selection including the new indicies */
+	/** Vertex selection including the new indices */
 	UPROPERTY(meta = (DataflowOutput, DisplayName = "VertexSelection"))
 	FDataflowVertexSelection VertexSelection;
 
@@ -1187,11 +1187,11 @@ public:
 	UPROPERTY(meta = (DataflowInput, DataflowOutput, DataflowPassthrough = "Collection", DataflowIntrinsic))
 	FManagedArrayCollection Collection;
 
-	/** Space separated list of face indicies to specify the selection */
+	/** Space separated list of face indices to specify the selection */
 	UPROPERTY(EditAnywhere, Category = "Selection")
-	FString FaceIndicies = FString();
+	FString FaceIndicies = FString(); //Fix typo for v2
 
-	/** Face selection including the new indicies */
+	/** Face selection including the new indices */
 	UPROPERTY(meta = (DataflowOutput, DisplayName = "FaceSelection"))
 	FDataflowFaceSelection FaceSelection;
 
@@ -1225,15 +1225,15 @@ public:
 	UPROPERTY(meta = (DataflowInput, DataflowOutput, DataflowPassthrough = "Collection", DataflowIntrinsic))
 	FManagedArrayCollection Collection;
 
-	/** Transform selection including the new indicies */
+	/** Transform selection including the new indices */
 	UPROPERTY(meta = (DataflowInput, DataflowOutput, DataflowPassthrough = "TransformSelection", DisplayName = "TransformSelection"))
 	FDataflowTransformSelection TransformSelection;
 
-	/** Face selection including the new indicies */
+	/** Face selection including the new indices */
 	UPROPERTY(meta = (DataflowInput, DataflowOutput, DataflowPassthrough = "FaceSelection", DisplayName = "FaceSelection"))
 	FDataflowFaceSelection FaceSelection;
 
-	/** Vertex selection including the new indicies */
+	/** Vertex selection including the new indices */
 	UPROPERTY(meta = (DataflowInput, DataflowOutput, DataflowPassthrough = "VertexSelection", DisplayName = "VertexSelection"))
 	FDataflowVertexSelection VertexSelection;
 	
@@ -1271,7 +1271,7 @@ struct FCollectionFaceSelectionInvertDataflowNode : public FDataflowNode
 	DATAFLOW_NODE_DEFINE_INTERNAL(FCollectionFaceSelectionInvertDataflowNode, "CollectionFaceSelectInvert", "GeometryCollection|Selection|Face", "")
 
 public:
-	/** Array of the selected bone indicies */
+	/** Array of the selected bone indices */
 	UPROPERTY(meta = (DataflowInput, DataflowOutput, DisplayName = "FaceSelection", DataflowPassthrough = "FaceSelection", DataflowIntrinsic))
 	FDataflowFaceSelection FaceSelection;
 
@@ -1299,7 +1299,7 @@ struct FCollectionVertexSelectionByPercentageDataflowNode : public FDataflowNode
 	DATAFLOW_NODE_DEFINE_INTERNAL(FCollectionVertexSelectionByPercentageDataflowNode, "CollectionVertexSelectByPercentage", "GeometryCollection|Selection|Vertex", "")
 
 public:
-	/** Array of the selected bone indicies */
+	/** Array of the selected bone indices */
 	UPROPERTY(meta = (DataflowInput, DataflowOutput, DisplayName = "VertexSelection", DataflowPassthrough = "VertexSelection", DataflowIntrinsic))
 	FDataflowVertexSelection VertexSelection;
 
@@ -1346,15 +1346,15 @@ public:
 	UPROPERTY(EditAnywhere, Category = "Compare");
 	ESetOperationEnum Operation = ESetOperationEnum::Dataflow_SetOperation_AND;
 
-	/** Array of the selected vertex indicies */
+	/** Array of the selected vertex indices */
 	UPROPERTY(meta = (DataflowInput, DisplayName = "VertexSelectionA", DataflowIntrinsic))
 	FDataflowVertexSelection VertexSelectionA;
 
-	/** Array of the selected vertex indicies */
+	/** Array of the selected vertex indices */
 	UPROPERTY(meta = (DataflowInput, DisplayName = "VertexSelectionB", DataflowIntrinsic))
 	FDataflowVertexSelection VertexSelectionB;
 
-	/** Array of the selected vertex indicies after operation */
+	/** Array of the selected vertex indices after operation */
 	UPROPERTY(meta = (DataflowOutput, DisplayName = "VertexSelection", DataflowPassthrough = "VertexSelectionA"))
 	FDataflowVertexSelection VertexSelection;
 
@@ -1479,6 +1479,44 @@ public:
 	virtual void Evaluate(UE::Dataflow::FContext& Context, const FDataflowOutput* Out) const override;
 };
 
+/**
+ * Converts GeometrySelection to VertexSelection
+ */
+USTRUCT(meta = (DataflowGeometryCollection))
+struct FGeometrySelectionToVertexSelectionDataflowNode final : public FDataflowNode
+{
+	GENERATED_USTRUCT_BODY()
+	DATAFLOW_NODE_DEFINE_INTERNAL(FGeometrySelectionToVertexSelectionDataflowNode, "GeometrySelectionToVertexSelection", "GeometryCollection|Selection|All", "")
+
+public:
+	/** GeometryCollection */
+	UPROPERTY(meta = (DataflowInput))
+	FManagedArrayCollection Collection;
+
+	/** Space separated list of geometry indices to specify the selection when GeometrySelection is not connected*/
+	UPROPERTY(EditAnywhere, Category = "Selection")
+	FString GeometryIndices = FString();
+
+	/** Input geometry selection */
+	UPROPERTY(meta = (DataflowInput, DisplayName = "GeometrySelection"))
+	FDataflowGeometrySelection GeometrySelection;
+
+	/** Vertex selection output */
+	UPROPERTY(meta = (DataflowOutput, DisplayName = "VertexSelection"))
+	FDataflowVertexSelection VertexSelection;
+
+	FGeometrySelectionToVertexSelectionDataflowNode(const UE::Dataflow::FNodeParameters& InParam, FGuid InGuid = FGuid::NewGuid())
+		: FDataflowNode(InParam, InGuid)
+	{
+		RegisterInputConnection(&Collection);
+		RegisterInputConnection(&GeometrySelection);
+		RegisterOutputConnection(&VertexSelection);
+	}
+
+private:
+	virtual void Evaluate(UE::Dataflow::FContext& Context, const FDataflowOutput* Out) const override;
+};
+
 namespace UE::Dataflow
 {
 	void GeometryCollectionSelectionNodes();
diff --git a/Engine/Source/Runtime/Experimental/Chaos/Private/GeometryCollection/Facades/CollectionTetrahedralBindingsFacade.cpp b/Engine/Source/Runtime/Experimental/Chaos/Private/GeometryCollection/Facades/CollectionTetrahedralBindingsFacade.cpp
--- a/Engine/Source/Runtime/Experimental/Chaos/Private/GeometryCollection/Facades/CollectionTetrahedralBindingsFacade.cpp
+++ b/Engine/Source/Runtime/Experimental/Chaos/Private/GeometryCollection/Facades/CollectionTetrahedralBindingsFacade.cpp
@@ -120,6 +120,9 @@ namespace GeometryCollection::Facades
 	const FName FTetrahedralBindings::OffsetsAttributeName = "Offsets";
 	const FName FTetrahedralBindings::MaskAttributeName = "Mask";
 
+	// Dependency
+	const FName FTetrahedralBindings::TetrahedralGroupDependency = TEXT("Tetrahedral");
+
 	FTetrahedralBindings::FTetrahedralBindings(FManagedArrayCollection& InCollection)
 		: MeshIdAttribute(InCollection, MeshIdAttributeName, MeshBindingsGroupName)
 	{}
@@ -216,7 +219,7 @@ namespace GeometryCollection::Facades
 		Offsets.Reset();
 		Masks.Reset();
 		FManagedArrayCollection& Collection = *MeshIdAttribute.GetCollection();
-		Parents.Reset(new TManagedArrayAccessor<FIntVector4>(Collection, ParentsAttributeName, GroupName));
+		Parents.Reset(new TManagedArrayAccessor<FIntVector4>(Collection, ParentsAttributeName, GroupName, TetrahedralGroupDependency));
 		Weights.Reset(new TManagedArrayAccessor<FVector4f>(Collection, WeightsAttributeName, GroupName));
 		Offsets.Reset(new TManagedArrayAccessor<FVector3f>(Collection, OffsetsAttributeName, GroupName));
 		Masks.Reset(new TManagedArrayAccessor<float>(Collection, MaskAttributeName, GroupName));
@@ -246,7 +249,7 @@ namespace GeometryCollection::Facades
 		if (!IsConst())
 		{
 			FManagedArrayCollection* Collection = MeshIdAttribute.GetCollection();
-			Parents.Reset(new TManagedArrayAccessor<FIntVector4>(*Collection, ParentsAttributeName, GroupName));
+			Parents.Reset(new TManagedArrayAccessor<FIntVector4>(*Collection, ParentsAttributeName, GroupName, TetrahedralGroupDependency));
 			Weights.Reset(new TManagedArrayAccessor<FVector4f>(*Collection, WeightsAttributeName, GroupName));
 			Offsets.Reset(new TManagedArrayAccessor<FVector3f>(*Collection, OffsetsAttributeName, GroupName));
 			Masks.Reset(new TManagedArrayAccessor<float>(*Collection, MaskAttributeName, GroupName));
@@ -254,7 +257,7 @@ namespace GeometryCollection::Facades
 		else
 		{
 			const FManagedArrayCollection& ConstCollection = MeshIdAttribute.GetConstCollection();
-			Parents.Reset(new TManagedArrayAccessor<FIntVector4>(ConstCollection, ParentsAttributeName, GroupName));
+			Parents.Reset(new TManagedArrayAccessor<FIntVector4>(ConstCollection, ParentsAttributeName, GroupName, TetrahedralGroupDependency));
 			Weights.Reset(new TManagedArrayAccessor<FVector4f>(ConstCollection, WeightsAttributeName, GroupName));
 			Offsets.Reset(new TManagedArrayAccessor<FVector3f>(ConstCollection, OffsetsAttributeName, GroupName));
 			Masks.Reset(new TManagedArrayAccessor<float>(ConstCollection, MaskAttributeName, GroupName));
diff --git a/Engine/Source/Runtime/Experimental/Chaos/Private/GeometryCollection/GeometryCollection.cpp b/Engine/Source/Runtime/Experimental/Chaos/Private/GeometryCollection/GeometryCollection.cpp
--- a/Engine/Source/Runtime/Experimental/Chaos/Private/GeometryCollection/GeometryCollection.cpp
+++ b/Engine/Source/Runtime/Experimental/Chaos/Private/GeometryCollection/GeometryCollection.cpp
@@ -802,7 +802,7 @@ void FGeometryCollection::ReorderTransformElements(const TArray<int32>& NewOrder
 		// remap the parents (-1 === Invalid )
 		if (Parent[Index] != -1)
 		{
-			Parent[Index] -= NewOrder[Parent[Index]];
+			Parent[Index] = NewOrder[Parent[Index]];
 		}
 
 		// remap children
diff --git a/Engine/Source/Runtime/Experimental/Chaos/Private/GeometryCollection/ManagedArrayCollection.cpp b/Engine/Source/Runtime/Experimental/Chaos/Private/GeometryCollection/ManagedArrayCollection.cpp
--- a/Engine/Source/Runtime/Experimental/Chaos/Private/GeometryCollection/ManagedArrayCollection.cpp
+++ b/Engine/Source/Runtime/Experimental/Chaos/Private/GeometryCollection/ManagedArrayCollection.cpp
@@ -84,6 +84,34 @@ void FManagedArrayCollection::RemoveElements(const FName& Group, const TArray<in
 	}
 }
 
+void FManagedArrayCollection::MergeElements(const FName& Group, const TArray<int32>& SortedMergeList, const TArray<int32>& MergeRemapIndex, FProcessingParameters Params)
+{
+	check(SortedMergeList.Num() == MergeRemapIndex.Num());
+	if (SortedMergeList.Num())
+	{
+		const int32 GroupSize = GroupInfo[Group].Size;
+		TArray<int32> InverseNewOrder;
+		InverseNewOrder.Init(INDEX_NONE, GroupSize);
+		for (int32 Idx = 0; Idx < GroupSize; ++Idx)
+		{
+			InverseNewOrder[Idx] = Idx;
+		}
+		for (int32 DeletedIdx = 0; DeletedIdx < MergeRemapIndex.Num(); ++DeletedIdx)
+		{
+			InverseNewOrder[SortedMergeList[DeletedIdx]] = MergeRemapIndex[DeletedIdx];
+		}
+		for (TTuple<FKeyType, FValueType>& Entry : Map)
+		{
+			// Reindex attributes dependent deleted elements
+			if (Entry.Value.GetGroupIndexDependency() == Group)
+			{
+				Entry.Value.Modify().ReindexFromLookup(InverseNewOrder);
+			}
+		}
+		FManagedArrayCollection::RemoveElements(Group, SortedMergeList, Params);
+	}
+}
+
 void FManagedArrayCollection::RemoveElements(const FName& Group, int32 NumberElements, int32 Position)
 {
 	TArray<int32> SortedDeletionList;
@@ -412,7 +440,7 @@ void FManagedArrayCollection::ReorderElements(FName Group, const TArray<int32>&
 	check(GroupSize == NewOrder.Num());
 
 	TArray<int32> InverseNewOrder;
-	InverseNewOrder.Init(-1, GroupSize);
+	InverseNewOrder.Init(INDEX_NONE, GroupSize);
 	for (int32 Idx = 0; Idx < GroupSize; ++Idx)
 	{
 		InverseNewOrder[NewOrder[Idx]] = Idx;
diff --git a/Engine/Source/Runtime/Experimental/Chaos/Private/GeometryCollection/TransformCollection.cpp b/Engine/Source/Runtime/Experimental/Chaos/Private/GeometryCollection/TransformCollection.cpp
--- a/Engine/Source/Runtime/Experimental/Chaos/Private/GeometryCollection/TransformCollection.cpp
+++ b/Engine/Source/Runtime/Experimental/Chaos/Private/GeometryCollection/TransformCollection.cpp
@@ -199,12 +199,12 @@ void FTransformCollection::RemoveElements(const FName & Group, const TArray<int3
 			TManagedArray<int32>& ParentArray = Parent;
 			TManagedArray<TSet<int32>>& ChildrenArray = Children;
 			TManagedArray<FTransform3f>&  LocalTransform = Transform;
-			for (int32 sdx = 0; sdx < SortedDeletionList.Num(); sdx++)
+			for (int32 Sdx = 0; Sdx < SortedDeletionList.Num(); Sdx++)
 			{
 				TArray<FTransform> GlobalTransform;
 				GeometryCollectionAlgo::GlobalMatrices(Transform, Parent, GlobalTransform);
 
-				int32 Index = SortedDeletionList[sdx];
+				int32 Index = SortedDeletionList[Sdx];
 				ensure(0 <= Index && Index < ParentArray.Num());
 
 				int32 ParentID = ParentArray[Index];
@@ -268,3 +268,78 @@ void FTransformCollection::RemoveElements(const FName & Group, const TArray<int3
 
 	}
 }
+
+void FTransformCollection::MergeElements(const FName& Group, const TArray<int32>& SortedMergeList, const TArray<int32>& MergeRemapIndex, FProcessingParameters Params)
+{
+	if (SortedMergeList.Num() && ensure(SortedMergeList.Num() == MergeRemapIndex.Num()))
+	{
+		if (Group == FTransformCollection::TransformGroup)
+		{
+			GeometryCollectionAlgo::ValidateSortedList(SortedMergeList, NumElements(Group));
+
+			TManagedArray<int32>& ParentArray = Parent;
+			TManagedArray<TSet<int32>>& ChildrenArray = Children;
+			TManagedArray<FTransform3f>& LocalTransform = Transform;
+			for (int32 Sdx = 0; Sdx < SortedMergeList.Num(); ++Sdx)
+			{
+				const int32 Index = SortedMergeList[Sdx];
+				if (ensure(0 <= Index && Index < ParentArray.Num()))
+				{
+					const int32 RemapId = MergeRemapIndex[Index];
+					const int32 ParentID = ParentArray[Index];
+					if (ensure(0 <= RemapId && RemapId < ParentArray.Num()))
+					{
+						for (const int32 ChildID : ChildrenArray[Index])
+						{
+							ParentArray[ChildID] = RemapId;
+							ChildrenArray[RemapId].Add(ChildID);
+						}
+						if (0 <= ParentID && ParentID < ChildrenArray.Num())
+						{
+							ChildrenArray[ParentID].Remove(Index);
+							ChildrenArray[ParentID].Add(RemapId);
+						}
+						ParentArray[Index] = INDEX_NONE;
+						ChildrenArray[Index].Empty();
+					}
+				}
+			}
+
+			// reindex
+			{
+				TArray<int32> Offsets;
+				GeometryCollectionAlgo::BuildIncrementMask(SortedMergeList, ParentArray.Num(), Offsets);
+
+				const int32 ArraySize = ParentArray.Num(), OffsetsSize = Offsets.Num();
+				const int32 FinalSize = ParentArray.Num() - SortedMergeList.Num();
+				for (int32 Index = 0; Index < ArraySize; Index++)
+				{
+					// remap the parents (-1 === Invalid )
+					if (ParentArray[Index] != -1)
+					{
+						ParentArray[Index] -= Offsets[ParentArray[Index]];
+					}
+					ensure(-1 <= ParentArray[Index] && ParentArray[Index] <= FinalSize);
+
+					// remap children
+					TSet<int32> ChildrenCopy = ChildrenArray[Index];
+					ChildrenArray[Index].Reset();
+					for (const int32 ChildID : ChildrenCopy)
+					{
+						if (0 <= ChildID && ChildID < OffsetsSize)
+						{
+							const int32 NewChildID = ChildID - Offsets[ChildID];
+							if (0 <= NewChildID && NewChildID < FinalSize)
+							{
+								ChildrenArray[Index].Add(NewChildID);
+							}
+						}
+					}
+				}
+			}
+
+		}
+
+		Super::MergeElements(Group, SortedMergeList, MergeRemapIndex, Params);
+	}
+}
diff --git a/Engine/Source/Runtime/Experimental/Chaos/Public/GeometryCollection/Facades/CollectionTetrahedralBindingsFacade.h b/Engine/Source/Runtime/Experimental/Chaos/Public/GeometryCollection/Facades/CollectionTetrahedralBindingsFacade.h
--- a/Engine/Source/Runtime/Experimental/Chaos/Public/GeometryCollection/Facades/CollectionTetrahedralBindingsFacade.h
+++ b/Engine/Source/Runtime/Experimental/Chaos/Public/GeometryCollection/Facades/CollectionTetrahedralBindingsFacade.h
@@ -76,6 +76,9 @@ namespace GeometryCollection::Facades
 		//! Per vertex amount for deformer masking.
 		static CHAOS_API const FName MaskAttributeName;
 
+		// Dependency
+		static const FName TetrahedralGroupDependency;
+
 		/**
 		* FSelectionFacade Constuctor
 		* @param VertixDependencyGroup : GroupName the index attribute is dependent on. 
diff --git a/Engine/Source/Runtime/Experimental/Chaos/Public/GeometryCollection/ManagedArrayCollection.h b/Engine/Source/Runtime/Experimental/Chaos/Public/GeometryCollection/ManagedArrayCollection.h
--- a/Engine/Source/Runtime/Experimental/Chaos/Public/GeometryCollection/ManagedArrayCollection.h
+++ b/Engine/Source/Runtime/Experimental/Chaos/Public/GeometryCollection/ManagedArrayCollection.h
@@ -168,6 +168,10 @@ public:
 	T* NewCopy() const
 	{
 		T* Collection = new T();
+		for (const FName& GroupName : Collection->GroupNames())
+		{
+			Collection->EmptyGroup(GroupName); //Geometry collection has 1 element in ConvexProperties group when initialized
+		}
 		CopyTo(Collection);
 		return Collection;
 	}
@@ -395,6 +399,12 @@ public:
 	*/
 	CHAOS_API virtual void RemoveElements(const FName & Group, const TArray<int32> & SortedDeletionList, FProcessingParameters Params = FProcessingParameters());
 
+	/**
+	* Merge the element indexed SortedMergeList[Idx] to the element indexed MergeRemapIndex[Idx]
+	* first redirect references of SortedMergeList[Idx] then delete
+	*/
+	CHAOS_API virtual void MergeElements(const FName& Group, const TArray<int32>& SortedMergeList, const TArray<int32>& MergeRemapIndex, FProcessingParameters Params = FProcessingParameters());
+
 	/**
 	* Remove the elements at Position and reindex the dependent arrays 
 	* @param Group - The group that manages the attributes
diff --git a/Engine/Source/Runtime/Experimental/Chaos/Public/GeometryCollection/TransformCollection.h b/Engine/Source/Runtime/Experimental/Chaos/Public/GeometryCollection/TransformCollection.h
--- a/Engine/Source/Runtime/Experimental/Chaos/Public/GeometryCollection/TransformCollection.h
+++ b/Engine/Source/Runtime/Experimental/Chaos/Public/GeometryCollection/TransformCollection.h
@@ -110,6 +110,13 @@ public:
 	*/
 	CHAOS_API virtual void RemoveElements(const FName & Group, const TArray<int32> & SortedDeletionList, FProcessingParameters Params = FProcessingParameters()) override;
 
+	/**
+	* MergeElements
+	*   Merge elements from the transform collection. Updates parent and children hierachy to the remapped index
+	*	Deletes merged elements 
+	*/
+	CHAOS_API virtual void MergeElements(const FName& Group, const TArray<int32>& SortedMergeList, const TArray<int32>& MergeRemapIndex, FProcessingParameters Params = FProcessingParameters()) override;
+
 	/**
 	* reset internal state
 	*/
