Fix potential crashes and ensures when use some clustering and fracture nodes

diff --git a/Engine/Plugins/Experimental/Fracture/Source/FractureEngine/Private/FractureEngineClustering.cpp b/Engine/Plugins/Experimental/Fracture/Source/FractureEngine/Private/FractureEngineClustering.cpp
--- a/Engine/Plugins/Experimental/Fracture/Source/FractureEngine/Private/FractureEngineClustering.cpp
+++ b/Engine/Plugins/Experimental/Fracture/Source/FractureEngine/Private/FractureEngineClustering.cpp
@@ -1082,6 +1082,11 @@ bool FFractureEngineClustering::ClusterSelected(
 	FGeometryCollection& GeometryCollection,
 	TArray<int32>& Selection)
 {
+	if (GeometryCollection.NumElements(FGeometryCollection::TransformGroup) == 0)
+	{
+		return false;
+	}
+
 	GeometryCollection::Facades::FCollectionTransformSelectionFacade SelectionFacade(GeometryCollection);
 	SelectionFacade.RemoveRootNodes(Selection);
 	SelectionFacade.Sanitize(Selection);
@@ -1190,6 +1195,11 @@ bool FFractureEngineClustering::ClusterMagnet(
 {
 	using namespace UE::Private::ClusterMagnet;
 
+	if (GeometryCollection.NumElements(FGeometryCollection::TransformGroup) == 0)
+	{
+		return false;
+	}
+
 	Chaos::Facades::FCollectionHierarchyFacade HierarchyFacade(GeometryCollection);
 	HierarchyFacade.GenerateLevelAttribute();
 	const TManagedArray<int32>& Levels = GeometryCollection.GetAttribute<int32>("Level", FGeometryCollection::TransformGroup);
diff --git a/Engine/Plugins/Experimental/GeometryCollectionPlugin/Source/GeometryCollectionNodes/Private/Dataflow/GeometryCollectionClusteringNodes.cpp b/Engine/Plugins/Experimental/GeometryCollectionPlugin/Source/GeometryCollectionNodes/Private/Dataflow/GeometryCollectionClusteringNodes.cpp
--- a/Engine/Plugins/Experimental/GeometryCollectionPlugin/Source/GeometryCollectionNodes/Private/Dataflow/GeometryCollectionClusteringNodes.cpp
+++ b/Engine/Plugins/Experimental/GeometryCollectionPlugin/Source/GeometryCollectionNodes/Private/Dataflow/GeometryCollectionClusteringNodes.cpp
@@ -113,7 +113,7 @@ void FAutoClusterDataflowNode::Evaluate(UE::Dataflow::FContext& Context, const F
 
 void FClusterFlattenDataflowNode::Evaluate(UE::Dataflow::FContext& Context, const FDataflowOutput* Out) const
 {
-	if (Out->IsA(&Collection) && IsConnected(&Collection))
+	if (Out->IsA(&Collection))
 	{
 		const FManagedArrayCollection& InCollection = GetValue(Context, &Collection);
 		if (InCollection.NumElements(FGeometryCollection::TransformAttribute) > 0)
@@ -151,6 +151,10 @@ void FClusterFlattenDataflowNode::Evaluate(UE::Dataflow::FContext& Context, cons
 				SetValue(Context, (const FManagedArrayCollection&)(*GeomCollection), &Collection);
 			}
 		}
+		else
+		{
+			SetValue(Context, InCollection, &Collection);
+		}
 	}
 }
 
@@ -160,24 +164,31 @@ void FClusterUnclusterDataflowNode::Evaluate(UE::Dataflow::FContext& Context, co
 	if (Out->IsA<FManagedArrayCollection>(&Collection))
 	{
 		const FManagedArrayCollection& InCollection = GetValue<FManagedArrayCollection>(Context, &Collection);
-		const FDataflowTransformSelection& InTransformSelection = GetValue<FDataflowTransformSelection>(Context, &TransformSelection);
-		if (TUniquePtr<FGeometryCollection> GeomCollection = TUniquePtr<FGeometryCollection>(InCollection.NewCopy<FGeometryCollection>()))
+		if (InCollection.NumElements(FGeometryCollection::TransformAttribute) > 0)
 		{
-			Chaos::Facades::FCollectionHierarchyFacade HierarchyFacade(*GeomCollection);
-			HierarchyFacade.GenerateLevelAttribute();
-
-			TArray<int32> Selection = InTransformSelection.AsArray();
-			GeometryCollection::Facades::FCollectionTransformSelectionFacade SelectionFacade(*GeomCollection);
-			SelectionFacade.ConvertSelectionToClusterNodes(Selection, false);
-			SelectionFacade.RemoveRootNodes(Selection);
-			if (!Selection.IsEmpty())
+			const FDataflowTransformSelection& InTransformSelection = GetValue<FDataflowTransformSelection>(Context, &TransformSelection);
+			if (TUniquePtr<FGeometryCollection> GeomCollection = TUniquePtr<FGeometryCollection>(InCollection.NewCopy<FGeometryCollection>()))
 			{
-				FGeometryCollectionClusteringUtility::CollapseHierarchyOneLevel(GeomCollection.Get(), Selection);
-				FGeometryCollectionClusteringUtility::RemoveDanglingClusters(GeomCollection.Get());
-
+				Chaos::Facades::FCollectionHierarchyFacade HierarchyFacade(*GeomCollection);
 				HierarchyFacade.GenerateLevelAttribute();
+
+				TArray<int32> Selection = InTransformSelection.AsArray();
+				GeometryCollection::Facades::FCollectionTransformSelectionFacade SelectionFacade(*GeomCollection);
+				SelectionFacade.ConvertSelectionToClusterNodes(Selection, false);
+				SelectionFacade.RemoveRootNodes(Selection);
+				if (!Selection.IsEmpty())
+				{
+					FGeometryCollectionClusteringUtility::CollapseHierarchyOneLevel(GeomCollection.Get(), Selection);
+					FGeometryCollectionClusteringUtility::RemoveDanglingClusters(GeomCollection.Get());
+
+					HierarchyFacade.GenerateLevelAttribute();
+				}
+				SetValue<const FManagedArrayCollection&>(Context, *GeomCollection, &Collection);
 			}
-			SetValue<const FManagedArrayCollection&>(Context, *GeomCollection, &Collection);
+		}
+		else
+		{
+			SetValue(Context, InCollection, &Collection);
 		}
 	}
 }
diff --git a/Engine/Plugins/Experimental/PlanarCutPlugin/Source/PlanarCut/Private/GeometryMeshConversion.cpp b/Engine/Plugins/Experimental/PlanarCutPlugin/Source/PlanarCut/Private/GeometryMeshConversion.cpp
--- a/Engine/Plugins/Experimental/PlanarCutPlugin/Source/PlanarCut/Private/GeometryMeshConversion.cpp
+++ b/Engine/Plugins/Experimental/PlanarCutPlugin/Source/PlanarCut/Private/GeometryMeshConversion.cpp
@@ -1611,8 +1611,8 @@ FCellMeshes::FCellMeshes(int32 NumUVLayersIn, const FDynamicMesh3& SingleCutter,
 		MeshTransforms::ApplyTransform(CellMeshes[0].AugMesh, FTransformSRT3d(Transform.GetValue()), true);
 	}
 
-	// Mesh should already be augmented
-	if (!ensure(AugmentedDynamicMesh::IsAugmented(CellMeshes[0].AugMesh)))
+	// Augment mesh if needed
+	if (!AugmentedDynamicMesh::IsAugmented(CellMeshes[0].AugMesh))
 	{
 		AugmentedDynamicMesh::Augment(CellMeshes[0].AugMesh, NumUVLayers);
 	}
