Dataflow : Fix assert in StaticMeshToCollection when connecting the material instance pin

diff --git a/Engine/Plugins/Experimental/GeometryCollectionPlugin/Source/GeometryCollectionNodes/Private/Dataflow/GeometryCollectionStaticMeshToCollectionNode.cpp b/Engine/Plugins/Experimental/GeometryCollectionPlugin/Source/GeometryCollectionNodes/Private/Dataflow/GeometryCollectionStaticMeshToCollectionNode.cpp
--- a/Engine/Plugins/Experimental/GeometryCollectionPlugin/Source/GeometryCollectionNodes/Private/Dataflow/GeometryCollectionStaticMeshToCollectionNode.cpp
+++ b/Engine/Plugins/Experimental/GeometryCollectionPlugin/Source/GeometryCollectionNodes/Private/Dataflow/GeometryCollectionStaticMeshToCollectionNode.cpp
@@ -26,7 +26,7 @@ FStaticMeshToCollectionDataflowNode::FStaticMeshToCollectionDataflowNode(const U
 
 void FStaticMeshToCollectionDataflowNode::Evaluate(UE::Dataflow::FContext& Context, const FDataflowOutput* Out) const
 {
-	ensure(Out->IsA(&Collection) || Out->IsA(&Materials) || Out->IsA(&InstancedMeshes));
+	ensure(Out->IsA(&Collection) || Out->IsA(&Materials) || Out->IsA(&MaterialInstances) || Out->IsA(&InstancedMeshes));
 
 	FManagedArrayCollection OutCollection;
 	TArray<TObjectPtr<UMaterialInterface>> OutMaterialInstances;
