Dataflow : selection nodes : add a generic operation node that works on all selection types
- also deprecated more sepcific operation node as they are completely replaced by this new node

diff --git a/Engine/Plugins/Experimental/GeometryCollectionPlugin/Source/GeometryCollectionNodes/Private/Dataflow/GeometryCollectionSelectionNodes.cpp b/Engine/Plugins/Experimental/GeometryCollectionPlugin/Source/GeometryCollectionNodes/Private/Dataflow/GeometryCollectionSelectionNodes.cpp
--- a/Engine/Plugins/Experimental/GeometryCollectionPlugin/Source/GeometryCollectionNodes/Private/Dataflow/GeometryCollectionSelectionNodes.cpp
+++ b/Engine/Plugins/Experimental/GeometryCollectionPlugin/Source/GeometryCollectionNodes/Private/Dataflow/GeometryCollectionSelectionNodes.cpp
@@ -72,6 +72,9 @@ namespace UE::Dataflow
 		DATAFLOW_NODE_REGISTER_CREATION_FACTORY(FCollectionVertexSelectionSetOperationDataflowNode);
 		DATAFLOW_NODE_REGISTER_CREATION_FACTORY(FCollectionSelectionByAttrDataflowNode);
 		DATAFLOW_NODE_REGISTER_CREATION_FACTORY(FGeometrySelectionToVertexSelectionDataflowNode);
+
+		DATAFLOW_NODE_REGISTER_CREATION_FACTORY(FCollectionSelectionSetOperationDataflowNode);
+		
 		// GeometryCollection|Selection
 		DATAFLOW_NODE_REGISTER_CREATION_FACTORY_NODE_COLORS_BY_CATEGORY("GeometryCollection|Selection", FLinearColor(1.f, 1.f, 0.05f), CDefaultNodeBodyTintColor);
 	}
@@ -139,7 +142,6 @@ void FCollectionTransformSelectionSetOperationDataflowNode::Evaluate(UE::Dataflo
 	}
 }
 
-
 namespace {
 	struct BoneInfo {
 		int32 BoneIndex;
@@ -1385,4 +1387,75 @@ void FGeometrySelectionToVertexSelectionDataflowNode::Evaluate(UE::Dataflow::FCo
 		}
 		SetValue(Context, MoveTemp(InVertexSelection), &VertexSelection);
 	}
-}
\ No newline at end of file
+}
+
+//////////////////////////////////////////////////////////////////////////////////////////////////////////////////////
+
+FCollectionSelectionSetOperationDataflowNode::FCollectionSelectionSetOperationDataflowNode(const UE::Dataflow::FNodeParameters& InParam, FGuid InGuid)
+	: FDataflowNode(InParam, InGuid)
+{
+	RegisterInputConnection(&SelectionA);
+	RegisterInputConnection(&SelectionB);
+	RegisterOutputConnection(&Selection, &SelectionA);
+}
+
+void FCollectionSelectionSetOperationDataflowNode::Evaluate(UE::Dataflow::FContext& Context, const FDataflowOutput* Out) const
+{
+	if (Out->IsA(&Selection))
+	{
+		const FDataflowSelection& InSelectionA = GetValue(Context, &SelectionA);
+		const FDataflowSelection& InSelectionB = GetValue(Context, &SelectionB);
+		ensure(FindInput(&SelectionA)->GetType() == FindInput(&SelectionB)->GetType());
+		ensure(FindOutput(&Selection)->GetType() == FindInput(&SelectionA)->GetType());
+
+		FDataflowSelection OutSelection;
+
+		if (InSelectionA.Num() == InSelectionB.Num())
+		{
+			switch (Operation)
+			{
+			case ESetOperationEnum::Dataflow_SetOperation_AND:
+				InSelectionA.AND(InSelectionB, OutSelection);
+				break;
+			case ESetOperationEnum::Dataflow_SetOperation_OR:
+				InSelectionA.OR(InSelectionB, OutSelection);
+				break;
+			case ESetOperationEnum::Dataflow_SetOperation_XOR:
+				InSelectionA.XOR(InSelectionB, OutSelection);
+				break;
+			case ESetOperationEnum::Dataflow_SetOperation_Subtract:
+				InSelectionA.Subtract(InSelectionB, OutSelection);
+				break;
+			}
+		}
+		else
+		{
+			UE_LOG(LogTemp, Error, TEXT("Dataflow: CollectionSelectionSetOperationDataflowNode : Input Selections have different number of elements."));
+		}
+
+		SetValue(Context, MoveTemp(OutSelection), &Selection);
+	}
+}
+
+bool FCollectionSelectionSetOperationDataflowNode::OnInputTypeChanged(const FDataflowInput* Input)
+{
+	// using single bitwise | operator to avoid skipping calls to SetInputConcreteType because of the shortcircuiting of ||
+	// need to disable the warning for static analysis ( V792 warning )
+	return bool(
+		SetInputConcreteType(&SelectionA, Input->GetType())
+		| SetInputConcreteType(&SelectionB, Input->GetType()) //-V792
+		| SetOutputConcreteType(&Selection, Input->GetType()) //-V792
+		);
+}
+
+bool FCollectionSelectionSetOperationDataflowNode::OnOutputTypeChanged(const FDataflowOutput* Input)
+{
+	// using single bitwise | operator to avoid skipping calls to SetInputConcreteType because of the shortcircuiting of ||
+	// need to disable the warning for static analysis ( V792 warning )
+	return bool(
+		SetInputConcreteType(&SelectionA, Input->GetType())
+		| SetInputConcreteType(&SelectionB, Input->GetType()) //-V792
+		);
+}
+
+//////////////////////////////////////////////////////////////////////////////////////////////////////////////////////
\ No newline at end of file
diff --git a/Engine/Plugins/Experimental/GeometryCollectionPlugin/Source/GeometryCollectionNodes/Public/Dataflow/GeometryCollectionSelectionNodes.h b/Engine/Plugins/Experimental/GeometryCollectionPlugin/Source/GeometryCollectionNodes/Public/Dataflow/GeometryCollectionSelectionNodes.h
--- a/Engine/Plugins/Experimental/GeometryCollectionPlugin/Source/GeometryCollectionNodes/Public/Dataflow/GeometryCollectionSelectionNodes.h
+++ b/Engine/Plugins/Experimental/GeometryCollectionPlugin/Source/GeometryCollectionNodes/Public/Dataflow/GeometryCollectionSelectionNodes.h
@@ -6,6 +6,7 @@
 #include "Dataflow/DataflowEngine.h"
 #include "GeometryCollection/ManagedArrayCollection.h"
 #include "Dataflow/DataflowSelection.h"
+#include "Dataflow/DataflowAnyType.h"
 #include "Math/MathFwd.h"
 #include "Math/Sphere.h"
 #include "Dataflow/DataflowConnectionTypes.h"
@@ -73,9 +74,10 @@ enum class ESetOperationEnum : uint8
 /**
  *
  * Runs boolean operation on TransformSelections
- *
+ * Deprecated (5.6) : use the generic CollectionSelectionSetOperation node instead
  */
-USTRUCT(meta = (DataflowGeometryCollection))
+
+USTRUCT(meta = (Deprecated = "5.6"))
 struct FCollectionTransformSelectionSetOperationDataflowNode : public FDataflowNode
 {
 	GENERATED_USTRUCT_BODY()
@@ -110,7 +112,6 @@ public:
 
 };
 
-
 /**
  *
  * Generates a formatted string of the bones and the selection
@@ -1329,13 +1330,13 @@ public:
 
 };
 
-
 /**
  *
  * Runs boolean operation on VertexSelections
+ * Deprecated (5.6) : use the generic CollectionSelectionSetOperation node instead
  *
  */
-USTRUCT(meta = (DataflowGeometryCollection))
+USTRUCT(meta = (Deprecated = "5.6"))
 struct FCollectionVertexSelectionSetOperationDataflowNode : public FDataflowNode
 {
 	GENERATED_USTRUCT_BODY()
@@ -1517,6 +1518,43 @@ private:
 	virtual void Evaluate(UE::Dataflow::FContext& Context, const FDataflowOutput* Out) const override;
 };
 
+/**
+ *
+ * Runs boolean operation on selection ( support all selection types )
+ *
+ */
+USTRUCT()
+struct FCollectionSelectionSetOperationDataflowNode : public FDataflowNode
+{
+	GENERATED_USTRUCT_BODY()
+	DATAFLOW_NODE_DEFINE_INTERNAL(FCollectionSelectionSetOperationDataflowNode, "CollectionSelectionSetOperation", "GeometryCollection|Selection", "")
+
+private:
+	/** Boolean operation */
+	UPROPERTY(EditAnywhere, Category = "Operation");
+	ESetOperationEnum Operation = ESetOperationEnum::Dataflow_SetOperation_AND;
+
+	/** First Selection object */
+	UPROPERTY(meta = (DataflowInput, DataflowIntrinsic))
+	FDataflowSelectionTypes SelectionA;
+
+	/** Second Selection object */
+	UPROPERTY(meta = (DataflowInput, DataflowIntrinsic))
+	FDataflowSelectionTypes SelectionB;
+
+	/** Array of the selected bone indices after operation*/
+	UPROPERTY(meta = (DataflowOutput, DataflowPassthrough = "SelectionA"))
+	FDataflowSelectionTypes Selection;
+
+	virtual bool OnInputTypeChanged(const FDataflowInput* Input) override;
+	virtual bool OnOutputTypeChanged(const FDataflowOutput* Input) override;
+
+	virtual void Evaluate(UE::Dataflow::FContext& Context, const FDataflowOutput* Out) const override;
+
+public:
+	FCollectionSelectionSetOperationDataflowNode(const UE::Dataflow::FNodeParameters& InParam, FGuid InGuid = FGuid::NewGuid());
+};
+
 namespace UE::Dataflow
 {
 	void GeometryCollectionSelectionNodes();
