Dataflow : Add generic invert selection node
- deprecate transform and face selection invert nodes as they are replaced by the generic one

diff --git a/Engine/Plugins/Experimental/GeometryCollectionPlugin/Source/GeometryCollectionNodes/Private/Dataflow/GeometryCollectionSelectionNodes.cpp b/Engine/Plugins/Experimental/GeometryCollectionPlugin/Source/GeometryCollectionNodes/Private/Dataflow/GeometryCollectionSelectionNodes.cpp
--- a/Engine/Plugins/Experimental/GeometryCollectionPlugin/Source/GeometryCollectionNodes/Private/Dataflow/GeometryCollectionSelectionNodes.cpp
+++ b/Engine/Plugins/Experimental/GeometryCollectionPlugin/Source/GeometryCollectionNodes/Private/Dataflow/GeometryCollectionSelectionNodes.cpp
@@ -39,10 +39,8 @@ namespace UE::Dataflow
 		static const FLinearColor CDefaultNodeBodyTintColor = FLinearColor(0.f, 0.f, 0.f, 0.5f);
 
 		DATAFLOW_NODE_REGISTER_CREATION_FACTORY(FCollectionTransformSelectionAllDataflowNode);
-		DATAFLOW_NODE_REGISTER_CREATION_FACTORY(FCollectionTransformSelectionSetOperationDataflowNode);
 		DATAFLOW_NODE_REGISTER_CREATION_FACTORY(FCollectionTransformSelectionInfoDataflowNode);
 		DATAFLOW_NODE_REGISTER_CREATION_FACTORY(FCollectionTransformSelectionNoneDataflowNode);
-		DATAFLOW_NODE_REGISTER_CREATION_FACTORY(FCollectionTransformSelectionInvertDataflowNode);
 		DATAFLOW_NODE_REGISTER_CREATION_FACTORY(FCollectionTransformSelectionRandomDataflowNode);
 		DATAFLOW_NODE_REGISTER_CREATION_FACTORY(FCollectionTransformSelectionRootDataflowNode);
 		DATAFLOW_NODE_REGISTER_CREATION_FACTORY(FCollectionTransformSelectionCustomDataflowNode);
@@ -67,14 +65,20 @@ namespace UE::Dataflow
 		DATAFLOW_NODE_REGISTER_CREATION_FACTORY(FCollectionVertexSelectionCustomDataflowNode);
 		DATAFLOW_NODE_REGISTER_CREATION_FACTORY(FCollectionFaceSelectionCustomDataflowNode);
 		DATAFLOW_NODE_REGISTER_CREATION_FACTORY(FCollectionSelectionConvertDataflowNode);
-		DATAFLOW_NODE_REGISTER_CREATION_FACTORY(FCollectionFaceSelectionInvertDataflowNode);
 		DATAFLOW_NODE_REGISTER_CREATION_FACTORY(FCollectionVertexSelectionByPercentageDataflowNode);
 		DATAFLOW_NODE_REGISTER_CREATION_FACTORY(FCollectionVertexSelectionSetOperationDataflowNode);
 		DATAFLOW_NODE_REGISTER_CREATION_FACTORY(FCollectionSelectionByAttrDataflowNode);
 		DATAFLOW_NODE_REGISTER_CREATION_FACTORY(FGeometrySelectionToVertexSelectionDataflowNode);
 
+		// generic input nodes
 		DATAFLOW_NODE_REGISTER_CREATION_FACTORY(FCollectionSelectionSetOperationDataflowNode);
-		
+		DATAFLOW_NODE_REGISTER_CREATION_FACTORY(FCollectionSelectionInvertDataflowNode);
+
+		// deprecated
+		DATAFLOW_NODE_REGISTER_CREATION_FACTORY(FCollectionTransformSelectionInvertDataflowNode);
+		DATAFLOW_NODE_REGISTER_CREATION_FACTORY(FCollectionFaceSelectionInvertDataflowNode);
+		DATAFLOW_NODE_REGISTER_CREATION_FACTORY(FCollectionTransformSelectionSetOperationDataflowNode);
+
 		// GeometryCollection|Selection
 		DATAFLOW_NODE_REGISTER_CREATION_FACTORY_NODE_COLORS_BY_CATEGORY("GeometryCollection|Selection", FLinearColor(1.f, 1.f, 0.05f), CDefaultNodeBodyTintColor);
 	}
@@ -1458,4 +1462,35 @@ bool FCollectionSelectionSetOperationDataflowNode::OnOutputTypeChanged(const FDa
 		);
 }
 
+//////////////////////////////////////////////////////////////////////////////////////////////////////////////////////
+
+FCollectionSelectionInvertDataflowNode::FCollectionSelectionInvertDataflowNode(const UE::Dataflow::FNodeParameters& InParam, FGuid InGuid)
+	: FDataflowNode(InParam, InGuid)
+{
+	RegisterInputConnection(&Selection);
+	RegisterOutputConnection(&Selection, &Selection);
+}
+
+void FCollectionSelectionInvertDataflowNode::Evaluate(UE::Dataflow::FContext& Context, const FDataflowOutput* Out) const
+{
+	if (Out->IsA(&Selection))
+	{
+		ensure(FindOutput(&Selection)->GetType() == FindInput(&Selection)->GetType());
+
+		FDataflowSelection InSelection = GetValue(Context, &Selection);
+		InSelection.Invert();
+		SetValue(Context, MoveTemp(InSelection), &Selection);
+	}
+}
+
+bool FCollectionSelectionInvertDataflowNode::OnInputTypeChanged(const FDataflowInput* Input)
+{
+	return SetOutputConcreteType(&Selection, Input->GetType());
+}
+
+bool FCollectionSelectionInvertDataflowNode::OnOutputTypeChanged(const FDataflowOutput* Input)
+{
+	return SetInputConcreteType(&Selection, Input->GetType());
+}
+
 //////////////////////////////////////////////////////////////////////////////////////////////////////////////////////
\ No newline at end of file
diff --git a/Engine/Plugins/Experimental/GeometryCollectionPlugin/Source/GeometryCollectionNodes/Public/Dataflow/GeometryCollectionSelectionNodes.h b/Engine/Plugins/Experimental/GeometryCollectionPlugin/Source/GeometryCollectionNodes/Public/Dataflow/GeometryCollectionSelectionNodes.h
--- a/Engine/Plugins/Experimental/GeometryCollectionPlugin/Source/GeometryCollectionNodes/Public/Dataflow/GeometryCollectionSelectionNodes.h
+++ b/Engine/Plugins/Experimental/GeometryCollectionPlugin/Source/GeometryCollectionNodes/Public/Dataflow/GeometryCollectionSelectionNodes.h
@@ -187,7 +187,7 @@ public:
  * Inverts selection of bones
  *
  */
-USTRUCT(meta = (DataflowGeometryCollection))
+USTRUCT(meta = (Deprecated = "5.6"))
 struct FCollectionTransformSelectionInvertDataflowNode : public FDataflowNode
 {
 	GENERATED_USTRUCT_BODY()
@@ -1265,7 +1265,7 @@ public:
  * Inverts selection of faces
  *
  */
-USTRUCT(meta = (DataflowGeometryCollection))
+USTRUCT(meta = (Deprecated = "5.6"))
 struct FCollectionFaceSelectionInvertDataflowNode : public FDataflowNode
 {
 	GENERATED_USTRUCT_BODY()
@@ -1555,6 +1555,32 @@ public:
 	FCollectionSelectionSetOperationDataflowNode(const UE::Dataflow::FNodeParameters& InParam, FGuid InGuid = FGuid::NewGuid());
 };
 
+
+/**
+ *
+ * Inverts selection ( support all selection types )
+ *
+ */
+USTRUCT()
+struct FCollectionSelectionInvertDataflowNode : public FDataflowNode
+{
+	GENERATED_USTRUCT_BODY()
+	DATAFLOW_NODE_DEFINE_INTERNAL(FCollectionSelectionInvertDataflowNode, "CollectionSelectionInvert", "GeometryCollection|Selection", "")
+
+private:
+	/** selection to invert */
+	UPROPERTY(meta = (DataflowInput, DataflowOutput, DataflowPassthrough = "Selection", DataflowIntrinsic))
+	FDataflowSelectionTypes Selection;
+
+	virtual bool OnInputTypeChanged(const FDataflowInput* Input) override;
+	virtual bool OnOutputTypeChanged(const FDataflowOutput* Input) override;
+
+	virtual void Evaluate(UE::Dataflow::FContext& Context, const FDataflowOutput* Out) const override;
+
+public:
+	FCollectionSelectionInvertDataflowNode(const UE::Dataflow::FNodeParameters& InParam, FGuid InGuid = FGuid::NewGuid());
+};
+
 namespace UE::Dataflow
 {
 	void GeometryCollectionSelectionNodes();
