Dataflow : Add Auto UV unwrap + Add UV channel nodes

diff --git a/Engine/Plugins/Experimental/GeometryCollectionPlugin/Source/GeometryCollectionNodes/Private/Dataflow/GeometryCollectionNodesPlugin.cpp b/Engine/Plugins/Experimental/GeometryCollectionPlugin/Source/GeometryCollectionNodes/Private/Dataflow/GeometryCollectionNodesPlugin.cpp
--- a/Engine/Plugins/Experimental/GeometryCollectionPlugin/Source/GeometryCollectionNodes/Private/Dataflow/GeometryCollectionNodesPlugin.cpp
+++ b/Engine/Plugins/Experimental/GeometryCollectionPlugin/Source/GeometryCollectionNodes/Private/Dataflow/GeometryCollectionNodesPlugin.cpp
@@ -29,6 +29,7 @@
 #include "Dataflow/GeometryCollectionTransferVertexAttributeNode.h"
 #include "Dataflow/GeometryCollectionTriangleBoundaryIndicesNode.h"
 #include "Dataflow/GeometryCollectionUtilityNodes.h"
+#include "Dataflow/GeometryCollectionUVNodes.h"
 #include "Dataflow/GeometryCollectionVerticesNodes.h"
 #include "Dataflow/GeometryCollectionVertexScalarToVertexIndicesNode.h"
 #include "Dataflow/SetVertexColorFromFloatArrayNode.h"
@@ -59,6 +60,7 @@ void IGeometryCollectionNodesPlugin::StartupModule()
 	UE::Dataflow::GeometryCollectionArrayNodes();
 	UE::Dataflow::GeometryCollectionDebugNodes();
 	UE::Dataflow::GeometryCollectionSamplingNodes();
+	UE::Dataflow::RegisterGeometryCollectionUVNodes();
 	DATAFLOW_NODE_REGISTER_CREATION_FACTORY(FCreateColorArrayFromFloatArrayDataflowNode);
 	DATAFLOW_NODE_REGISTER_CREATION_FACTORY(FGeometryCollectionVertexScalarToVertexIndicesNode);
 	DATAFLOW_NODE_REGISTER_CREATION_FACTORY(FSetVertexColorFromFloatArrayDataflowNode);
diff --git a/Engine/Plugins/Experimental/GeometryCollectionPlugin/Source/GeometryCollectionNodes/Private/Dataflow/GeometryCollectionUVNodes.cpp b/Engine/Plugins/Experimental/GeometryCollectionPlugin/Source/GeometryCollectionNodes/Private/Dataflow/GeometryCollectionUVNodes.cpp
new file mode 100644
--- /dev/null
+++ b/Engine/Plugins/Experimental/GeometryCollectionPlugin/Source/GeometryCollectionNodes/Private/Dataflow/GeometryCollectionUVNodes.cpp
@@ -0,0 +1,113 @@
+// Copyright Epic Games, Inc. All Rights Reserved.
+
+#include "Dataflow/GeometryCollectionUVNodes.h"
+#include "GeometryCollection/Facades/CollectionUVFacade.h"
+
+#include "FractureAutoUV.h"
+
+namespace UE::Dataflow
+{
+	void RegisterGeometryCollectionUVNodes()
+	{
+		DATAFLOW_NODE_REGISTER_CREATION_FACTORY(FAddUVChannelDataflowNode);
+		DATAFLOW_NODE_REGISTER_CREATION_FACTORY(FAutoUnwrapUVDataflowNode);
+
+		static const FLinearColor CDefaultNodeBodyTintColor = FLinearColor(0.f, 0.f, 0.f, 0.5f);
+		DATAFLOW_NODE_REGISTER_CREATION_FACTORY_NODE_COLORS_BY_CATEGORY("GeometryCollection|UV", FLinearColor(.1f, 1.f, 0.6f), CDefaultNodeBodyTintColor);
+	}
+}
+
+////////////////////////////////////////////////////////////////////////////////////////////
+
+FAddUVChannelDataflowNode::FAddUVChannelDataflowNode(const UE::Dataflow::FNodeParameters& InParam, FGuid InGuid)
+	: FDataflowNode(InParam, InGuid)
+{
+	RegisterInputConnection(&Collection);
+	RegisterOutputConnection(&Collection, &Collection);
+	RegisterOutputConnection(&UVChannel);
+}
+
+void FAddUVChannelDataflowNode::Evaluate(UE::Dataflow::FContext& Context, const FDataflowOutput* Out) const
+{
+	if (Out->IsA(&Collection) || Out->IsA(&UVChannel))
+	{
+		FManagedArrayCollection InCollection = GetValue(Context, &Collection);
+
+		GeometryCollection::Facades::FCollectionUVFacade UVFacade(InCollection);
+
+		int32 NumUVChannel = 0;
+		if (UVFacade.IsValid())
+		{
+			NumUVChannel = UVFacade.GetNumUVLayers();
+		}
+
+		int32 NewUVChannel = NumUVChannel;
+		if (UVFacade.SetNumUVLayers(NumUVChannel + 1))
+		{
+			// init the UV
+			if (TManagedArray<FVector2f>* UVAttribute = UVFacade.FindUVLayer(NewUVChannel))
+			{
+				UVAttribute->Fill(DefaultValue);
+			}
+		}
+		else
+		{
+			// for now if we fail ( we maxed out thenumber of channel for example ), let's use the channel 0 so that it is obvious
+			// todo(dataflow) in the future we should output an error 
+			NewUVChannel = 0;
+		}
+
+		SetValue(Context, MoveTemp(InCollection), &Collection);
+		SetValue(Context, NewUVChannel, &UVChannel);
+	}
+}
+
+////////////////////////////////////////////////////////////////////////////////////////////
+
+FAutoUnwrapUVDataflowNode::FAutoUnwrapUVDataflowNode(const UE::Dataflow::FNodeParameters& InParam, FGuid InGuid)
+	: FDataflowNode(InParam, InGuid)
+{
+	RegisterInputConnection(&Collection);
+	RegisterInputConnection(&FaceSelection);
+	RegisterInputConnection(&UVChannel);
+	RegisterOutputConnection(&Collection, &Collection);
+}
+
+void FAutoUnwrapUVDataflowNode::Evaluate(UE::Dataflow::FContext& Context, const FDataflowOutput* Out) const
+{
+	if (Out->IsA(&Collection))
+	{
+		const FManagedArrayCollection& InCollection = GetValue(Context, &Collection);
+		const int32 FaceCount = InCollection.NumElements(FGeometryCollection::FacesGroup);
+		if (FaceCount == 0)
+		{
+			// nothing to do : forward and exit 
+			SafeForwardInput(Context, &Collection, &Collection);
+			return;
+		}
+
+		const int32 InUVChannel = GetValue(Context, &UVChannel);
+
+		TUniquePtr<FGeometryCollection> GeomCollection = TUniquePtr<FGeometryCollection>(InCollection.NewCopy<FGeometryCollection>());
+		check(GeomCollection.IsValid());
+
+		// this is used as reference for the gutter size that is expressed in pixels for a specific resolution 
+		constexpr int32 ReferenceResolution = 4096;
+		constexpr bool bRecreateUVsForDegenerateIslands = true;
+
+		if (IsConnected(&FaceSelection))
+		{
+			UE::PlanarCut::UVLayout(InUVChannel, *GeomCollection, ReferenceResolution, GutterSize,
+				UE::PlanarCut::ETargetFaces::AllFaces, {}, bRecreateUVsForDegenerateIslands, nullptr /* progress */);
+		}
+		else
+		{
+			// All faces
+			UE::PlanarCut::UVLayout(InUVChannel, *GeomCollection, ReferenceResolution, GutterSize,
+				UE::PlanarCut::ETargetFaces::AllFaces, {}, bRecreateUVsForDegenerateIslands, nullptr /* progress */);
+		}
+		
+		SetValue(Context, (const FManagedArrayCollection&)(*GeomCollection), &Collection);
+	}
+}
+
diff --git a/Engine/Plugins/Experimental/GeometryCollectionPlugin/Source/GeometryCollectionNodes/Public/Dataflow/GeometryCollectionUVNodes.h b/Engine/Plugins/Experimental/GeometryCollectionPlugin/Source/GeometryCollectionNodes/Public/Dataflow/GeometryCollectionUVNodes.h
new file mode 100644
--- /dev/null
+++ b/Engine/Plugins/Experimental/GeometryCollectionPlugin/Source/GeometryCollectionNodes/Public/Dataflow/GeometryCollectionUVNodes.h
@@ -0,0 +1,76 @@
+// Copyright Epic Games, Inc. All Rights Reserved.
+
+#pragma once 
+
+#include "Dataflow/DataflowEngine.h"
+#include "Dataflow/DataflowSelection.h"
+#include "GeometryCollection/ManagedArrayCollection.h"
+
+#include "GeometryCollectionUVNodes.generated.h"
+
+/*
+* Add a new UV channel to the collection
+* note that there's a maximum of 8 channels that can be handled by a collection
+*/
+USTRUCT()
+struct FAddUVChannelDataflowNode : public FDataflowNode
+{
+	GENERATED_USTRUCT_BODY()
+	DATAFLOW_NODE_DEFINE_INTERNAL(FAddUVChannelDataflowNode, "AddUVChannel", "GeometryCollection|UV", "")
+
+private:
+	/** Target Collection */
+	UPROPERTY(meta = (DataflowInput, DataflowOutput, DataflowPassthrough = "Collection", DataflowIntrinsic))
+	FManagedArrayCollection Collection;
+
+	/** Index of the added UV channel */
+	UPROPERTY(meta = (DataflowOutput))
+	int32 UVChannel = 0;
+
+	/** Value to initialize the UV with */
+	UPROPERTY(EditAnywhere, Category = Options)
+	FVector2f DefaultValue = { 0, 0 };
+
+	virtual void Evaluate(UE::Dataflow::FContext& Context, const FDataflowOutput* Out) const override;
+
+public:
+	FAddUVChannelDataflowNode(const UE::Dataflow::FNodeParameters& InParam, FGuid InGuid = FGuid::NewGuid());
+};
+
+/**
+ * Auto unwrap UVs for a specific UV channel
+ */
+USTRUCT()
+struct FAutoUnwrapUVDataflowNode : public FDataflowNode
+{
+	GENERATED_USTRUCT_BODY()
+	DATAFLOW_NODE_DEFINE_INTERNAL(FAutoUnwrapUVDataflowNode, "AutoUnwrapUV", "GeometryCollection|UV", "")
+
+private:
+	/** Target Collection */
+	UPROPERTY(meta = (DataflowInput, DataflowOutput, DataflowPassthrough = "Collection", DataflowIntrinsic))
+	FManagedArrayCollection Collection;
+
+	/** Faces to auto unwrap, no selection means all faces */
+	UPROPERTY(meta = (DataflowInput))
+	FDataflowFaceSelection FaceSelection;
+
+	/** UV channel to unwrap into ( 0 by default ) */
+	UPROPERTY(EditAnywhere, Category = Options, meta = (DataflowInput))
+	int32 UVChannel = 0;
+
+	/** Gutter Size in pixels for a 4096x4096 reference resolution */
+	UPROPERTY(EditAnywhere, Category = Options)
+	int32 GutterSize;
+
+	virtual void Evaluate(UE::Dataflow::FContext& Context, const FDataflowOutput* Out) const override;
+
+public:
+	FAutoUnwrapUVDataflowNode(const UE::Dataflow::FNodeParameters& InParam, FGuid InGuid = FGuid::NewGuid());
+};
+
+namespace UE::Dataflow
+{
+	void RegisterGeometryCollectionUVNodes();
+}
+
diff --git a/Engine/Plugins/Experimental/PlanarCutPlugin/Source/PlanarCut/Private/FractureAutoUV.cpp b/Engine/Plugins/Experimental/PlanarCutPlugin/Source/PlanarCut/Private/FractureAutoUV.cpp
--- a/Engine/Plugins/Experimental/PlanarCutPlugin/Source/PlanarCut/Private/FractureAutoUV.cpp
+++ b/Engine/Plugins/Experimental/PlanarCutPlugin/Source/PlanarCut/Private/FractureAutoUV.cpp
@@ -478,12 +478,31 @@ bool UVLayout(
 	bool bRecreateUVsForDegenerateIslands,
 	FProgressCancel* Progress
 )
+{
+	TArray<bool> FaceSelection;
+	int32 NumActive = SetActiveTriangles(&Collection, FaceSelection, true, TargetFaces, WhichMaterials);
+	return UVLayout(TargetUVLayer, Collection, UVRes, GutterSize, FaceSelection, bRecreateUVsForDegenerateIslands, Progress);
+}
+
+
+bool PLANARCUT_API UVLayout(
+	int32 TargetUVLayer,
+	FGeometryCollection& Collection,
+	int32 UVRes,
+	float GutterSize,
+	TArrayView<bool> FaceSelection,
+	bool bRecreateUVsForDegenerateIslands,
+	FProgressCancel* Progress
+)
 {
 	FProgressCancel::FProgressScope InitScope = FProgressCancel::CreateScopeTo(Progress, .2);
 
-	TArray<bool> ActiveTriangles;
-	int32 NumActive = SetActiveTriangles(&Collection, ActiveTriangles, true, TargetFaces, WhichMaterials);
-	FGeomMesh UVMesh(TargetUVLayer, &Collection, ActiveTriangles, NumActive);
+	int32 NumActive = 0;
+	for (const bool& bSelected : FaceSelection)
+	{
+		NumActive += (bSelected) ? 1 : 0;
+	}
+	FGeomMesh UVMesh(TargetUVLayer, &Collection, FaceSelection, NumActive);
 
 	TArray<TArray<int32>> UVIslands;
 	UE::UVPacking::CreateUVIslandsFromMeshTopology<FGeomMesh>(UVMesh, UVIslands);
diff --git a/Engine/Plugins/Experimental/PlanarCutPlugin/Source/PlanarCut/Public/FractureAutoUV.h b/Engine/Plugins/Experimental/PlanarCutPlugin/Source/PlanarCut/Public/FractureAutoUV.h
--- a/Engine/Plugins/Experimental/PlanarCutPlugin/Source/PlanarCut/Public/FractureAutoUV.h
+++ b/Engine/Plugins/Experimental/PlanarCutPlugin/Source/PlanarCut/Public/FractureAutoUV.h
@@ -103,6 +103,26 @@ bool PLANARCUT_API UVLayout(
 	FProgressCancel* Progress = nullptr
 );
 
+/**
+ * Make a UV atlas of non-overlapping UV charts for a geometry collection ( face selection version )
+ *
+ * @param TargetUVLayer		Which UV layer to update with new UV coordinates
+ * @param Collection		The collection to be atlas'd
+ * @param UVRes				Target resolution for the atlas
+ * @param GutterSize		Space to leave between UV islands, in pixels at the target resolution
+ * @param FacesSelection	Which faces to automatically consider for UV island layout
+  * @param bRecreateUVsForDegenerateIslands If true, detect and fix islands that don't have proper UVs (i.e. UVs all zero or otherwise collapsed to a point)
+ */
+bool PLANARCUT_API UVLayout(
+	int32 TargetUVLayer,
+	FGeometryCollection& Collection,
+	int32 UVRes,
+	float GutterSize,
+	TArrayView<bool> FaceSelection,
+	bool bRecreateUVsForDegenerateIslands = true,
+	FProgressCancel* Progress = nullptr
+);
+
 /**
  * Make a UV atlas of non-overlapping UV charts for a geometry collection
  *
