Dataflow fix face selection code path in uwrap auto UV node

diff --git a/Engine/Plugins/Experimental/GeometryCollectionPlugin/Source/GeometryCollectionNodes/Private/Dataflow/GeometryCollectionUVNodes.cpp b/Engine/Plugins/Experimental/GeometryCollectionPlugin/Source/GeometryCollectionNodes/Private/Dataflow/GeometryCollectionUVNodes.cpp
--- a/Engine/Plugins/Experimental/GeometryCollectionPlugin/Source/GeometryCollectionNodes/Private/Dataflow/GeometryCollectionUVNodes.cpp
+++ b/Engine/Plugins/Experimental/GeometryCollectionPlugin/Source/GeometryCollectionNodes/Private/Dataflow/GeometryCollectionUVNodes.cpp
@@ -97,8 +97,19 @@ void FAutoUnwrapUVDataflowNode::Evaluate(UE::Dataflow::FContext& Context, const
 
 		if (IsConnected(&FaceSelection))
 		{
-			UE::PlanarCut::UVLayout(InUVChannel, *GeomCollection, ReferenceResolution, GutterSize,
-				UE::PlanarCut::ETargetFaces::AllFaces, {}, bRecreateUVsForDegenerateIslands, nullptr /* progress */);
+			const FDataflowFaceSelection& InFaceSelection = GetValue(Context, &FaceSelection);
+			if (InFaceSelection.Num() > 0)
+			{
+				TArray<bool> FacesToUnwrap;
+				FacesToUnwrap.SetNumUninitialized(InFaceSelection.Num());
+				for (int32 Index = 0; Index < InFaceSelection.Num(); ++Index)
+				{
+					FacesToUnwrap[Index] = InFaceSelection.IsSelected(Index);
+				}
+
+				UE::PlanarCut::UVLayout(InUVChannel, *GeomCollection, ReferenceResolution, GutterSize,
+					FacesToUnwrap, bRecreateUVsForDegenerateIslands, nullptr /* progress */);
+			}
 		}
 		else
 		{
