Dataflow:
- Updated FCollectionTransformSelectionCustomDataflowNode to use regex to specify the bone indices
- Bug fixes

diff --git a/Engine/Plugins/Experimental/GeometryCollectionPlugin/Source/GeometryCollectionNodes/Private/Dataflow/GeometryCollectionSelectionNodes.cpp b/Engine/Plugins/Experimental/GeometryCollectionPlugin/Source/GeometryCollectionNodes/Private/Dataflow/GeometryCollectionSelectionNodes.cpp
--- a/Engine/Plugins/Experimental/GeometryCollectionPlugin/Source/GeometryCollectionNodes/Private/Dataflow/GeometryCollectionSelectionNodes.cpp
+++ b/Engine/Plugins/Experimental/GeometryCollectionPlugin/Source/GeometryCollectionNodes/Private/Dataflow/GeometryCollectionSelectionNodes.cpp
@@ -28,6 +28,7 @@
 #include "FractureEngineClustering.h"
 #include "FractureEngineSelection.h"
 #include "GeometryCollection/Facades/CollectionTransformSelectionFacade.h"
+#include "Internationalization/Regex.h"
 
 #include UE_INLINE_GENERATED_CPP_BY_NAME(GeometryCollectionSelectionNodes)
 
@@ -44,6 +45,7 @@ namespace UE::Dataflow
 		DATAFLOW_NODE_REGISTER_CREATION_FACTORY(FCollectionTransformSelectionRandomDataflowNode);
 		DATAFLOW_NODE_REGISTER_CREATION_FACTORY(FCollectionTransformSelectionRootDataflowNode);
 		DATAFLOW_NODE_REGISTER_CREATION_FACTORY(FCollectionTransformSelectionCustomDataflowNode);
+		DATAFLOW_NODE_REGISTER_CREATION_FACTORY(FCollectionTransformSelectionCustomDataflowNode_v2);
 		DATAFLOW_NODE_REGISTER_CREATION_FACTORY(FCollectionTransformSelectionFromIndexArrayDataflowNode);
 		DATAFLOW_NODE_REGISTER_CREATION_FACTORY(FCollectionTransformSelectionParentDataflowNode);
 		DATAFLOW_NODE_REGISTER_CREATION_FACTORY(FCollectionTransformSelectionByPercentageDataflowNode);
@@ -415,6 +417,122 @@ void FCollectionTransformSelectionCustomDataflowNode::Evaluate(UE::Dataflow::FCo
 	}
 }
 
+namespace UE::Dataflow::Private
+{
+	const static uint32 Error_InvalidChars = 1;
+	const static uint32 Error_InvalidFormatInSegment = 2;
+
+	/* e.g. "0, 2, 5-10, 12-15". If left empty, all will be used */
+	static bool ParseIndicesStr(const FString& InFramesString, TArray<int32>& OutIndices, uint32& OutErrorCode)
+	{
+		static const FRegexPattern AllowedCharsPattern(TEXT("^[-,0-9\\s]+$"));
+
+		if (!FRegexMatcher(AllowedCharsPattern, InFramesString).FindNext())
+		{
+			// Input contains invalid characters
+			OutErrorCode = Error_InvalidChars;
+			return false;
+		}
+
+		static const FRegexPattern SingleNumberPattern(TEXT("^\\s*(\\d+)\\s*$"));
+		static const FRegexPattern RangePattern(TEXT("^\\s*(\\d+)\\s*-\\s*(\\d+)\\s*$"));
+
+		TArray<FString> Segments;
+		InFramesString.ParseIntoArray(Segments, TEXT(","), true);
+		for (const FString& Segment : Segments)
+		{
+			bool bSegmentValid = false;
+
+			FRegexMatcher SingleNumberMatcher(SingleNumberPattern, Segment);
+			if (SingleNumberMatcher.FindNext())
+			{
+				const int32 SingleNumber = FCString::Atoi(*SingleNumberMatcher.GetCaptureGroup(1));
+				OutIndices.Add(SingleNumber);
+				bSegmentValid = true;
+			}
+			else
+			{
+				FRegexMatcher RangeMatcher(RangePattern, Segment);
+				if (RangeMatcher.FindNext())
+				{
+					const int32 RangeStart = FCString::Atoi(*RangeMatcher.GetCaptureGroup(1));
+					const int32 RangeEnd = FCString::Atoi(*RangeMatcher.GetCaptureGroup(2));
+
+					for (int32 i = RangeStart; i <= RangeEnd; ++i)
+					{
+						OutIndices.Add(i);
+					}
+					bSegmentValid = true;
+				}
+			}
+
+			if (!bSegmentValid)
+			{
+				// Invalid format in segment
+				OutErrorCode = Error_InvalidFormatInSegment;
+				return false;
+			}
+		}
+
+		return true;
+	}
+}
+
+FCollectionTransformSelectionCustomDataflowNode_v2::FCollectionTransformSelectionCustomDataflowNode_v2(const UE::Dataflow::FNodeParameters& InParam, FGuid InGuid)
+	: FDataflowNode(InParam, InGuid)
+{
+	RegisterInputConnection(&Collection);
+	RegisterInputConnection(&BoneIndices);
+
+	RegisterOutputConnection(&Collection, &Collection);
+	RegisterOutputConnection(&TransformSelection);
+}
+
+void FCollectionTransformSelectionCustomDataflowNode_v2::Evaluate(UE::Dataflow::FContext& Context, const FDataflowOutput* Out) const
+{
+	if (Out->IsA(&TransformSelection))
+	{
+		const FManagedArrayCollection& InCollection = GetValue(Context, &Collection);
+
+		if (InCollection.HasGroup(FGeometryCollection::TransformGroup))
+		{
+			const int32 NumTransforms = InCollection.NumElements(FGeometryCollection::TransformGroup);
+
+			FDataflowTransformSelection NewTransformSelection;
+			NewTransformSelection.Initialize(NumTransforms, false);
+
+			const FString InBoneIndices = GetValue(Context, &BoneIndices);
+
+			TArray<int32> Indices;
+			uint32 ErrorCode = 0;
+			if (UE::Dataflow::Private::ParseIndicesStr(InBoneIndices, Indices, ErrorCode))
+			{
+				NewTransformSelection.SetSelected(Indices);
+			}
+			else
+			{
+				if (ErrorCode == UE::Dataflow::Private::Error_InvalidChars)
+				{
+					// Handle Error				
+				}
+				else if (ErrorCode == UE::Dataflow::Private::Error_InvalidFormatInSegment)
+				{
+					// Handle Error				
+				}
+			}
+
+			SetValue(Context, MoveTemp(NewTransformSelection), &TransformSelection);
+		}
+		else
+		{
+			SetValue(Context, FDataflowTransformSelection(), &TransformSelection);
+		}
+	}
+	else if (Out->IsA(&Collection))
+	{
+		SafeForwardInput(Context, &Collection, &Collection);
+	}
+}
 
 void FCollectionTransformSelectionFromIndexArrayDataflowNode::Evaluate(UE::Dataflow::FContext& Context, const FDataflowOutput* Out) const
 {
diff --git a/Engine/Plugins/Experimental/GeometryCollectionPlugin/Source/GeometryCollectionNodes/Public/Dataflow/GeometryCollectionSelectionNodes.h b/Engine/Plugins/Experimental/GeometryCollectionPlugin/Source/GeometryCollectionNodes/Public/Dataflow/GeometryCollectionSelectionNodes.h
--- a/Engine/Plugins/Experimental/GeometryCollectionPlugin/Source/GeometryCollectionNodes/Public/Dataflow/GeometryCollectionSelectionNodes.h
+++ b/Engine/Plugins/Experimental/GeometryCollectionPlugin/Source/GeometryCollectionNodes/Public/Dataflow/GeometryCollectionSelectionNodes.h
@@ -294,10 +294,10 @@ public:
 /**
  * 
  * Selects specified bones in the GeometryCollection by using a 
- * space separated list
+ * space separated list, e.g. "0 1 2 12 23"
  *
  */
-USTRUCT(meta = (DataflowGeometryCollection))
+USTRUCT(meta = (DataflowGeometryCollection, Deprecated = "5.6"))
 struct FCollectionTransformSelectionCustomDataflowNode : public FDataflowNode
 {
 	GENERATED_USTRUCT_BODY()
@@ -308,7 +308,7 @@ public:
 	UPROPERTY(meta = (DataflowInput, DataflowOutput, DataflowPassthrough = "Collection", DataflowIntrinsic))
 	FManagedArrayCollection Collection;
 
-	/** Space separated list of bone indices to specify the selection */
+	/** Space separated list of bone indices to specify the selection, e.g. "0 1 2 3 23 34" */
 	UPROPERTY(EditAnywhere, Category = "Selection", meta=(DisplayName="Bone Indices"))
 	FString BoneIndicies = FString(); //Fix typo for v2
 
@@ -329,6 +329,36 @@ public:
 
 };
 
+/**
+ *
+ * Selects specified bones in the GeometryCollection by using a
+ * comma separated list, e.g. "0, 2, 5-10, 12-15"
+ *
+ */
+USTRUCT(meta = (DataflowGeometryCollection))
+struct FCollectionTransformSelectionCustomDataflowNode_v2 : public FDataflowNode
+{
+	GENERATED_USTRUCT_BODY()
+	DATAFLOW_NODE_DEFINE_INTERNAL(FCollectionTransformSelectionCustomDataflowNode_v2, "CollectionTransformSelectCustom", "GeometryCollection|Selection|Transform", "")
+
+public:
+	/** GeometryCollection for the selection */
+	UPROPERTY(meta = (DataflowInput, DataflowOutput, DataflowPassthrough = "Collection", DataflowIntrinsic))
+	FManagedArrayCollection Collection;
+
+	/** Comma separated list of single or a range of bone indices to specify the selection, e.g. "0, 2, 5-10, 12-15" */
+	UPROPERTY(EditAnywhere, Category = "Selection", meta = (DisplayName = "Bone Indices"))
+	FString BoneIndices;
+
+	/** Array of the selected bone indices */
+	UPROPERTY(meta = (DataflowOutput, DisplayName = "TransformSelection"))
+	FDataflowTransformSelection TransformSelection;
+
+	FCollectionTransformSelectionCustomDataflowNode_v2(const UE::Dataflow::FNodeParameters& InParam, FGuid InGuid = FGuid::NewGuid());
+
+	virtual void Evaluate(UE::Dataflow::FContext& Context, const FDataflowOutput* Out) const override;
+};
+
 /**
  *
  * Convert index array to a transform selection
diff --git a/Engine/Source/Runtime/Experimental/Dataflow/Core/Private/Dataflow/DataflowSelection.cpp b/Engine/Source/Runtime/Experimental/Dataflow/Core/Private/Dataflow/DataflowSelection.cpp
--- a/Engine/Source/Runtime/Experimental/Dataflow/Core/Private/Dataflow/DataflowSelection.cpp
+++ b/Engine/Source/Runtime/Experimental/Dataflow/Core/Private/Dataflow/DataflowSelection.cpp
@@ -138,3 +138,13 @@ void FDataflowSelection::SetWithMask(const bool Value, const FDataflowSelection&
 	}
 }
 
+void FDataflowSelection::SetSelected(TArray<int32> Indices)
+{
+	for (int32 Idx = 0; Idx < Indices.Num(); ++Idx)
+	{
+		if (SelectionArray.IsValidIndex(Indices[Idx]))
+		{
+			SetSelected(Indices[Idx]);
+		}
+	}
+}
diff --git a/Engine/Source/Runtime/Experimental/Dataflow/Core/Public/Dataflow/DataflowSelection.h b/Engine/Source/Runtime/Experimental/Dataflow/Core/Public/Dataflow/DataflowSelection.h
--- a/Engine/Source/Runtime/Experimental/Dataflow/Core/Public/Dataflow/DataflowSelection.h
+++ b/Engine/Source/Runtime/Experimental/Dataflow/Core/Public/Dataflow/DataflowSelection.h
@@ -12,12 +12,13 @@ struct FDataflowSelection
 	DATAFLOWCORE_API void Initialize(int32 NumBits, bool Value);
 	DATAFLOWCORE_API void Initialize(const FDataflowSelection& Other);
 	DATAFLOWCORE_API void Clear();
-	int32 Num() const { return SelectionArray.Num(); }
+	DATAFLOWCORE_API int32 Num() const { return SelectionArray.Num(); }
 	DATAFLOWCORE_API int32 NumSelected() const;
 	DATAFLOWCORE_API bool AnySelected() const;
-	bool IsSelected(int32 Idx) const { return SelectionArray[Idx]; }
-	void SetSelected(int32 Idx) { SelectionArray[Idx] = true; }
-	void SetNotSelected(int32 Idx) { SelectionArray[Idx] = false; }
+	DATAFLOWCORE_API bool IsSelected(int32 Idx) const { return SelectionArray[Idx]; }
+	DATAFLOWCORE_API void SetSelected(int32 Idx) { SelectionArray[Idx] = true; }
+	DATAFLOWCORE_API void SetSelected(TArray<int32> Indices);
+	DATAFLOWCORE_API void SetNotSelected(int32 Idx) { SelectionArray[Idx] = false; }
 	DATAFLOWCORE_API void AsArray(TArray<int32>& SelectionArr) const;
 	DATAFLOWCORE_API TArray<int32> AsArray() const;
 	
@@ -36,9 +37,9 @@ struct FDataflowSelection
 	DATAFLOWCORE_API void XOR(const FDataflowSelection& Other, FDataflowSelection& Result) const;
 	// subtract the selected elements of 'Other' from this selection
 	DATAFLOWCORE_API void Subtract(const FDataflowSelection& Other, FDataflowSelection& Result) const;
-	void Invert() { SelectionArray.BitwiseNOT(); }
+	DATAFLOWCORE_API void Invert() { SelectionArray.BitwiseNOT(); }
 	DATAFLOWCORE_API void SetWithMask(const bool Value, const FDataflowSelection& Mask);
-	const TBitArray<>& GetBitArray() const { return SelectionArray; };
+	DATAFLOWCORE_API const TBitArray<>& GetBitArray() const { return SelectionArray; };
 
 private:
 
