Add DynamicMeshProcessor concept as a way to define BP-based generic processors for dynamic meshes, and add a dataflow node to use these BPs in a dataflow graph

+ add support for dataflow nodes to receive on-property-changed callbacks, to support this functionality

diff --git a/Engine/Plugins/Experimental/Dataflow/Source/DataflowEditor/Private/Dataflow/DataflowEditorCommands.cpp b/Engine/Plugins/Experimental/Dataflow/Source/DataflowEditor/Private/Dataflow/DataflowEditorCommands.cpp
--- a/Engine/Plugins/Experimental/Dataflow/Source/DataflowEditor/Private/Dataflow/DataflowEditorCommands.cpp
+++ b/Engine/Plugins/Experimental/Dataflow/Source/DataflowEditor/Private/Dataflow/DataflowEditorCommands.cpp
@@ -335,6 +335,11 @@ void FDataflowEditorCommands::OnPropertyValueChanged(UDataflow* OutDataflow, TSh
 						DataflowNode->Invalidate();
 						OutLastNodeTimestamp = UE::Dataflow::FTimestamp::Invalid;
 
+						if (Context.IsValid())
+						{
+							DataflowNode->OnPropertyChanged(*Context, InPropertyChangedEvent);
+						}
+
 						// Reflect the active state on the drawing of the node
 						if (DataflowNode->bActive != Node->IsNodeEnabled())
 						{
diff --git a/Engine/Plugins/Experimental/GeometryCollectionPlugin/Source/GeometryCollectionNodes/Private/Dataflow/GeometryCollectionMeshNodes.cpp b/Engine/Plugins/Experimental/GeometryCollectionPlugin/Source/GeometryCollectionNodes/Private/Dataflow/GeometryCollectionMeshNodes.cpp
--- a/Engine/Plugins/Experimental/GeometryCollectionPlugin/Source/GeometryCollectionNodes/Private/Dataflow/GeometryCollectionMeshNodes.cpp
+++ b/Engine/Plugins/Experimental/GeometryCollectionPlugin/Source/GeometryCollectionNodes/Private/Dataflow/GeometryCollectionMeshNodes.cpp
@@ -3,6 +3,7 @@
 #include "Dataflow/GeometryCollectionMeshNodes.h"
 #include "Dataflow/DataflowCore.h"
 
+#include "Engine/Blueprint.h"
 #include "Engine/StaticMesh.h"
 #include "GeometryCollection/GeometryCollectionObject.h"
 #include "GeometryCollection/ManagedArrayCollection.h"
@@ -50,6 +51,8 @@ namespace UE::Dataflow
 		DATAFLOW_NODE_REGISTER_CREATION_FACTORY(FMeshBooleanDataflowNode);
 		DATAFLOW_NODE_REGISTER_CREATION_FACTORY(FMeshCopyToPointsDataflowNode);
 		DATAFLOW_NODE_REGISTER_CREATION_FACTORY(FGetMeshDataDataflowNode);
+		DATAFLOW_NODE_REGISTER_CREATION_FACTORY(FApplyMeshProcessorToMeshDataflowNode);
+
 
 		// Mesh
 		DATAFLOW_NODE_REGISTER_CREATION_FACTORY_NODE_COLORS_BY_CATEGORY("Mesh", FLinearColor(1.f, 0.16f, 0.05f), CDefaultNodeBodyTintColor);
@@ -417,3 +420,77 @@ void FGetMeshDataDataflowNode::Evaluate(UE::Dataflow::FContext& Context, const F
 	}
 }
 
+void FMeshProcessorDataflowNodeBase::OnPropertyChanged(UE::Dataflow::FContext& Context, const FPropertyChangedEvent& PropertyChangedEvent)
+{
+	if (PropertyChangedEvent.ChangeType == EPropertyChangeType::ValueSet 
+		&& PropertyChangedEvent.GetPropertyName() == GET_MEMBER_NAME_CHECKED(FMeshProcessorDataflowNodeBase, MeshProcessor))
+	{
+		if (MeshProcessor)
+		{
+			MeshProcessorInstance = NewObject<UDynamicMeshProcessorBlueprint>(OwningObject, MeshProcessor, NAME_None, RF_Transactional);
+			TeardownBlueprintEvent();
+			SetupBlueprintEvent();
+		}
+		else
+		{
+			MeshProcessorInstance = nullptr;
+		}
+	}
+}
+
+void FMeshProcessorDataflowNodeBase::SetupBlueprintEvent()
+{
+#if WITH_EDITOR
+	if (MeshProcessor)
+	{
+		if (UBlueprint* Blueprint = Cast<UBlueprint>(MeshProcessor->ClassGeneratedBy))
+		{
+			BlueprintChangeDelegateHandle = Blueprint->OnChanged().AddLambda([this](UBlueprint* BP)
+			{
+				Invalidate();
+			});
+		}
+	}
+#endif
+}
+
+void FMeshProcessorDataflowNodeBase::TeardownBlueprintEvent()
+{
+#if WITH_EDITOR
+	if (MeshProcessor && BlueprintChangeDelegateHandle.IsValid())
+	{
+		if (UBlueprint* Blueprint = Cast<UBlueprint>(MeshProcessor->ClassGeneratedBy))
+		{
+			Blueprint->OnChanged().Remove(BlueprintChangeDelegateHandle);
+		}
+	}
+#endif
+}
+
+void FApplyMeshProcessorToMeshDataflowNode::Evaluate(UE::Dataflow::FContext& Context, const FDataflowOutput* Out) const
+{
+	if (Out->IsA<TObjectPtr<UDynamicMesh>>(&Mesh))
+	{
+		if (TObjectPtr<const UDynamicMesh> InMesh = GetValue<TObjectPtr<UDynamicMesh>>(Context, &Mesh))
+		{
+			if (!MeshProcessorInstance)
+			{
+				SafeForwardInput(Context, &Mesh, &Mesh);
+				return;
+			}
+
+			// Creating a new mesh object from InMesh
+			TObjectPtr<UDynamicMesh> NewMesh = NewObject<UDynamicMesh>();
+			NewMesh->SetMesh(InMesh->GetMeshRef());
+
+			bool bFailed = false;
+			MeshProcessorInstance->ProcessDynamicMesh(NewMesh, bFailed);
+
+			SetValue(Context, NewMesh, &Mesh);
+		}
+		else
+		{
+			SetValue(Context, TObjectPtr<UDynamicMesh>(NewObject<UDynamicMesh>()), &Mesh);
+		}
+	}
+}
diff --git a/Engine/Plugins/Experimental/GeometryCollectionPlugin/Source/GeometryCollectionNodes/Public/Dataflow/GeometryCollectionMeshNodes.h b/Engine/Plugins/Experimental/GeometryCollectionPlugin/Source/GeometryCollectionNodes/Public/Dataflow/GeometryCollectionMeshNodes.h
--- a/Engine/Plugins/Experimental/GeometryCollectionPlugin/Source/GeometryCollectionNodes/Public/Dataflow/GeometryCollectionMeshNodes.h
+++ b/Engine/Plugins/Experimental/GeometryCollectionPlugin/Source/GeometryCollectionNodes/Public/Dataflow/GeometryCollectionMeshNodes.h
@@ -7,6 +7,7 @@
 #include "GeometryCollection/ManagedArrayCollection.h"
 #include "Dataflow/DataflowSelection.h"
 #include "DynamicMesh/DynamicMesh3.h"
+#include "DynamicMeshProcessor.h"
 
 #include "GeometryCollectionMeshNodes.generated.h"
 
@@ -57,6 +58,82 @@ public:
 };
 
 
+/**
+ * Base class for nodes that applying Geometry Script Mesh Processors
+ */
+USTRUCT()
+struct FMeshProcessorDataflowNodeBase : public FDataflowNode
+{
+	GENERATED_USTRUCT_BODY()
+
+public:
+
+	FMeshProcessorDataflowNodeBase() = default;
+
+	FMeshProcessorDataflowNodeBase(const UE::Dataflow::FNodeParameters& InParam, FGuid InGuid = FGuid::NewGuid())
+		: FDataflowNode(InParam, InGuid), OwningObject(InParam.OwningObject)
+	{
+	}
+
+	~FMeshProcessorDataflowNodeBase()
+	{
+		TeardownBlueprintEvent();
+	}
+
+protected:
+
+	UPROPERTY(EditAnywhere, Category = "Processor Type")
+	TSubclassOf<UDynamicMeshProcessorBlueprint> MeshProcessor = nullptr;
+
+	UPROPERTY(EditAnywhere, Instanced, NoClear, Category = "Instance", meta = (ShowOnlyInnerProperties))
+	TObjectPtr<UDynamicMeshProcessorBlueprint> MeshProcessorInstance;
+
+	virtual void PostSerialize(const FArchive& Ar) override final
+	{
+		Super::PostSerialize(Ar);
+		SetupBlueprintEvent();
+	}
+
+private:
+
+	virtual void OnPropertyChanged(UE::Dataflow::FContext& Context, const FPropertyChangedEvent& PropertyChangedEvent) override;
+
+	// Handling for the selected blueprint being changed under the node (e.g., recompiled)
+	FDelegateHandle BlueprintChangeDelegateHandle;
+	void TeardownBlueprintEvent();
+	void SetupBlueprintEvent();
+
+	// Remember the parent UObject so that we can properly parent MeshProcessorInstance when the MeshProcessor changes 
+	UPROPERTY(Transient)
+	TObjectPtr<UObject> OwningObject;
+};
+
+/**
+ * Apply a Geometry Script Mesh Processors to an input UDynamicMesh
+ */
+USTRUCT(meta = (DataflowGeometryCollection))
+struct FApplyMeshProcessorToMeshDataflowNode final : public FMeshProcessorDataflowNodeBase
+{
+	GENERATED_USTRUCT_BODY()
+	DATAFLOW_NODE_DEFINE_INTERNAL(FApplyMeshProcessorToMeshDataflowNode, "ApplyGeometryScriptToMesh", "Mesh|Utilities", "")
+	DATAFLOW_NODE_RENDER_TYPE("SurfaceRender", FName("FDynamicMesh3"), "Mesh")
+
+public:
+
+	/** Input/Output mesh */
+	UPROPERTY(meta = (DataflowInput, DataflowOutput, DataflowPassthrough = "Mesh", DataflowIntrinsic))
+	TObjectPtr<UDynamicMesh> Mesh;
+
+	FApplyMeshProcessorToMeshDataflowNode(const UE::Dataflow::FNodeParameters& InParam, FGuid InGuid = FGuid::NewGuid())
+		: FMeshProcessorDataflowNodeBase(InParam, InGuid)
+	{
+		RegisterInputConnection(&Mesh);
+		RegisterOutputConnection(&Mesh, &Mesh);
+	}
+
+	virtual void Evaluate(UE::Dataflow::FContext& Context, const FDataflowOutput* Out) const override;
+};
+
 /**
  *
  * Converts a BoundingBox into a DynamicMesh
diff --git a/Engine/Source/Runtime/Experimental/Dataflow/Core/Public/Dataflow/DataflowNode.h b/Engine/Source/Runtime/Experimental/Dataflow/Core/Public/Dataflow/DataflowNode.h
--- a/Engine/Source/Runtime/Experimental/Dataflow/Core/Public/Dataflow/DataflowNode.h
+++ b/Engine/Source/Runtime/Experimental/Dataflow/Core/Public/Dataflow/DataflowNode.h
@@ -517,6 +517,9 @@ PRAGMA_ENABLE_DEPRECATION_WARNINGS
 
 	DATAFLOWCORE_API virtual void ValidateProperties();
 
+	// Called when UPROPERTY members of the dataflow node have been changed in the editor
+	virtual void OnPropertyChanged(UE::Dataflow::FContext& Context, const FPropertyChangedEvent& InPropertyChangedEvent) {}
+
 	bool HasValidConnections() const { return bHasValidConnections; }
 
 	virtual bool IsA(FName InType) const 
diff --git a/Engine/Source/Runtime/GeometryFramework/Public/DynamicMeshProcessor.h b/Engine/Source/Runtime/GeometryFramework/Public/DynamicMeshProcessor.h
new file mode 100644
--- /dev/null
+++ b/Engine/Source/Runtime/GeometryFramework/Public/DynamicMeshProcessor.h
@@ -0,0 +1,25 @@
+// Copyright Epic Games, Inc. All Rights Reserved.
+
+#pragma once
+
+#include "UDynamicMesh.h"
+
+#include "DynamicMeshProcessor.generated.h"
+
+
+
+// Blueprints with this parent class can define general processing to apply to a dynamic mesh
+// which can then be used to define procedural operations e.g. in Dataflow or other contexts
+UCLASS(Abstract, BlueprintType, Blueprintable)
+class GEOMETRYFRAMEWORK_API UDynamicMeshProcessorBlueprint : public UObject
+{
+	GENERATED_BODY()
+public:
+
+	/// Apply some processing to change the input mesh
+	/// @param TargetMesh Mesh to update (in place)
+	/// @param bFailed Set to true to report that the processing has failed
+	UFUNCTION(BlueprintImplementableEvent, BlueprintCallable, Category = "MeshProcessor")
+	void ProcessDynamicMesh(UDynamicMesh* TargetMesh, bool& bFailed);
+};
+
