Dataflow UV nodes : adjust the reference resolution for the gutter in the UV unwrap node to address performance issues

diff --git a/Engine/Plugins/Experimental/GeometryCollectionPlugin/Source/GeometryCollectionNodes/Private/Dataflow/GeometryCollectionUVNodes.cpp b/Engine/Plugins/Experimental/GeometryCollectionPlugin/Source/GeometryCollectionNodes/Private/Dataflow/GeometryCollectionUVNodes.cpp
--- a/Engine/Plugins/Experimental/GeometryCollectionPlugin/Source/GeometryCollectionNodes/Private/Dataflow/GeometryCollectionUVNodes.cpp
+++ b/Engine/Plugins/Experimental/GeometryCollectionPlugin/Source/GeometryCollectionNodes/Private/Dataflow/GeometryCollectionUVNodes.cpp
@@ -92,7 +92,7 @@ void FAutoUnwrapUVDataflowNode::Evaluate(UE::Dataflow::FContext& Context, const
 		check(GeomCollection.IsValid());
 
 		// this is used as reference for the gutter size that is expressed in pixels for a specific resolution 
-		constexpr int32 ReferenceResolution = 4096;
+		constexpr int32 ReferenceResolution = 512;
 		constexpr bool bRecreateUVsForDegenerateIslands = true;
 
 		if (IsConnected(&FaceSelection))
diff --git a/Engine/Plugins/Experimental/GeometryCollectionPlugin/Source/GeometryCollectionNodes/Public/Dataflow/GeometryCollectionUVNodes.h b/Engine/Plugins/Experimental/GeometryCollectionPlugin/Source/GeometryCollectionNodes/Public/Dataflow/GeometryCollectionUVNodes.h
--- a/Engine/Plugins/Experimental/GeometryCollectionPlugin/Source/GeometryCollectionNodes/Public/Dataflow/GeometryCollectionUVNodes.h
+++ b/Engine/Plugins/Experimental/GeometryCollectionPlugin/Source/GeometryCollectionNodes/Public/Dataflow/GeometryCollectionUVNodes.h
@@ -59,7 +59,7 @@ private:
 	UPROPERTY(EditAnywhere, Category = Options, meta = (DataflowInput))
 	int32 UVChannel = 0;
 
-	/** Gutter Size in pixels for a 4096x4096 reference resolution */
+	/** Approximate space to leave between UV islands, measured in texels for 512x512 texture */
 	UPROPERTY(EditAnywhere, Category = Options)
 	int32 GutterSize = 1;
 
