Dataflow: Add texture Bake collection node

diff --git a/Engine/Plugins/Experimental/Dataflow/Source/DataflowEnginePlugin/Private/Dataflow/DataflowEnginePlugin.cpp b/Engine/Plugins/Experimental/Dataflow/Source/DataflowEnginePlugin/Private/Dataflow/DataflowEnginePlugin.cpp
--- a/Engine/Plugins/Experimental/Dataflow/Source/DataflowEnginePlugin/Private/Dataflow/DataflowEnginePlugin.cpp
+++ b/Engine/Plugins/Experimental/Dataflow/Source/DataflowEnginePlugin/Private/Dataflow/DataflowEnginePlugin.cpp
@@ -29,6 +29,7 @@ IMPLEMENT_MODULE( FDataflowEnginePlugin, DataflowEnginePlugin)
 void FDataflowEnginePlugin::StartupModule()
 {
 	UE::Dataflow::ContextCachingCallbacks();
+	FModuleManager::Get().LoadModule("DataflowEngine");
 	FModuleManager::Get().LoadModule("DataflowSimulation");
 }
 
diff --git a/Engine/Plugins/Experimental/GeometryCollectionPlugin/Source/GeometryCollectionNodes/Private/Dataflow/GeometryCollectionNodesPlugin.cpp b/Engine/Plugins/Experimental/GeometryCollectionPlugin/Source/GeometryCollectionNodes/Private/Dataflow/GeometryCollectionNodesPlugin.cpp
--- a/Engine/Plugins/Experimental/GeometryCollectionPlugin/Source/GeometryCollectionNodes/Private/Dataflow/GeometryCollectionNodesPlugin.cpp
+++ b/Engine/Plugins/Experimental/GeometryCollectionPlugin/Source/GeometryCollectionNodes/Private/Dataflow/GeometryCollectionNodesPlugin.cpp
@@ -26,6 +26,7 @@
 #include "Dataflow/GeometryCollectionSkeletalMeshToCollectionNode.h"
 #include "Dataflow/GeometryCollectionSkeletonToCollectionNode.h"
 #include "Dataflow/GeometryCollectionStaticMeshToCollectionNode.h"
+#include "Dataflow/GeometryCollectionTextureNodes.h"
 #include "Dataflow/GeometryCollectionTransferVertexAttributeNode.h"
 #include "Dataflow/GeometryCollectionTriangleBoundaryIndicesNode.h"
 #include "Dataflow/GeometryCollectionUtilityNodes.h"
@@ -61,6 +62,7 @@ void IGeometryCollectionNodesPlugin::StartupModule()
 	UE::Dataflow::GeometryCollectionDebugNodes();
 	UE::Dataflow::GeometryCollectionSamplingNodes();
 	UE::Dataflow::RegisterGeometryCollectionUVNodes();
+	UE::Dataflow::RegisterGeometryCollectionTextureNodes();
 	DATAFLOW_NODE_REGISTER_CREATION_FACTORY(FCreateColorArrayFromFloatArrayDataflowNode);
 	DATAFLOW_NODE_REGISTER_CREATION_FACTORY(FGeometryCollectionVertexScalarToVertexIndicesNode);
 	DATAFLOW_NODE_REGISTER_CREATION_FACTORY(FSetVertexColorFromFloatArrayDataflowNode);
diff --git a/Engine/Plugins/Experimental/GeometryCollectionPlugin/Source/GeometryCollectionNodes/Private/Dataflow/GeometryCollectionTextureNodes.cpp b/Engine/Plugins/Experimental/GeometryCollectionPlugin/Source/GeometryCollectionNodes/Private/Dataflow/GeometryCollectionTextureNodes.cpp
new file mode 100644
--- /dev/null
+++ b/Engine/Plugins/Experimental/GeometryCollectionPlugin/Source/GeometryCollectionNodes/Private/Dataflow/GeometryCollectionTextureNodes.cpp
@@ -0,0 +1,99 @@
+// Copyright Epic Games, Inc. All Rights Reserved.
+
+#include "Dataflow/GeometryCollectionTextureNodes.h"
+#include "GeometryCollection/GeometryCollectionObject.h"
+
+namespace UE::Dataflow
+{
+	void RegisterGeometryCollectionTextureNodes()
+	{
+		DATAFLOW_NODE_REGISTER_CREATION_FACTORY(FBakeTextureFromCollectionDataflowNode);
+
+		static const FLinearColor CDefaultNodeBodyTintColor = FLinearColor(0.f, 0.f, 0.f, 0.5f);
+		DATAFLOW_NODE_REGISTER_CREATION_FACTORY_NODE_COLORS_BY_CATEGORY("GeometryCollection|Texture", FLinearColor(.1f, 1.f, 0.6f), CDefaultNodeBodyTintColor);
+	}
+}
+
+////////////////////////////////////////////////////////////////////////////////////////////
+
+FBakeTextureFromCollectionDataflowNode::FBakeTextureFromCollectionDataflowNode(const UE::Dataflow::FNodeParameters& InParam, FGuid InGuid)
+	: FDataflowNode(InParam, InGuid)
+{
+	RegisterInputConnection(&Collection);
+	RegisterInputConnection(&FaceSelection);
+	RegisterInputConnection(&UVChannel);
+	RegisterInputConnection(&GutterSize);
+	RegisterInputConnection(&Resolution);
+
+	RegisterOutputConnection(&Collection, &Collection);
+	RegisterOutputConnection(&Image);
+}
+
+void FBakeTextureFromCollectionDataflowNode::Evaluate(UE::Dataflow::FContext& Context, const FDataflowOutput* Out) const
+{
+	if (Out->IsA(&Collection))
+	{
+		SafeForwardInput(Context, &Collection, &Collection);
+	}
+	else if (Out->IsA(&Image))
+	{
+		FManagedArrayCollection InCollection = GetValue(Context, &Collection);
+		const int32 InUVChannel = GetValue(Context, &UVChannel);
+		const int32 InResolution = (int32)GetValue(Context, &Resolution);
+		const int32 InGutterSize = GetValue(Context, &GutterSize);
+
+		TUniquePtr<FGeometryCollection> GeomCollection = TUniquePtr<FGeometryCollection>(InCollection.NewCopy<FGeometryCollection>());
+		check(GeomCollection.IsValid());
+
+		const UE::Geometry::FIndex4i BakeAttributes{ (int32)RedChannel, (int32)GreenChannel, (int32)BlueChannel, (int32)AlphaChannel };
+
+		UE::PlanarCut::FTextureAttributeSettings BakeAttributeSettings;
+		BakeAttributeSettings.ToExternal_MaxDistance = MaxDistance;
+		BakeAttributeSettings.AO_Rays = OcclusionRays;
+		BakeAttributeSettings.AO_BlurRadius = OcclusionBlurRadius;
+		BakeAttributeSettings.Curvature_BlurRadius = CurvatureBlurRadius;
+		BakeAttributeSettings.Curvature_SmoothingSteps = SmoothingIterations;
+		BakeAttributeSettings.Curvature_VoxelRes = VoxelResolution;
+		BakeAttributeSettings.Curvature_ThicknessFactor = ThicknessFactor;
+		BakeAttributeSettings.Curvature_MaxValue = MaxCurvature;
+		BakeAttributeSettings.ClearGutterChannel = 3; // default clear the gutters for the alpha channel, so it shows more clearly the island boundaries
+
+		UE::Geometry::TImageBuilder<FVector4f> ImageBuilder;
+		ImageBuilder.SetDimensions({ InResolution, InResolution });
+		ImageBuilder.Clear(FVector4f(1, 0, 0, 0));
+
+		if (IsConnected(&FaceSelection))
+		{
+			const FDataflowFaceSelection& InFaceSelection = GetValue(Context, &FaceSelection);
+			if (InFaceSelection.Num() > 0)
+			{
+				TArray<bool> FacesToBake;
+				FacesToBake.SetNumUninitialized(InFaceSelection.Num());
+				for (int32 Index = 0; Index < InFaceSelection.Num(); ++Index)
+				{
+					FacesToBake[Index] = InFaceSelection.IsSelected(Index);
+				}
+
+				UE::PlanarCut::TextureSpecifiedFaces(
+					InUVChannel, *GeomCollection, InGutterSize,
+					BakeAttributes, BakeAttributeSettings, ImageBuilder,
+					FacesToBake, nullptr /* Progress */);
+			}
+		}
+		else
+		{
+			// selection not connected, bake all faces
+			UE::PlanarCut::TextureSpecifiedFaces(
+				InUVChannel, *GeomCollection, InGutterSize,
+				BakeAttributes, BakeAttributeSettings, ImageBuilder,
+				UE::PlanarCut::ETargetFaces::AllFaces, {}, nullptr /* Progress */);
+		}
+
+		// copy the image data over and set the corresponding value
+		FDataflowImage OutImage;
+		OutImage.CreateRGBA32F(InResolution, InResolution);
+		OutImage.CopyRGBAPixels(ImageBuilder.GetImageBuffer());
+		SetValue(Context, MoveTemp(OutImage), &Image);
+	}
+}
+
diff --git a/Engine/Plugins/Experimental/GeometryCollectionPlugin/Source/GeometryCollectionNodes/Public/Dataflow/GeometryCollectionTextureNodes.h b/Engine/Plugins/Experimental/GeometryCollectionPlugin/Source/GeometryCollectionNodes/Public/Dataflow/GeometryCollectionTextureNodes.h
new file mode 100644
--- /dev/null
+++ b/Engine/Plugins/Experimental/GeometryCollectionPlugin/Source/GeometryCollectionNodes/Public/Dataflow/GeometryCollectionTextureNodes.h
@@ -0,0 +1,158 @@
+// Copyright Epic Games, Inc. All Rights Reserved.
+
+#pragma once 
+
+#include "Dataflow/DataflowEngine.h"
+#include "Dataflow/DataflowSelection.h"
+#include "Dataflow/DataflowImage.h"
+#include "GeometryCollection/ManagedArrayCollection.h"
+#include "FractureAutoUV.h"
+#include "Image/ImageBuilder.h"
+#include "Engine/Texture2D.h"
+
+#include "GeometryCollectionTextureNodes.generated.h"
+
+UENUM()
+enum class ECollectionBakeTextureAttribute: int32
+{
+	// no attribute selected
+	None = (int32)UE::PlanarCut::EBakeAttributes::None,
+
+	// Phi value : distance to the closest face 
+	DistanceToExternal = (int32)UE::PlanarCut::EBakeAttributes::DistanceToExternal,
+
+	// Ambient occlusion
+	AmbientOcclusion = (int32)UE::PlanarCut::EBakeAttributes::AmbientOcclusion,
+
+	// surface smoothed curvature
+	Curvature = (int32)UE::PlanarCut::EBakeAttributes::Curvature,
+	
+	// Normal X coordinate ( object space ) 
+	NormalX = (int32)UE::PlanarCut::EBakeAttributes::NormalX,
+	
+	// Normal Y coordinate ( object space ) 
+	NormalY = (int32)UE::PlanarCut::EBakeAttributes::NormalY,
+	
+	// Normal Z coordinate ( object space ) 
+	NormalZ = (int32)UE::PlanarCut::EBakeAttributes::NormalZ,
+
+	// Position X coordinate ( relative to object local bounds ) 
+	PositionX = (int32)UE::PlanarCut::EBakeAttributes::PositionX,
+	
+	// Position Y coordinate ( relative to object local bounds ) 
+	PositionY = (int32)UE::PlanarCut::EBakeAttributes::PositionY,
+	
+	// Position Z coordinate ( relative to object local bounds ) 
+	PositionZ = (int32)UE::PlanarCut::EBakeAttributes::PositionZ,
+};
+
+/*
+* Bake a texture from a geometry collection 
+* Output to a 4 channels Image object  ( RGBA )
+*/
+USTRUCT()
+struct FBakeTextureFromCollectionDataflowNode : public FDataflowNode
+{
+	GENERATED_USTRUCT_BODY()
+	DATAFLOW_NODE_DEFINE_INTERNAL(FBakeTextureFromCollectionDataflowNode, "BakeTextureFromCollection", "GeometryCollection|Texture", "")
+
+private:
+	/** Target Collection */
+	UPROPERTY(meta = (DataflowInput, DataflowOutput, DataflowPassthrough = "Collection", DataflowIntrinsic))
+	FManagedArrayCollection Collection;
+
+	/** Output image with the bake attributes */
+	UPROPERTY(meta = (DataflowOutput))
+	FDataflowImage Image;
+
+	/** selection of faces to bake : if not connected, all faces will be used */
+	UPROPERTY(meta = (DataflowInput))
+	FDataflowFaceSelection FaceSelection;
+
+	/** resolution of the image to bake */
+	UPROPERTY(EditAnywhere, Category = ImageOptions, meta = (DataflowInput))
+	EDataflowImageResolution Resolution = EDataflowImageResolution::Resolution512;
+
+	/** Approximate space to leave between UV islands, measured in texels */
+	UPROPERTY(EditAnywhere, Category = ChannelOptions, meta = (DataflowInput, ClampMin = 1, UIMin = 1))
+	int32 GutterSize = 1;
+
+	/** Index of the added UV channel */
+	UPROPERTY(EditAnywhere, Category = ChannelOptions, meta = (DataflowInput, ClampMin = 0, ClampMax = 7, UIMin = 0, UIMax = 7))
+	int32 UVChannel = 0;
+
+	/** Attribute to bake in the red channel */
+	UPROPERTY(EditAnywhere, Category = ChannelOptions)
+	ECollectionBakeTextureAttribute RedChannel = ECollectionBakeTextureAttribute::DistanceToExternal;
+
+	/** Attribute to bake in the green channel */
+	UPROPERTY(EditAnywhere, Category = ChannelOptions)
+	ECollectionBakeTextureAttribute GreenChannel = ECollectionBakeTextureAttribute::AmbientOcclusion;
+
+	/** Attribute to bake in the blue channel */
+	UPROPERTY(EditAnywhere, Category = ChannelOptions)
+	ECollectionBakeTextureAttribute BlueChannel = ECollectionBakeTextureAttribute::Curvature;
+
+	/** Attribute to bake in the alpha channel */
+	UPROPERTY(EditAnywhere, Category = ChannelOptions)
+	ECollectionBakeTextureAttribute AlphaChannel = ECollectionBakeTextureAttribute::None;
+
+	/** Max distance to search for the outer mesh surface */
+	UPROPERTY(EditAnywhere, Category = DistToOuterSettings, meta = (
+		EditCondition = "RedChannel == ECollectionBakeTextureAttribute::DistanceToExternal || GreenChannel == ECollectionBakeTextureAttribute::DistanceToExternal || BlueChannel == ECollectionBakeTextureAttribute::DistanceToExternal || AlphaChannel == ECollectionBakeTextureAttribute::DistanceToExternal", 
+		UIMin = "1", UIMax = "100", ClampMin = ".01", ClampMax = "1000"))
+	double MaxDistance = 100;
+
+	/** Number of occlusion rays */
+	UPROPERTY(EditAnywhere, Category = AmbientOcclusionSettings, meta = (
+		EditCondition = "RedChannel == ECollectionBakeTextureAttribute::AmbientOcclusion || GreenChannel == ECollectionBakeTextureAttribute::AmbientOcclusion || BlueChannel == ECollectionBakeTextureAttribute::AmbientOcclusion || AlphaChannel == ECollectionBakeTextureAttribute::AmbientOcclusion",
+		UIMin = "1", UIMax = "1024", ClampMin = "0", ClampMax = "50000"))
+	int OcclusionRays = 16;
+
+	/** Pixel Radius of Gaussian Blur Kernel applied to AO map (0 will apply no blur) */
+	UPROPERTY(EditAnywhere, Category = AmbientOcclusionSettings, meta = (
+		EditCondition = "RedChannel == ECollectionBakeTextureAttribute::AmbientOcclusion || GreenChannel == ECollectionBakeTextureAttribute::AmbientOcclusion || BlueChannel == ECollectionBakeTextureAttribute::AmbientOcclusion || AlphaChannel == ECollectionBakeTextureAttribute::AmbientOcclusion",
+		UIMin = "0", UIMax = "10.0", ClampMin = "0", ClampMax = "100.0"))
+	double OcclusionBlurRadius = 2.25;
+
+	/** Pixel Radius of Gaussian Blur Kernel applied to Curvature map (0 will apply no blur) */
+	UPROPERTY(EditAnywhere, Category = SmoothedCurvatureSettings, meta = (
+		EditCondition = "RedChannel == ECollectionBakeTextureAttribute::Curvature || GreenChannel == ECollectionBakeTextureAttribute::Curvature || BlueChannel == ECollectionBakeTextureAttribute::Curvature || AlphaChannel == ECollectionBakeTextureAttribute::Curvature",
+		UIMin = "0", UIMax = "10.0", ClampMin = "0", ClampMax = "100.0"))
+	double CurvatureBlurRadius = 2.25;
+
+	/** Voxel resolution of smoothed shape representation */
+	UPROPERTY(EditAnywhere, Category = SmoothedCurvatureSettings, meta = (
+		EditCondition = "RedChannel == ECollectionBakeTextureAttribute::Curvature || GreenChannel == ECollectionBakeTextureAttribute::Curvature || BlueChannel == ECollectionBakeTextureAttribute::Curvature || AlphaChannel == ECollectionBakeTextureAttribute::Curvature",
+		UIMin = "8", UIMax = "512", ClampMin = "4", ClampMax = "1024"))
+	int VoxelResolution = 128;
+
+	/** Amount of smoothing iterations to apply before computing curvature */
+	UPROPERTY(EditAnywhere, Category = SmoothedCurvatureSettings, meta = (
+		EditCondition = "RedChannel == ECollectionBakeTextureAttribute::Curvature || GreenChannel == ECollectionBakeTextureAttribute::Curvature || BlueChannel == ECollectionBakeTextureAttribute::Curvature || AlphaChannel == ECollectionBakeTextureAttribute::Curvature",
+		UIMin = "2", UIMax = "100", ClampMin = "2", ClampMax = "1000"))
+	int SmoothingIterations = 10;
+
+	/** Distance to search for correspondence between fractured shape and smoothed shape, as factor of voxel size */
+	UPROPERTY(EditAnywhere, Category = SmoothedCurvatureSettings, meta = (
+		EditCondition = "RedChannel == ECollectionBakeTextureAttribute::Curvature || GreenChannel == ECollectionBakeTextureAttribute::Curvature || BlueChannel == ECollectionBakeTextureAttribute::Curvature || AlphaChannel == ECollectionBakeTextureAttribute::Curvature",
+		UIMin = "2", UIMax = "10.0", ClampMin = "1", ClampMax = "100.0"))
+	double ThicknessFactor = 4;
+
+	/** Curvatures in the range [-MaxCurvature, MaxCurvature] will be mapped from [0,1]. Values outside that range will be clamped */
+	UPROPERTY(EditAnywhere, Category = SmoothedCurvatureSettings, meta = (
+		EditCondition = "RedChannel == ECollectionBakeTextureAttribute::Curvature || GreenChannel == ECollectionBakeTextureAttribute::Curvature || BlueChannel == ECollectionBakeTextureAttribute::Curvature || AlphaChannel == ECollectionBakeTextureAttribute::Curvature",
+		UIMin = ".01", UIMax = "1", ClampMin = ".0001", ClampMax = "10"))
+	double MaxCurvature = .1;
+
+	virtual void Evaluate(UE::Dataflow::FContext& Context, const FDataflowOutput* Out) const override;
+
+public:
+	FBakeTextureFromCollectionDataflowNode(const UE::Dataflow::FNodeParameters& InParam, FGuid InGuid = FGuid::NewGuid());
+};
+
+namespace UE::Dataflow
+{
+	void RegisterGeometryCollectionTextureNodes();
+}
+
diff --git a/Engine/Source/Runtime/Experimental/Dataflow/Engine/DataflowEngine.Build.cs b/Engine/Source/Runtime/Experimental/Dataflow/Engine/DataflowEngine.Build.cs
--- a/Engine/Source/Runtime/Experimental/Dataflow/Engine/DataflowEngine.Build.cs
+++ b/Engine/Source/Runtime/Experimental/Dataflow/Engine/DataflowEngine.Build.cs
@@ -15,7 +15,8 @@ namespace UnrealBuildTool.Rules
 					"SlateCore",
 					"Engine",
 					"DataflowCore",
-					"Chaos"
+					"Chaos",
+					"ImageCore"
 				}
 			);
 		}
diff --git a/Engine/Source/Runtime/Experimental/Dataflow/Engine/Private/Dataflow/DataflowEngineModule.cpp b/Engine/Source/Runtime/Experimental/Dataflow/Engine/Private/Dataflow/DataflowEngineModule.cpp
--- a/Engine/Source/Runtime/Experimental/Dataflow/Engine/Private/Dataflow/DataflowEngineModule.cpp
+++ b/Engine/Source/Runtime/Experimental/Dataflow/Engine/Private/Dataflow/DataflowEngineModule.cpp
@@ -5,6 +5,7 @@
 #include "Modules/ModuleInterface.h"
 #include "Modules/ModuleManager.h"
 
+#include "Dataflow/DataflowTextureAssetNodes.h"
 
 /**
  * The public interface to this module
@@ -15,7 +16,11 @@ public:
 
 	// IModuleInterface interface
 
-	virtual void StartupModule() override { }
+	virtual void StartupModule() override
+	{
+		UE::Dataflow::RegisterTextureAssetNodes();
+	}
+
 	virtual void ShutdownModule() override { }
 
 	virtual bool SupportsDynamicReloading() override
diff --git a/Engine/Source/Runtime/Experimental/Dataflow/Engine/Private/Dataflow/DataflowTextureAssetNodes.cpp b/Engine/Source/Runtime/Experimental/Dataflow/Engine/Private/Dataflow/DataflowTextureAssetNodes.cpp
new file mode 100644
--- /dev/null
+++ b/Engine/Source/Runtime/Experimental/Dataflow/Engine/Private/Dataflow/DataflowTextureAssetNodes.cpp
@@ -0,0 +1,73 @@
+// Copyright Epic Games, Inc. All Rights Reserved.
+
+#include "Dataflow/DataflowTextureAssetNodes.h"
+#include "Dataflow/DataflowNodeFactory.h"
+#include "Engine/Texture2D.h"
+
+namespace UE::Dataflow
+{
+	void RegisterTextureAssetNodes()
+	{
+		DATAFLOW_NODE_REGISTER_CREATION_FACTORY(FDataflowTextureTerminalNode);
+	}
+}
+
+namespace DataflowTextureAssetNodes::Private
+{
+	void UpdateTexture2DFromImage(UTexture2D& InTexture, const FImage& InImage)
+	{
+		// convert to the BGRA8
+		FImage ConvertedImage(InImage);
+		ConvertedImage.ChangeFormat(ERawImageFormat::Type::BGRA8, EGammaSpace::Linear);
+
+#if WITH_EDITOR
+		InTexture.PreEditChange(nullptr);
+#endif
+
+#if WITH_EDITORONLY_DATA
+		InTexture.Source.Init(ConvertedImage);
+#endif
+		InTexture.UpdateResource();
+
+#if WITH_EDITOR
+		InTexture.PostEditChange();
+#endif
+	}
+}
+
+
+////////////////////////////////////////////////////////////////////////////////////////////
+
+FDataflowTextureTerminalNode::FDataflowTextureTerminalNode(const UE::Dataflow::FNodeParameters& InParam, FGuid InGuid)
+	: FDataflowTerminalNode(InParam, InGuid)
+{
+	RegisterInputConnection(&Image);
+	RegisterInputConnection(&TextureAsset);
+	RegisterOutputConnection(&Image, &Image);
+}
+
+void FDataflowTextureTerminalNode::Evaluate(UE::Dataflow::FContext& Context) const
+{
+	// forward values
+	SafeForwardInput(Context, &Image, &Image);
+}
+
+
+void FDataflowTextureTerminalNode::SetAssetValue(TObjectPtr<UObject> Asset, UE::Dataflow::FContext& Context) const
+{
+	UTexture2D* AssetToSet = Cast<UTexture2D>(Asset.Get());
+	if (!AssetToSet)
+	{
+		// use the input instead
+		AssetToSet = GetValue(Context, &TextureAsset);
+	}
+
+	if (AssetToSet)
+	{
+		const FDataflowImage& InImage = GetValue(Context, &Image);
+		if (InImage.GetWidth() > 0 && InImage.GetHeight() > 0)
+		{
+			DataflowTextureAssetNodes::Private::UpdateTexture2DFromImage(*AssetToSet, InImage.GetImage());
+		}
+	}
+}
diff --git a/Engine/Source/Runtime/Experimental/Dataflow/Engine/Public/Dataflow/DataflowTextureAssetNodes.h b/Engine/Source/Runtime/Experimental/Dataflow/Engine/Public/Dataflow/DataflowTextureAssetNodes.h
new file mode 100644
--- /dev/null
+++ b/Engine/Source/Runtime/Experimental/Dataflow/Engine/Public/Dataflow/DataflowTextureAssetNodes.h
@@ -0,0 +1,39 @@
+// Copyright Epic Games, Inc. All Rights Reserved.
+
+#pragma once 
+
+#include "Dataflow/DataflowTerminalNode.h"
+#include "Dataflow/DataflowImage.h"
+
+#include "DataflowTextureAssetNodes.generated.h"
+
+class UTexture2D;
+
+/*
+* terminal node to a save a dependent 2D texture
+*/
+USTRUCT(meta = (DataflowTerminal))
+struct FDataflowTextureTerminalNode : public FDataflowTerminalNode
+{
+	GENERATED_USTRUCT_BODY()
+	DATAFLOW_NODE_DEFINE_INTERNAL(FDataflowTextureTerminalNode, "TextureTerminal", "Terminal", "")
+
+private:
+	UPROPERTY(meta = (DataflowInput, DataflowOutput, DataflowPassthrough = "Image"))
+	FDataflowImage Image;
+
+	UPROPERTY(EditAnywhere, Category = Asset, meta = (DataflowInput))
+	TObjectPtr<UTexture2D> TextureAsset = nullptr;
+
+	virtual void Evaluate(UE::Dataflow::FContext& Context) const override;
+	virtual void SetAssetValue(TObjectPtr<UObject> Asset, UE::Dataflow::FContext& Context) const override;
+
+public:
+	FDataflowTextureTerminalNode(const UE::Dataflow::FNodeParameters& InParam, FGuid InGuid = FGuid::NewGuid());
+};
+
+namespace UE::Dataflow
+{
+	void RegisterTextureAssetNodes();
+}
+
