ChaosFlesh: fixed TransferVertexAttribute node, now supports vertex-vertex transfer and triangle-vertex transfer. Added suffix for geometry transform name matching. Tested on Emil and SKM Collection.

diff --git a/Engine/Plugins/Experimental/GeometryCollectionPlugin/Source/GeometryCollectionNodes/Private/Dataflow/GeometryCollectionTransferVertexAttributeNode.cpp b/Engine/Plugins/Experimental/GeometryCollectionPlugin/Source/GeometryCollectionNodes/Private/Dataflow/GeometryCollectionTransferVertexAttributeNode.cpp
--- a/Engine/Plugins/Experimental/GeometryCollectionPlugin/Source/GeometryCollectionNodes/Private/Dataflow/GeometryCollectionTransferVertexAttributeNode.cpp
+++ b/Engine/Plugins/Experimental/GeometryCollectionPlugin/Source/GeometryCollectionNodes/Private/Dataflow/GeometryCollectionTransferVertexAttributeNode.cpp
@@ -104,6 +104,7 @@ namespace UE::Private
 		float FalloffThreshold = 0.f;
 		float EdgeMultiplier = 0.f;
 		float BoundMultiplier = 0.f;
+		FString TransformNameSuffix;
 	};
 
 	static float MaxEdgeLength(TArray<FVector3f>& Vert, const TManagedArray<FIntVector3>& Tri, int32 VertexOffset, int32 TriStart, int32 TriCount)
@@ -211,7 +212,7 @@ namespace UE::Private
 		return 1.0;
 	}
 
-	static TArray<FIntVector2> FindSourceToTargetGeometryMap(const FManagedArrayCollection& SourceCollection, const FManagedArrayCollection& TargetCollection)
+	static TArray<FIntVector2> FindSourceToTargetGeometryMap(const FManagedArrayCollection& SourceCollection, const FManagedArrayCollection& TargetCollection, const FString& TransformNameSuffix)
 	{
 		TArray<FIntVector2> Mapping;
 		const TManagedArray<FString>* const SourceName = SourceCollection.FindAttribute<FString>(FName("BoneName"), FTransformCollection::TransformGroup);
@@ -224,7 +225,7 @@ namespace UE::Private
 			{
 				for (int32 j = 0; j < TargetName->Num(); ++j)
 				{
-					const FString TestName = FString::Printf(TEXT("%s_Tet"), *(*SourceName)[i]);
+					const FString TestName = FString::Printf(TEXT("%s%s"), *(*SourceName)[i], *TransformNameSuffix);
 					if ((*TargetName)[j].StartsWith(TestName))
 					{
 						Mapping.Add(FIntVector2((*SourceGeometryGroup)[i], (*TargetGeometryGroup)[j]));
@@ -266,7 +267,7 @@ namespace UE::Private
 	{
 		TArray<FTransferData> TransferDataArray;
 		TransferDataArray.SetNum(Target.Vertex.Num());
-		Chaos::FReal SphereFullRadius;
+		Chaos::FReal SphereFullRadius = 0.f;
 
 		if (TransferProperties.SourceScale == EDataflowTransferVertexAttributeNodeSourceScale::Asset_Edge || TransferProperties.SourceScale == EDataflowTransferVertexAttributeNodeSourceScale::Asset_Bound)
 		{
@@ -287,21 +288,13 @@ namespace UE::Private
 			}
 			else if (TransferProperties.SourceScale == EDataflowTransferVertexAttributeNodeSourceScale::Asset_Bound)
 			{
+
 				Chaos::FVec3f CoordMaxs(-FLT_MAX);
 				Chaos::FVec3f CoordMins(FLT_MAX);
 				for (int32 i = 0; i < ComponentSpaceFullVertices.Num(); i++)
 				{
-					for (int32 j = 0; j < 3; j++)
-					{
-						if (ComponentSpaceFullVertices[i][j] > CoordMaxs[j])
-						{
-							CoordMaxs[j] = ComponentSpaceFullVertices[i][j];
-						}
-						if (ComponentSpaceFullVertices[i][j] < CoordMins[j])
-						{
-							CoordMins[j] = ComponentSpaceFullVertices[i][j];
-						}
-					}
+					CoordMaxs = CoordMaxs.ComponentwiseMax(ComponentSpaceFullVertices[i]);
+					CoordMins = CoordMins.ComponentwiseMin(ComponentSpaceFullVertices[i]);
 				}
 				const Chaos::FVec3f CoordDiff = (CoordMaxs - CoordMins) * TransferProperties.BoundMultiplier;
 				SphereFullRadius = FMath::Min3(CoordDiff[0], CoordDiff[1], CoordDiff[2]);
@@ -360,7 +353,7 @@ namespace UE::Private
 							TriangleMesh.Init(SourceElements);
 							Chaos::FTriangleMesh::TSpatialHashType<Chaos::FReal> SpatialHash;
 							TriangleMesh.BuildSpatialHash(ConstComponentSpaceVertices, SpatialHash, SphereRadius);
-							for (int32 TargetIndex = 0; TargetIndex < TargetVertexCountVal; TargetIndex++)
+							for (int32 TargetIndex = 0; TargetIndex < TargetVertexCountVal; ++TargetIndex)
 							{
 								TArray<Chaos::TTriangleCollisionPoint<Chaos::FReal>> Result;
 								if (TriangleMesh.PointClosestTriangleQuery(SpatialHash, ConstComponentSpaceVertices,
@@ -391,39 +384,27 @@ namespace UE::Private
 						}
 						else if (TransferProperties.BoundingVolumeType == EDataflowTransferVertexAttributeNodeBoundingVolume::Vertex)
 						{
-							TUniquePtr<BVH> VertexBVH(BuildParticleSphereBVH(ComponentSpaceTargetVertices, SphereRadius));
-							for (int32 i = 0; i < FaceCountVal; i++)
+							const float FalloffScale = 1.f;
+							TUniquePtr<BVH> VertexBVH(BuildParticleSphereBVH(ComponentSpaceVertices, SphereRadius));
+							for (int32 TargetIdx = 0; TargetIdx < ComponentSpaceTargetVertices.Num(); ++TargetIdx)
 							{
-								FIntVector3 ComponentTriangle = Source.Indices[FaceStartVal + i] - FIntVector3(VertexStartVal);
-								TArray<int32> TargetVertexIntersection({});
-								TriangleToVertexIntersections(*VertexBVH, ComponentSpaceVertices, ComponentTriangle, TargetVertexIntersection);
-
-								for (int32 j = 0; j < TargetVertexIntersection.Num(); j++)
+								float MinDistance = FLT_MAX;
+								int32 MinSourceIdx = INDEX_NONE;
+								for (int32 SourceIdx : VertexBVH->FindAllIntersections(ComponentSpaceTargetVertices[TargetIdx]))
 								{
-									Chaos::FVec3f Bary,
-										TriPos0(ComponentSpaceVertices[ComponentTriangle[0]]),
-										TriPos1(ComponentSpaceVertices[ComponentTriangle[1]]),
-										TriPos2(ComponentSpaceVertices[ComponentTriangle[2]]),
-										ParticlePos(ComponentSpaceTargetVertices[TargetVertexIntersection[j]]);
-
-									Chaos::FVec3f ClosestPoint = Chaos::FindClosestPointAndBaryOnTriangle(TriPos0, TriPos1, TriPos2, ParticlePos, Bary);
-									float CurrentDistance = (ParticlePos - ClosestPoint).Size();
-									float TriRadius = TransferProperties.FalloffThreshold * MaxEdgeLength(ComponentSpaceVertices, Source.Indices.Get(), VertexStartVal, FaceStartVal + i, 1);
-									float FalloffScale = CalculateFalloffScale(TransferProperties.Falloff, TriRadius, CurrentDistance);
-									if (!FMath::IsNearlyZero(FalloffScale))
+									float Distance = (ComponentSpaceTargetVertices[TargetIdx] - ComponentSpaceVertices[SourceIdx]).Size();
+									if (Distance < MinDistance)
 									{
-										int32 TargetIndex = TargetVertexIntersection[j] + TargetVertexStartVal;
-										TArray<int32> TransferIndices;
-										TArray<float> TransferWeights;
-										for (int32 k = 0; k < 3; k++)
-										{
-											TransferIndices.Add(ComponentTriangle[k] + VertexStartVal);
-											TransferWeights.Add(Bary[k]);
-										}
-										TransferDataArray[TargetIndex] = FTransferData(TransferIndices, TransferWeights, FalloffScale);
-										break;
+										MinDistance = Distance;
+										MinSourceIdx = SourceIdx;
 									}
 								}
+								if (MinSourceIdx != INDEX_NONE)
+								{
+									TArray<int32> TransferIndices = { MinSourceIdx + VertexStartVal };
+									TArray<float> TransferWeights = { 1.f };
+									TransferDataArray[TargetIdx + TargetVertexStartVal] = FTransferData(TransferIndices, TransferWeights, FalloffScale);
+								}
 							}
 						}
 					}
@@ -432,7 +413,7 @@ namespace UE::Private
 		return MoveTemp(TransferDataArray);
 	}
 
-	static TArray<FTransferData> NearestVertexTransfer(const FTransferFacade& Source, const FTransferFacade& Target, const FTransferProperties& TransferProperties)
+	static TArray<FTransferData> GlobalAssetTransfer(const FTransferFacade& Source, const FTransferFacade& Target, const FTransferProperties& TransferProperties)
 	{
 		TArray<FTransferData> TransferDataArray;
 		TransferDataArray.SetNum(Target.Vertex.Num());
@@ -445,41 +426,95 @@ namespace UE::Private
 		BuildComponentSpaceVertices(Source.Transform.Get(), Source.Parent.Get(), Source.BoneMap.Get(), Source.Vertex.Get(), 0, Source.Vertex.Num(), ComponentSpaceVertices);
 
 		// build Sphere based BVH
-		Chaos::FReal SphereRadius = Chaos::FReal(TransferProperties.EdgeMultiplier * FMath::Max(
-			MaxEdgeLength(ComponentSpaceTargetVertices, Target.Indices.Get(), 0, 0, Target.Indices.Num()),
-			MaxEdgeLength(ComponentSpaceVertices, Source.Indices.Get(), 0, 0, Source.Indices.Num())));
-		TUniquePtr<BVH> VertexBVH(BuildParticleSphereBVH(ComponentSpaceTargetVertices, SphereRadius));
-
-		for (int32 i = 0; i < Source.Indices.Num(); i++)
+		Chaos::FReal SphereRadius = 0.f;
+		if (TransferProperties.SourceScale == EDataflowTransferVertexAttributeNodeSourceScale::Asset_Edge || TransferProperties.SourceScale == EDataflowTransferVertexAttributeNodeSourceScale::Component_Edge)
 		{
-			FIntVector3 Triangle = Source.Indices[i];
-			TArray<int32> TargetVertexIntersection({});
-			TriangleToVertexIntersections(*VertexBVH, ComponentSpaceVertices, Triangle, TargetVertexIntersection);
-
-			for (int32 j = 0; j < TargetVertexIntersection.Num(); j++)
+			SphereRadius = Chaos::FReal(TransferProperties.EdgeMultiplier * FMath::Max(
+				MaxEdgeLength(ComponentSpaceTargetVertices, Target.Indices.Get(), 0, 0, Target.Indices.Num()),
+				MaxEdgeLength(ComponentSpaceVertices, Source.Indices.Get(), 0, 0, Source.Indices.Num())));
+		}
+		else if (TransferProperties.SourceScale == EDataflowTransferVertexAttributeNodeSourceScale::Asset_Bound)
+		{
+			Chaos::FVec3f CoordMaxs(-FLT_MAX);
+			Chaos::FVec3f CoordMins(FLT_MAX);
+			for (int32 i = 0; i < ComponentSpaceVertices.Num(); i++)
 			{
-				Chaos::FVec3f Bary,
-					TriPos0(ComponentSpaceVertices[Source.Indices[i][0]]),
-					TriPos1(ComponentSpaceVertices[Source.Indices[i][1]]),
-					TriPos2(ComponentSpaceVertices[Source.Indices[i][2]]),
-					ParticlePos(ComponentSpaceTargetVertices[TargetVertexIntersection[j]]);
-
-				Chaos::FVec3f ClosestPoint = Chaos::FindClosestPointAndBaryOnTriangle(TriPos0, TriPos1, TriPos2, ParticlePos, Bary);
-				float CurrentDistance = (ParticlePos - ClosestPoint).Size();
-				float TriRadius = TransferProperties.FalloffThreshold * MaxEdgeLength(ComponentSpaceVertices, Source.Indices.Get(), 0, i, 1);
-				float FalloffScale = CalculateFalloffScale(TransferProperties.Falloff, TriRadius, CurrentDistance);
-				if (!FMath::IsNearlyZero(FalloffScale))
+				CoordMaxs = CoordMaxs.ComponentwiseMax(ComponentSpaceVertices[i]);
+				CoordMins = CoordMins.ComponentwiseMin(ComponentSpaceVertices[i]);
+			}
+			const Chaos::FVec3f CoordDiff = (CoordMaxs - CoordMins) * TransferProperties.BoundMultiplier;
+			SphereRadius = FMath::Min3(CoordDiff[0], CoordDiff[1], CoordDiff[2]);
+		}
+		if (TransferProperties.BoundingVolumeType == EDataflowTransferVertexAttributeNodeBoundingVolume::Triangle)
+		{
+			TArray<Chaos::TVec3<Chaos::FReal>> ComponentSpaceVerticesTVec3;
+			ComponentSpaceVerticesTVec3.SetNum(ComponentSpaceVertices.Num());
+			for (int32 SourceIndex = 0; SourceIndex < ComponentSpaceVerticesTVec3.Num(); SourceIndex++)
+			{
+				ComponentSpaceVerticesTVec3[SourceIndex] = Chaos::TVec3<Chaos::FReal>(ComponentSpaceVertices[SourceIndex]);
+			}
+			TConstArrayView<Chaos::TVec3<Chaos::FReal>> ConstComponentSpaceVertices(ComponentSpaceVerticesTVec3);
+			Chaos::FTriangleMesh TriangleMesh;
+			TArray<Chaos::TVec3<int32>> SourceElements;
+			SourceElements.SetNum(Source.Indices.Num());
+			for (int32 ElementIndex = 0; ElementIndex < Source.Indices.Num(); ++ElementIndex)
+			{
+				SourceElements[ElementIndex] = Chaos::TVec3<int32>(Source.Indices[ElementIndex][0], Source.Indices[ElementIndex][1], Source.Indices[ElementIndex][2]);
+			}
+			TriangleMesh.Init(SourceElements);
+			Chaos::FTriangleMesh::TSpatialHashType<Chaos::FReal> SpatialHash;
+			TriangleMesh.BuildSpatialHash(ConstComponentSpaceVertices, SpatialHash, SphereRadius);
+			for (int32 TargetIndex = 0; TargetIndex < ConstComponentSpaceVertices.Num(); ++TargetIndex)
+			{
+				TArray<Chaos::TTriangleCollisionPoint<Chaos::FReal>> Result;
+				if (TriangleMesh.PointClosestTriangleQuery(SpatialHash, ConstComponentSpaceVertices,
+					TargetIndex, Chaos::TVec3<Chaos::FReal>(ComponentSpaceTargetVertices[TargetIndex]), SphereRadius / 2.f, SphereRadius / 2.f,
+					[](const int32 PointIndex, const int32 TriangleIndex)->bool {return true; }, Result))
+				{
+					for (const Chaos::TTriangleCollisionPoint<Chaos::FReal>& CollisionPoint : Result)
+					{
+						const float CurrentDistance = abs(CollisionPoint.Phi);
+						const float TriRadius = TransferProperties.FalloffThreshold * MaxEdgeLength(ComponentSpaceVertices, Source.Indices.Get(), 0, CollisionPoint.Indices[1], 1);
+						const float FalloffScale = CalculateFalloffScale(TransferProperties.Falloff, TriRadius, CurrentDistance);
+						if (!FMath::IsNearlyZero(FalloffScale))
+						{
+							const int32 TargetCandidateIndex = CollisionPoint.Indices[0];
+							TArray<int32> TransferIndices;
+							TArray<float> TransferWeights;
+							for (int32 k = 0; k < 3; k++)
+							{
+								TransferIndices.Add(Source.Indices[CollisionPoint.Indices[1]][k]);
+								TransferWeights.Add(CollisionPoint.Bary[k + 1]);
+							}
+							TransferDataArray[TargetCandidateIndex] = FTransferData(TransferIndices, TransferWeights, FalloffScale);
+							break;
+						}
+					}
+				}
+			}
+		}
+		else if (TransferProperties.BoundingVolumeType == EDataflowTransferVertexAttributeNodeBoundingVolume::Vertex)
+		{
+			const float FalloffScale = 1.f;
+			TUniquePtr<BVH> VertexBVH(BuildParticleSphereBVH(ComponentSpaceVertices, SphereRadius));
+			for (int32 TargetIdx = 0; TargetIdx < ComponentSpaceTargetVertices.Num(); ++TargetIdx)
+			{
+				float MinDistance = FLT_MAX;
+				int32 MinSourceIdx = INDEX_NONE;
+				for (int32 SourceIdx : VertexBVH->FindAllIntersections(ComponentSpaceTargetVertices[TargetIdx]))
 				{
-					int32 TargetIndex = TargetVertexIntersection[j];
-					TArray<int32> TransferIndices;
-					TArray<float> TransferWeights;
-					for (int32 k = 0; k < 3; k++)
+					float Distance = (ComponentSpaceTargetVertices[TargetIdx] - ComponentSpaceVertices[SourceIdx]).Size();
+					if (Distance < MinDistance)
 					{
-						TransferIndices.Add(Source.Indices[i][k]);
-						TransferWeights.Add(Bary[k]);
+						MinDistance = Distance;
+						MinSourceIdx = SourceIdx;
 					}
-					TransferDataArray[TargetIndex] = FTransferData(TransferIndices, TransferWeights, FalloffScale);
-					break;
+				}
+				if (MinSourceIdx != INDEX_NONE)
+				{
+					TArray<int32> TransferIndices = { MinSourceIdx };
+					TArray<float> TransferWeights = { 1.f };
+					TransferDataArray[TargetIdx] = FTransferData(TransferIndices, TransferWeights, FalloffScale);
 				}
 			}
 		}
@@ -493,14 +528,14 @@ namespace UE::Private
 		const FTransferFacade Source(SourceCollection);
 		if (Target.IsValid() && Source.IsValid())
 		{
-			TArray<FIntVector2> AlignedGeometry = FindSourceToTargetGeometryMap(SourceCollection, TargetCollection);
+			TArray<FIntVector2> AlignedGeometry = FindSourceToTargetGeometryMap(SourceCollection, TargetCollection, TransferProperties.TransformNameSuffix);
 			if (AlignedGeometry.Num() == SourceCollection.NumElements(FGeometryCollection::GeometryGroup))
 			{
 				TransferDataArray = PairedGeometryTransfer(AlignedGeometry, Source, Target, TransferProperties);
 			}
 			else
 			{
-				TransferDataArray = NearestVertexTransfer(Source, Target, TransferProperties);
+				TransferDataArray = GlobalAssetTransfer(Source, Target, TransferProperties);
 			}
 		}
 		return MoveTemp(TransferDataArray);
@@ -524,7 +559,8 @@ void FGeometryCollectionTransferVertexAttributeNode::Evaluate(UE::Dataflow::FCon
 			Falloff,
 			FalloffThreshold,
 			EdgeMultiplier,
-			BoundMultiplier
+			BoundMultiplier,
+			TransformNameSuffix
 		};
 
 		if (Target.IsValid() && Source.IsValid())
@@ -564,6 +600,10 @@ void FGeometryCollectionTransferVertexAttributeNode::Evaluate(UE::Dataflow::FCon
 					}
 				}
 			}
+			else
+			{
+				UE_LOG(LogChaosDataflow, Warning, TEXT("TransferVertexAttribute: attribute[%s] of group[%s] has type that is not float or FLinearColor, transfer failed."), *Key.Attribute, *Key.Group);
+			}
 		}
 
 		SetValue(Context, MoveTemp(TargetCollection), &Collection);
@@ -589,7 +629,8 @@ void FGeometryCollectionTransferVertexSkinWeightsNode::Evaluate(UE::Dataflow::FC
 			Falloff,
 			FalloffThreshold,
 			EdgeMultiplier,
-			BoundMultiplier
+			BoundMultiplier,
+			TransformNameSuffix
 		};
 
 		if (Target.IsValid() && Source.IsValid())
diff --git a/Engine/Plugins/Experimental/GeometryCollectionPlugin/Source/GeometryCollectionNodes/Public/Dataflow/GeometryCollectionTransferVertexAttributeNode.h b/Engine/Plugins/Experimental/GeometryCollectionPlugin/Source/GeometryCollectionNodes/Public/Dataflow/GeometryCollectionTransferVertexAttributeNode.h
--- a/Engine/Plugins/Experimental/GeometryCollectionPlugin/Source/GeometryCollectionNodes/Public/Dataflow/GeometryCollectionTransferVertexAttributeNode.h
+++ b/Engine/Plugins/Experimental/GeometryCollectionPlugin/Source/GeometryCollectionNodes/Public/Dataflow/GeometryCollectionTransferVertexAttributeNode.h
@@ -36,7 +36,8 @@ enum class EDataflowTransferVertexAttributeNodeFalloff : uint8
 UENUM(BlueprintType)
 enum class EDataflowTransferVertexAttributeNodeSourceScale : uint8
 {
-	/** Bounding volume hierarchy cell size based on max edge length of each geometry group*/
+	/** Bounding volume hierarchy cell size based on max edge length of each geometry group. 
+	Only works if there is 1-1 correspondence between source and target geometries, otherwise default to Asset Max Edge*/
 	Component_Edge  UMETA(DisplayName = "Component Max Edge"),
 
 	/** Bounding volume hierarchy cell size based on max edge length of the whole asset*/
@@ -53,7 +54,7 @@ enum class EDataflowTransferVertexAttributeNodeSourceScale : uint8
 UENUM(BlueprintType)
 enum class EDataflowTransferVertexAttributeNodeBoundingVolume : uint8
 {
-	/** Bounding volume on vertices of the source triangle mesh*/
+	/** Bounding volume on vertices of the source mesh*/
 	Vertex UMETA(DisplayName = "Vertex"),
 
 	/** Bounding volume on triangles of the source triangle mesh*/
@@ -97,11 +98,11 @@ public:
 	EDataflowTransferVertexAttributeNodeSourceScale SourceScale = EDataflowTransferVertexAttributeNodeSourceScale::Asset_Bound;
 
 	/* Falloff of source value based on distance from source triangle[default: Squared] */
-	UPROPERTY(EditAnywhere, Category = "Thresholds")
+	UPROPERTY(EditAnywhere, Category = "Thresholds", meta = (EditCondition = "BoundingVolumeType == EDataflowTransferVertexAttributeNodeBoundingVolume::Triangle", EditConditionHides))
 	EDataflowTransferVertexAttributeNodeFalloff Falloff = EDataflowTransferVertexAttributeNodeFalloff::None;
 
 	/* Threshold based on distance from source triangle.Values past the threshold will falloff.[Defaults to 1 percent of triangle size(0.01)] */
-	UPROPERTY(EditAnywhere, Category = "Thresholds")
+	UPROPERTY(EditAnywhere, Category = "Thresholds", meta = (EditCondition = "BoundingVolumeType == EDataflowTransferVertexAttributeNodeBoundingVolume::Triangle", EditConditionHides))
 	float FalloffThreshold = 0.01f;
 
 	/* Edge multiplier for the Bounding Volume Hierarchy(BVH) target's particle search radius. */
@@ -114,6 +115,10 @@ public:
 		EditConditionHides))
 	float BoundMultiplier = 0.01f;
 
+	/* Suffix of transform names for geometry matching during transfer[default: _Tet]. In CreateTetrahedron node we add _Tet to tetrahedral geometries.*/
+	UPROPERTY(EditAnywhere, Category = "Transform Name")
+	FString TransformNameSuffix = "_Tet";
+
 	FGeometryCollectionTransferVertexAttributeNode(const UE::Dataflow::FNodeParameters& InParam, FGuid InGuid = FGuid::NewGuid())
 		: FDataflowNode(InParam, InGuid)
 	{
@@ -159,11 +164,11 @@ public:
 	EDataflowTransferVertexAttributeNodeSourceScale SourceScale = EDataflowTransferVertexAttributeNodeSourceScale::Asset_Bound;
 
 	/* Falloff of source value based on distance from source triangle[default: Squared] */
-	UPROPERTY(EditAnywhere, Category = "Thresholds")
+	UPROPERTY(EditAnywhere, Category = "Thresholds", meta = (EditCondition = "BoundingVolumeType == EDataflowTransferVertexAttributeNodeBoundingVolume::Triangle", EditConditionHides))
 	EDataflowTransferVertexAttributeNodeFalloff Falloff = EDataflowTransferVertexAttributeNodeFalloff::None;
 
 	/* Threshold based on distance from source triangle.Values past the threshold will falloff.[Defaults to 1 percent of triangle size(0.01)] */
-	UPROPERTY(EditAnywhere, Category = "Thresholds")
+	UPROPERTY(EditAnywhere, Category = "Thresholds", meta = (EditCondition = "BoundingVolumeType == EDataflowTransferVertexAttributeNodeBoundingVolume::Triangle", EditConditionHides))
 	float FalloffThreshold = 0.01f;
 
 	/* Edge multiplier for the Bounding Volume Hierarchy(BVH) target's particle search radius. */
@@ -176,6 +181,10 @@ public:
 		EditConditionHides))
 	float BoundMultiplier = 0.01f;
 
+	/* Suffix of transform names for geometry matching during transfer[default: _Tet]. In CreateTetrahedron node we add _Tet to tetrahedral geometries.*/
+	UPROPERTY(EditAnywhere, Category = "Transform Name")
+	FString TransformNameSuffix = "_Tet";
+
 	FGeometryCollectionTransferVertexSkinWeightsNode(const UE::Dataflow::FNodeParameters& InParam, FGuid InGuid = FGuid::NewGuid())
 		: FDataflowNode(InParam, InGuid)
 	{
