Add material interface dataflow nodes ( goal is to replace the more restrictive UMaterial based Nodes that don't support all type of material )
- MakeMaterialArray
- GetMaterialAsset
- GetMaterialArraySize
- GetFromMaterialArray
- SetIntoMaterialArray
- AddToMaterialArray
- AssignMaterialToCollection

diff --git a/Engine/Plugins/Experimental/GeometryCollectionPlugin/Source/GeometryCollectionNodes/Private/Dataflow/GeometryCollectionMaterialInterfaceNodes.cpp b/Engine/Plugins/Experimental/GeometryCollectionPlugin/Source/GeometryCollectionNodes/Private/Dataflow/GeometryCollectionMaterialInterfaceNodes.cpp
new file mode 100644
--- /dev/null
+++ b/Engine/Plugins/Experimental/GeometryCollectionPlugin/Source/GeometryCollectionNodes/Private/Dataflow/GeometryCollectionMaterialInterfaceNodes.cpp
@@ -0,0 +1,314 @@
+// Copyright Epic Games, Inc. All Rights Reserved.
+
+#include "Dataflow/GeometryCollectionMaterialInterfaceNodes.h"
+#include "Dataflow/DataflowCore.h"
+#include "GeometryCollection/GeometryCollection.h"
+
+namespace UE::Dataflow
+{
+	void RegisterGeometryCollectionMaterialInterfaceNodes()
+	{
+		static const FLinearColor CDefaultNodeBodyTintColor = FLinearColor(0.f, 0.f, 0.f, 0.5f);
+
+		DATAFLOW_NODE_REGISTER_CREATION_FACTORY(FMakeMaterialInterfaceArrayDataflowNode);
+		DATAFLOW_NODE_REGISTER_CREATION_FACTORY(FGetMaterialInterfaceAssetDataflowNode);
+		DATAFLOW_NODE_REGISTER_CREATION_FACTORY(FGetMaterialInterfaceArraySizeDataflowNode);
+		DATAFLOW_NODE_REGISTER_CREATION_FACTORY(FGetFromMaterialInterfaceArrayDataflowNode);
+		DATAFLOW_NODE_REGISTER_CREATION_FACTORY(FSetIntoMaterialInterfaceArrayDataflowNode);
+		DATAFLOW_NODE_REGISTER_CREATION_FACTORY(FAddToMaterialInterfaceArrayDataflowNode);
+		DATAFLOW_NODE_REGISTER_CREATION_FACTORY(FAssignMaterialInterfaceToCollectionDataflowNode);
+
+		// Material
+		DATAFLOW_NODE_REGISTER_CREATION_FACTORY_NODE_COLORS_BY_CATEGORY("Materials", FLinearColor(FColor::Purple), CDefaultNodeBodyTintColor);
+	}
+}
+
+/////////////////////////////////////////////////////////////////////////////////////////////////////////////
+FMakeMaterialInterfaceArrayDataflowNode::FMakeMaterialInterfaceArrayDataflowNode(const UE::Dataflow::FNodeParameters& InParam, FGuid InGuid)
+	: FDataflowNode(InParam, InGuid)
+{
+	RegisterOutputConnection(&MaterialArray);
+}
+
+void FMakeMaterialInterfaceArrayDataflowNode::Evaluate(UE::Dataflow::FContext& Context, const FDataflowOutput* Out) const
+{
+	if (Out->IsA(&MaterialArray))
+	{
+		SetValue(Context, MaterialArray, &MaterialArray);
+	}
+}
+
+/////////////////////////////////////////////////////////////////////////////////////////////////////////////
+FGetMaterialInterfaceAssetDataflowNode::FGetMaterialInterfaceAssetDataflowNode(const UE::Dataflow::FNodeParameters& InParam, FGuid InGuid)
+	: FDataflowNode(InParam, InGuid)
+{
+	RegisterOutputConnection(&Material);
+}
+
+void FGetMaterialInterfaceAssetDataflowNode::Evaluate(UE::Dataflow::FContext& Context, const FDataflowOutput* Out) const
+{
+	if (Out->IsA(&Material))
+	{
+		SetValue(Context, Material, &Material);
+	}
+}
+
+/////////////////////////////////////////////////////////////////////////////////////////////////////////////
+FGetMaterialInterfaceArraySizeDataflowNode::FGetMaterialInterfaceArraySizeDataflowNode(const UE::Dataflow::FNodeParameters& InParam, FGuid InGuid)
+	: FDataflowNode(InParam, InGuid)
+{
+	RegisterInputConnection(&MaterialArray);
+	RegisterOutputConnection(&MaterialArray, &MaterialArray);
+	RegisterOutputConnection(&Size);
+}
+
+void FGetMaterialInterfaceArraySizeDataflowNode::Evaluate(UE::Dataflow::FContext& Context, const FDataflowOutput* Out) const
+{
+	if (Out->IsA(&MaterialArray))
+	{
+		SafeForwardInput(Context, &MaterialArray, &MaterialArray);
+	}
+	else if (Out->IsA(&Size))
+	{
+		const TArray<TObjectPtr<UMaterialInterface>>& InMaterialArray = GetValue(Context, &MaterialArray);
+		SetValue(Context, InMaterialArray.Num(), &Size);
+	}
+}
+
+/////////////////////////////////////////////////////////////////////////////////////////////////////////////
+
+FGetFromMaterialInterfaceArrayDataflowNode::FGetFromMaterialInterfaceArrayDataflowNode(const UE::Dataflow::FNodeParameters& InParam, FGuid InGuid)
+	: FDataflowNode(InParam, InGuid)
+{
+	RegisterInputConnection(&MaterialArray);
+	RegisterInputConnection(&Index);
+	RegisterOutputConnection(&MaterialArray, &MaterialArray);
+	RegisterOutputConnection(&Material);
+}
+
+void FGetFromMaterialInterfaceArrayDataflowNode::Evaluate(UE::Dataflow::FContext& Context, const FDataflowOutput* Out) const
+{
+	if (Out->IsA(&MaterialArray))
+	{
+		SafeForwardInput(Context, &MaterialArray, &MaterialArray);
+	}
+	else if (Out->IsA(&Material))
+	{
+		const TArray<TObjectPtr<UMaterialInterface>>& InMaterialArray = GetValue(Context, &MaterialArray);
+		const int32 InIndex = GetValue(Context, &Index);
+
+		TObjectPtr<UMaterialInterface> OutMaterial = nullptr;
+		if (InMaterialArray.IsValidIndex(InIndex))
+		{
+			OutMaterial = InMaterialArray[InIndex];
+		}
+		SetValue(Context, OutMaterial, &Material);
+	}
+}
+
+/////////////////////////////////////////////////////////////////////////////////////////////////////////////
+
+FSetIntoMaterialInterfaceArrayDataflowNode::FSetIntoMaterialInterfaceArrayDataflowNode(const UE::Dataflow::FNodeParameters& InParam, FGuid InGuid)
+	: FDataflowNode(InParam, InGuid)
+{
+	RegisterInputConnection(&MaterialArray);
+	RegisterInputConnection(&Material);
+	RegisterInputConnection(&Index);
+	RegisterOutputConnection(&MaterialArray, &MaterialArray);
+}
+
+void FSetIntoMaterialInterfaceArrayDataflowNode::Evaluate(UE::Dataflow::FContext& Context, const FDataflowOutput* Out) const
+{
+	if (Out->IsA(&MaterialArray))
+	{
+		TArray<TObjectPtr<UMaterialInterface>> InMaterialArray = GetValue(Context, &MaterialArray);
+		TObjectPtr<UMaterialInterface> InMaterial = GetValue(Context, &Material);
+		const int32 InIndex = GetValue(Context, &Index);
+
+		if (InMaterialArray.IsValidIndex(InIndex))
+		{
+			InMaterialArray[InIndex] = InMaterial;
+		}
+		SetValue(Context, MoveTemp(InMaterialArray), &MaterialArray);
+	}
+}
+
+/////////////////////////////////////////////////////////////////////////////////////////////////////////////
+
+FAddToMaterialInterfaceArrayDataflowNode::FAddToMaterialInterfaceArrayDataflowNode(const UE::Dataflow::FNodeParameters& InParam, FGuid InGuid)
+	: FDataflowNode(InParam, InGuid)
+{
+	RegisterInputConnection(&MaterialArray);
+	RegisterOutputConnection(&MaterialArray, &MaterialArray);
+	
+	// Add initial variable inputs
+	for (int32 Index = 0; Index < NumInitialVariableInputs; ++Index)
+	{
+		AddPins();
+	}
+}
+
+void FAddToMaterialInterfaceArrayDataflowNode::Evaluate(UE::Dataflow::FContext& Context, const FDataflowOutput* Out) const
+{
+	if (Out->IsA(&MaterialArray))
+	{
+		TArray<TObjectPtr<UMaterialInterface>> InMaterialArray = GetValue(Context, &MaterialArray);
+		InMaterialArray.Append(MaterialsToAdd);
+		SetValue(Context, MoveTemp(InMaterialArray), &MaterialArray);
+	}
+}
+
+bool FAddToMaterialInterfaceArrayDataflowNode::CanAddPin() const
+{
+	return true;
+}
+
+bool FAddToMaterialInterfaceArrayDataflowNode::CanRemovePin() const
+{
+	return MaterialsToAdd.Num() > 0;
+}
+
+UE::Dataflow::TConnectionReference<TObjectPtr<UMaterialInterface>> FAddToMaterialInterfaceArrayDataflowNode::GetConnectionReference(int32 Index) const
+{
+	return { &MaterialsToAdd[Index], Index, &MaterialsToAdd };
+}
+
+TArray<UE::Dataflow::FPin> FAddToMaterialInterfaceArrayDataflowNode::AddPins()
+{
+	const int32 Index = MaterialsToAdd.AddDefaulted();
+	const FDataflowInput& Input = RegisterInputArrayConnection(GetConnectionReference(Index));
+	return { { UE::Dataflow::FPin::EDirection::INPUT, Input.GetType(), Input.GetName() } };
+}
+
+TArray<UE::Dataflow::FPin> FAddToMaterialInterfaceArrayDataflowNode::GetPinsToRemove() const
+{
+	const int32 Index = (MaterialsToAdd.Num() - 1);
+	check(MaterialsToAdd.IsValidIndex(Index));
+	if (const FDataflowInput* const Input = FindInput(GetConnectionReference(Index)))
+	{
+		return { { UE::Dataflow::FPin::EDirection::INPUT, Input->GetType(), Input->GetName() } };
+	}
+	return Super::GetPinsToRemove();
+}
+
+void FAddToMaterialInterfaceArrayDataflowNode::OnPinRemoved(const UE::Dataflow::FPin& Pin)
+{
+	const int32 Index = MaterialsToAdd.Num() - 1;
+	check(MaterialsToAdd.IsValidIndex(Index));
+#if DO_CHECK
+	const FDataflowInput* const Input = FindInput(GetConnectionReference(Index));
+	check(Input);
+	check(Input->GetName() == Pin.Name);
+	check(Input->GetType() == Pin.Type);
+#endif
+	MaterialsToAdd.SetNum(Index);
+
+	return Super::OnPinRemoved(Pin);
+}
+
+void FAddToMaterialInterfaceArrayDataflowNode::PostSerialize(const FArchive& Ar)
+{
+	if (Ar.IsLoading())
+	{
+		check(MaterialsToAdd.Num() >= 0);
+		// register new elements from the array as inputs
+		for (int32 Index = 0; Index < MaterialsToAdd.Num(); ++Index)
+		{
+			FindOrRegisterInputArrayConnection(GetConnectionReference(Index));
+		}
+		if (Ar.IsTransacting())
+		{
+			// if we have more inputs than materials then we need to unregister the inputs 
+			const int32 NumVariableInputs = (GetNumInputs() - NumOtherInputs);
+			const int32 NumMaterials = MaterialsToAdd.Num();
+			if (NumVariableInputs > NumMaterials)
+			{
+				// Inputs have been removed.
+				// Temporarily expand Collections so we can get connection references.
+				MaterialsToAdd.SetNum(NumVariableInputs);
+				for (int32 Index = NumMaterials; Index < MaterialsToAdd.Num(); ++Index)
+				{
+					UnregisterInputConnection(GetConnectionReference(Index));
+				}
+				MaterialsToAdd.SetNum(NumMaterials);
+			}
+		}
+		else
+		{
+			ensureAlways(MaterialsToAdd.Num() + NumOtherInputs == GetNumInputs());
+		}
+	}
+}
+
+/////////////////////////////////////////////////////////////////////////////////////////////////////////////
+
+FAssignMaterialInterfaceToCollectionDataflowNode::FAssignMaterialInterfaceToCollectionDataflowNode(const UE::Dataflow::FNodeParameters& InParam, FGuid InGuid)
+	: FDataflowNode(InParam, InGuid)
+{
+	RegisterInputConnection(&Collection);
+	RegisterInputConnection(&FaceSelection);
+	RegisterInputConnection(&MaterialArray);
+	RegisterInputConnection(&Material);
+	RegisterOutputConnection(&Collection, &Collection);
+	RegisterOutputConnection(&MaterialArray, &MaterialArray);
+	RegisterOutputConnection(&MaterialIndex);
+}
+
+int32 FAssignMaterialInterfaceToCollectionDataflowNode::AddOrMergeMaterialToArray(TArray<TObjectPtr<UMaterialInterface>>& InOutMaterials, TObjectPtr<UMaterialInterface> InMaterialToAdd) const
+{
+	if (bMergeDuplicateMaterials)
+	{
+		return InOutMaterials.AddUnique(InMaterialToAdd);
+	}
+	return InOutMaterials.Add(InMaterialToAdd);
+}
+
+void FAssignMaterialInterfaceToCollectionDataflowNode::Evaluate(UE::Dataflow::FContext& Context, const FDataflowOutput* Out) const
+{
+	if (Out->IsA(&Collection))
+	{
+		FManagedArrayCollection InCollection = GetValue<FManagedArrayCollection>(Context, &Collection);
+		FDataflowFaceSelection InFaceSelection = GetValue<FDataflowFaceSelection>(Context, &FaceSelection);
+		TArray<TObjectPtr<UMaterialInterface>> InMaterialArray = GetValue(Context, &MaterialArray);
+		TObjectPtr<UMaterialInterface> InMaterial = GetValue(Context, &Material);
+
+		// add the material to the array 
+		const int32 InMaterialIndex = AddOrMergeMaterialToArray(InMaterialArray, InMaterial);
+
+		const int32 NumFaces = Collection.NumElements(FGeometryCollection::FacesGroup);
+		if (!IsConnected(&FaceSelection))
+		{
+			// if not connected select all faces
+			InFaceSelection.Initialize(NumFaces, true);
+		}
+
+		static const FName MaterialIdAttributeName = "MaterialID";
+		if (InFaceSelection.Num() == NumFaces && 
+			InCollection.HasAttribute(MaterialIdAttributeName, FGeometryCollection::FacesGroup))
+		{
+			TManagedArray<int32>& MaterialIDs = InCollection.ModifyAttribute<int32>(MaterialIdAttributeName, FGeometryCollection::FacesGroup);
+
+			// Update MaterialIdx for selected outside faces
+			for (int32 FaceIdx = 0; FaceIdx < InFaceSelection.Num(); ++FaceIdx)
+			{
+				if (InFaceSelection.IsSelected(FaceIdx))
+				{
+					MaterialIDs[FaceIdx] = InMaterialIndex;
+				}
+			}
+		}
+
+		SetValue(Context, MoveTemp(InCollection), &Collection);
+
+	}
+	else if (Out->IsA(&MaterialArray) || Out->IsA(&MaterialIndex))
+	{
+		TArray<TObjectPtr<UMaterialInterface>> InMaterialArray = GetValue(Context, &MaterialArray);
+		TObjectPtr<UMaterialInterface> InMaterial = GetValue(Context, &Material);
+
+		const int32 InMaterialIndex = AddOrMergeMaterialToArray(InMaterialArray, InMaterial);
+
+		SetValue(Context, InMaterialArray, &MaterialArray);
+		SetValue(Context, InMaterialIndex, &MaterialIndex);
+	}
+}
diff --git a/Engine/Plugins/Experimental/GeometryCollectionPlugin/Source/GeometryCollectionNodes/Private/Dataflow/GeometryCollectionNodesPlugin.cpp b/Engine/Plugins/Experimental/GeometryCollectionPlugin/Source/GeometryCollectionNodes/Private/Dataflow/GeometryCollectionNodesPlugin.cpp
--- a/Engine/Plugins/Experimental/GeometryCollectionPlugin/Source/GeometryCollectionNodes/Private/Dataflow/GeometryCollectionNodesPlugin.cpp
+++ b/Engine/Plugins/Experimental/GeometryCollectionPlugin/Source/GeometryCollectionNodes/Private/Dataflow/GeometryCollectionNodesPlugin.cpp
@@ -16,6 +16,7 @@
 #include "Dataflow/GeometryCollectionFracturingNodes.h"
 #include "Dataflow/GeometryCollectionMakeNodes.h"
 #include "Dataflow/GeometryCollectionMaterialNodes.h"
+#include "Dataflow/GeometryCollectionMaterialInterfaceNodes.h"
 #include "Dataflow/GeometryCollectionMathNodes.h"
 #include "Dataflow/GeometryCollectionMeshNodes.h"
 #include "Dataflow/GeometryCollectionNodes.h"
@@ -37,6 +38,7 @@
 #include "Dataflow/SetVertexColorFromVertexIndicesNode.h"
 #include "Dataflow/SetVertexColorFromVertexSelectionNode.h"
 
+
 #define LOCTEXT_NAMESPACE "DataflowNodes"
 
 
@@ -63,6 +65,7 @@ void IGeometryCollectionNodesPlugin::StartupModule()
 	UE::Dataflow::GeometryCollectionSamplingNodes();
 	UE::Dataflow::RegisterGeometryCollectionUVNodes();
 	UE::Dataflow::RegisterGeometryCollectionTextureNodes();
+	UE::Dataflow::RegisterGeometryCollectionMaterialInterfaceNodes();
 	DATAFLOW_NODE_REGISTER_CREATION_FACTORY(FCreateColorArrayFromFloatArrayDataflowNode);
 	DATAFLOW_NODE_REGISTER_CREATION_FACTORY(FGeometryCollectionVertexScalarToVertexIndicesNode);
 	DATAFLOW_NODE_REGISTER_CREATION_FACTORY(FSetVertexColorFromFloatArrayDataflowNode);
diff --git a/Engine/Plugins/Experimental/GeometryCollectionPlugin/Source/GeometryCollectionNodes/Public/Dataflow/GeometryCollectionMaterialInterfaceNodes.h b/Engine/Plugins/Experimental/GeometryCollectionPlugin/Source/GeometryCollectionNodes/Public/Dataflow/GeometryCollectionMaterialInterfaceNodes.h
new file mode 100644
--- /dev/null
+++ b/Engine/Plugins/Experimental/GeometryCollectionPlugin/Source/GeometryCollectionNodes/Public/Dataflow/GeometryCollectionMaterialInterfaceNodes.h
@@ -0,0 +1,217 @@
+// Copyright Epic Games, Inc. All Rights Reserved.
+
+#pragma once 
+
+#include "CoreMinimal.h"
+#include "Dataflow/DataflowNode.h"
+#include "GeometryCollection/ManagedArrayCollection.h"
+
+#include "GeometryCollectionMaterialInterfaceNodes.generated.h"
+
+class UMaterialInterface;
+
+/**
+* Make a array from a user defined list of material objects
+*/
+USTRUCT()
+struct FMakeMaterialInterfaceArrayDataflowNode : public FDataflowNode
+{
+	GENERATED_USTRUCT_BODY()
+	DATAFLOW_NODE_DEFINE_INTERNAL(FMakeMaterialInterfaceArrayDataflowNode, "MakeMaterialArray", "Materials", "")
+
+private:
+	/** Material array set by the user */
+	UPROPERTY(EditAnywhere, Category=Materials, meta=(DataflowOutput))
+	TArray<TObjectPtr<UMaterialInterface>> MaterialArray;
+
+	virtual void Evaluate(UE::Dataflow::FContext& Context, const FDataflowOutput* Out) const override;
+
+public:
+	FMakeMaterialInterfaceArrayDataflowNode(const UE::Dataflow::FNodeParameters& InParam, FGuid InGuid = FGuid::NewGuid());
+};
+
+/**
+* Get number of element in an material array 
+*/
+USTRUCT()
+struct FGetMaterialInterfaceArraySizeDataflowNode : public FDataflowNode
+{
+	GENERATED_USTRUCT_BODY()
+	DATAFLOW_NODE_DEFINE_INTERNAL(FGetMaterialInterfaceArraySizeDataflowNode, "GetMaterialArraySize", "Materials", "")
+
+private:
+	/** Material array to get size from */
+	UPROPERTY(meta = (DataflowInput, DataflowOutput, DataflowPassthrough = MaterialArray, DataflowIntrinsic))
+	TArray<TObjectPtr<UMaterialInterface>> MaterialArray;
+
+	/** Size of the array */
+	UPROPERTY(meta = (DataflowOutput))
+	int32 Size = 0;
+
+	virtual void Evaluate(UE::Dataflow::FContext& Context, const FDataflowOutput* Out) const override;
+
+public:
+	FGetMaterialInterfaceArraySizeDataflowNode(const UE::Dataflow::FNodeParameters& InParam, FGuid InGuid = FGuid::NewGuid());
+};
+
+/**
+* Get a material interface from an existing asset
+*/
+USTRUCT()
+struct FGetMaterialInterfaceAssetDataflowNode : public FDataflowNode
+{
+	GENERATED_USTRUCT_BODY()
+	DATAFLOW_NODE_DEFINE_INTERNAL(FGetMaterialInterfaceAssetDataflowNode, "GetMaterialAsset", "Materials", "")
+
+private:
+	/** Material asset to get */
+	UPROPERTY(EditAnywhere, Category = Material, meta = (DataflowOutput));
+	TObjectPtr<UMaterialInterface> Material;
+
+	virtual void Evaluate(UE::Dataflow::FContext& Context, const FDataflowOutput* Out) const override;
+
+public:
+	FGetMaterialInterfaceAssetDataflowNode(const UE::Dataflow::FNodeParameters& InParam, FGuid InGuid = FGuid::NewGuid());
+};
+
+/**
+ * Get an element from a material array
+ * (if the index does not match the range of the array, null is returned)
+ */
+USTRUCT()
+struct FGetFromMaterialInterfaceArrayDataflowNode : public FDataflowNode
+{
+	GENERATED_USTRUCT_BODY()
+	DATAFLOW_NODE_DEFINE_INTERNAL(FGetFromMaterialInterfaceArrayDataflowNode, "GetFromMaterialsArray", "Materials", "")
+
+private:
+	/** Material array to get the material from  */
+	UPROPERTY(meta = (DataflowInput, DataflowOutput, DataflowPassthrough = MaterialArray, DataflowIntrinsic))
+	TArray<TObjectPtr<UMaterialInterface>> MaterialArray;
+
+	/** Material at the requested index ( may be null if index does not match the array range ) */
+	UPROPERTY(meta = (DataflowOutput))
+	TObjectPtr<UMaterialInterface> Material;
+
+	/** Index in the array to get the material from. Invalid index will return null material */
+	UPROPERTY(EditAnywhere, Category = Material, meta = (DataflowInput))
+	int32 Index = 0;
+
+	virtual void Evaluate(UE::Dataflow::FContext& Context, const FDataflowOutput* Out) const override;
+
+public:
+	FGetFromMaterialInterfaceArrayDataflowNode(const UE::Dataflow::FNodeParameters& InParam, FGuid InGuid = FGuid::NewGuid());
+};
+
+/**
+ * Set an element into a material array at a specific index
+ * (if the index does not match the range of the array, the array will remain unchanged)
+ */
+USTRUCT()
+struct FSetIntoMaterialInterfaceArrayDataflowNode : public FDataflowNode
+{
+	GENERATED_USTRUCT_BODY()
+	DATAFLOW_NODE_DEFINE_INTERNAL(FSetIntoMaterialInterfaceArrayDataflowNode, "SetIntoMaterialsArray", "Materials", "")
+
+private:
+	/** Material array to modify */
+	UPROPERTY(meta = (DataflowInput, DataflowOutput, DataflowPassthrough = MaterialArray, DataflowIntrinsic))
+	TArray<TObjectPtr<UMaterialInterface>> MaterialArray;
+
+	/** Material to set at the specific index into the array */
+	UPROPERTY(EditAnywhere, Category = Material, meta = (DataflowInput))
+	TObjectPtr<UMaterialInterface> Material;
+
+	/** Index Set the material at (if the index does not match the range of the array, the array will remain unchanged)*/
+	UPROPERTY(EditAnywhere, Category = Material, meta = (DataflowInput))
+	int32 Index = 0;
+
+	virtual void Evaluate(UE::Dataflow::FContext& Context, const FDataflowOutput* Out) const override;
+
+public:
+	FSetIntoMaterialInterfaceArrayDataflowNode(const UE::Dataflow::FNodeParameters& InParam, FGuid InGuid = FGuid::NewGuid());
+};
+
+/**
+* Add material(s) to an array
+*/
+USTRUCT()
+struct FAddToMaterialInterfaceArrayDataflowNode : public FDataflowNode
+{
+	GENERATED_USTRUCT_BODY()
+	DATAFLOW_NODE_DEFINE_INTERNAL(FAddToMaterialInterfaceArrayDataflowNode, "AddToMaterialArray", "Materials", "")
+
+private:
+	/** Material array to add to */
+	UPROPERTY(meta = (DataflowInput, DataflowOutput, DataflowPassthrough = MaterialArray))
+	TArray<TObjectPtr<UMaterialInterface>> MaterialArray;
+
+	UPROPERTY()
+	TArray<TObjectPtr<UMaterialInterface>> MaterialsToAdd;
+
+	//~ Begin FDataflowNode interface
+	virtual TArray<UE::Dataflow::FPin> AddPins() override;
+	virtual bool CanAddPin() const override;
+	virtual bool CanRemovePin() const override;
+	virtual TArray<UE::Dataflow::FPin> GetPinsToRemove() const override;
+	virtual void OnPinRemoved(const UE::Dataflow::FPin& Pin) override;
+	virtual void PostSerialize(const FArchive& Ar) override;
+	virtual void Evaluate(UE::Dataflow::FContext& Context, const FDataflowOutput* Out) const override;
+	//~ End FDataflowNode interface
+
+	UE::Dataflow::TConnectionReference<TObjectPtr<UMaterialInterface>> GetConnectionReference(int32 Index) const;
+
+	static constexpr int32 NumOtherInputs = 1; // MaterialArray
+	static constexpr int32 NumInitialVariableInputs = 1; 
+
+public:
+	FAddToMaterialInterfaceArrayDataflowNode(const UE::Dataflow::FNodeParameters& InParam, FGuid InGuid = FGuid::NewGuid());
+};
+
+/**
+ * Assign material to a set of face in a geometry collection
+ */
+USTRUCT()
+struct FAssignMaterialInterfaceToCollectionDataflowNode : public FDataflowNode
+{
+	GENERATED_USTRUCT_BODY()
+	DATAFLOW_NODE_DEFINE_INTERNAL(FAssignMaterialInterfaceToCollectionDataflowNode, "AssignMaterialToCollection", "Materials", "")
+
+private:
+	/** Collection to assign material to */
+	UPROPERTY(meta = (DataflowInput, DataflowOutput, DataflowPassthrough = Collection, DataflowIntrinsic))
+	FManagedArrayCollection Collection;
+
+	/** Faces that will be set with this material index, if no selection is connected , all faces will be set */
+	UPROPERTY(meta = (DataflowInput))
+	FDataflowFaceSelection FaceSelection;
+
+	/** Array holding the materials objects */
+	UPROPERTY(meta = (DataflowInput, DataflowOutput, DataflowPassthrough = MaterialArray, DataflowIntrinsic))
+	TArray<TObjectPtr<UMaterialInterface>> MaterialArray;
+
+	/** Material to assign to the selection */
+	UPROPERTY(EditAnywhere, Category = Material, meta = (DataflowInput))
+	TObjectPtr<UMaterialInterface> Material;
+
+	/** Index where the material was set in the array */
+	UPROPERTY(meta = (DataflowOutput))
+	int32 MaterialIndex = 0;
+
+	/** If true, detect duplicate in the material array and only add the material in the array if it does not yet exists */
+	UPROPERTY(EditAnywhere, Category = Material)
+	bool bMergeDuplicateMaterials = false;
+
+	virtual void Evaluate(UE::Dataflow::FContext& Context, const FDataflowOutput* Out) const override;
+
+	/** add or merge a material and return the index where the material was set inthe array */
+	int32 AddOrMergeMaterialToArray(TArray<TObjectPtr<UMaterialInterface>>& InOutMaterials, TObjectPtr<UMaterialInterface> InMaterialToAdd) const;
+
+public:
+	FAssignMaterialInterfaceToCollectionDataflowNode(const UE::Dataflow::FNodeParameters& InParam, FGuid InGuid = FGuid::NewGuid());
+};
+
+namespace UE::Dataflow
+{
+	void RegisterGeometryCollectionMaterialInterfaceNodes();
+}
