Change default 'generate method' for the "create leaf convex hulls" node to 'compute from geometry' rather than 'external'
(note: should not change the setting on existing nodes)

diff --git a/Engine/Plugins/Experimental/GeometryCollectionPlugin/Source/GeometryCollectionNodes/Public/Dataflow/GeometryCollectionUtilityNodes.h b/Engine/Plugins/Experimental/GeometryCollectionPlugin/Source/GeometryCollectionNodes/Public/Dataflow/GeometryCollectionUtilityNodes.h
--- a/Engine/Plugins/Experimental/GeometryCollectionPlugin/Source/GeometryCollectionNodes/Public/Dataflow/GeometryCollectionUtilityNodes.h
+++ b/Engine/Plugins/Experimental/GeometryCollectionPlugin/Source/GeometryCollectionNodes/Public/Dataflow/GeometryCollectionUtilityNodes.h
@@ -125,7 +125,7 @@ public:
 
 	/** How convex hulls are generated -- computed from geometry, imported from external collision shapes, or an intersection of both options. */
 	UPROPERTY(EditAnywhere, Category = Options)
-	EGenerateConvexMethod GenerateMethod = EGenerateConvexMethod::ExternalCollision;
+	EGenerateConvexMethod GenerateMethod = EGenerateConvexMethod::ComputedFromGeometry;
 
 	/** If GenerateMethod is Intersect, only actually intersect when the volume of the Computed Hull is less than this fraction of the volume of the External Hull(s). */
 	UPROPERTY(EditAnywhere, Category = IntersectionFilters, meta = (ClampMin = 0.0, ClampMax = 1.0, EditCondition = "GenerateMethod == EGenerateConvexMethod::IntersectExternalWithComputed"))
