[Backout] - CL38051190
Original CL Desc
-----------------------------------------------------------------
Add a geometry collection to dynamic mesh conversion struct to the MeshConversion module, to support dynamic mesh-based updates of geometry collections in more places. This roughly follows the approach of the PlanarCuts module (which has a specialized/private version of this, for fracture only), but puts the data into the dynamic mesh using standard overlays instead of custom attributes.

As an initial use of this new conversion, add a dataflow node to run a geometry script processor on selected bones of a geometry collection.

diff --git a/Engine/Plugins/Experimental/GeometryCollectionPlugin/Source/GeometryCollectionNodes/Private/Dataflow/GeometryCollectionMeshNodes.cpp b/Engine/Plugins/Experimental/GeometryCollectionPlugin/Source/GeometryCollectionNodes/Private/Dataflow/GeometryCollectionMeshNodes.cpp
--- a/Engine/Plugins/Experimental/GeometryCollectionPlugin/Source/GeometryCollectionNodes/Private/Dataflow/GeometryCollectionMeshNodes.cpp
+++ b/Engine/Plugins/Experimental/GeometryCollectionPlugin/Source/GeometryCollectionNodes/Private/Dataflow/GeometryCollectionMeshNodes.cpp
@@ -10,7 +10,6 @@
 #include "GeometryCollection/GeometryCollection.h"
 #include "GeometryCollection/GeometryCollectionEngineUtility.h"
 #include "GeometryCollection/GeometryCollectionEngineConversion.h"
-#include "GeometryCollection/Facades/CollectionTransformSelectionFacade.h"
 #include "Logging/LogMacros.h"
 #include "Templates/SharedPointer.h"
 #include "UObject/UnrealTypePrivate.h"
@@ -20,8 +19,6 @@
 #include "DynamicMeshEditor.h"
 #include "Operations/MeshBoolean.h"
 
-#include "GeometryCollectionToDynamicMesh.h"
-
 #include "EngineGlobals.h"
 #include "GeometryCollection/GeometryCollectionAlgo.h"
 #include "GeometryCollection/GeometryCollectionClusteringUtility.h"
@@ -55,7 +52,6 @@ namespace UE::Dataflow
 		DATAFLOW_NODE_REGISTER_CREATION_FACTORY(FMeshCopyToPointsDataflowNode);
 		DATAFLOW_NODE_REGISTER_CREATION_FACTORY(FGetMeshDataDataflowNode);
 		DATAFLOW_NODE_REGISTER_CREATION_FACTORY(FApplyMeshProcessorToMeshDataflowNode);
-		DATAFLOW_NODE_REGISTER_CREATION_FACTORY(FApplyMeshProcessorToGeometryCollectionDataflowNode);
 
 
 		// Mesh
@@ -498,90 +494,3 @@ void FApplyMeshProcessorToMeshDataflowNode::Evaluate(UE::Dataflow::FContext& Con
 		}
 	}
 }
-
-
-void FApplyMeshProcessorToGeometryCollectionDataflowNode::Evaluate(UE::Dataflow::FContext& Context, const FDataflowOutput* Out) const
-{
-	if (Out->IsA(&Collection) ||
-		Out->IsA(&TransformSelection))
-	{
-		if (!MeshProcessorInstance)
-		{
-			SafeForwardInput(Context, &Collection, &Collection);
-			return;
-		}
-
-		FDataflowTransformSelection InTransformSelection = GetValue(Context, &TransformSelection);
-
-		//
-		// If not connected select everything by default
-		//
-		if (!IsConnected(&TransformSelection))
-		{
-			const FManagedArrayCollection& InCollection = GetValue(Context, &Collection);
-
-			GeometryCollection::Facades::FCollectionTransformSelectionFacade TransformSelectionFacade(InCollection);
-			const TArray<int32>& SelectionArr = TransformSelectionFacade.SelectAll();
-
-			FDataflowTransformSelection NewTransformSelection;
-			NewTransformSelection.Initialize(InCollection.NumElements(FGeometryCollection::TransformGroup), false);
-			NewTransformSelection.SetFromArray(SelectionArr);
-
-			InTransformSelection = NewTransformSelection;
-		}
-
-		if (InTransformSelection.AnySelected())
-		{
-			const FManagedArrayCollection& InCollection = GetValue(Context, &Collection);
-
-			UE::Geometry::FGeometryCollectionToDynamicMeshes CollectionToMeshes;
-			UE::Geometry::FGeometryCollectionToDynamicMeshes::FToMeshOptions ToMeshOptions;
-			ToMeshOptions.bWeldVertices = bWeldVertices;
-			ToMeshOptions.bSaveIsolatedVertices = bPreserveIsolatedVertices;
-			if (CollectionToMeshes.InitFromTransformSelection(InCollection, InTransformSelection.AsArray(), ToMeshOptions)
-				&& !CollectionToMeshes.Meshes.IsEmpty())
-			{
-				// Temporarily create a UDynamicMesh as a container to hold the meshes we pass to BP
-				TObjectPtr<UDynamicMesh> NewMesh = NewObject<UDynamicMesh>();
-
-				bool bAnySuccess = false;
-				for (UE::Geometry::FGeometryCollectionToDynamicMeshes::FMeshInfo& MeshInfo : CollectionToMeshes.Meshes)
-				{
-					NewMesh->SetMesh(MoveTemp(*MeshInfo.Mesh));
-
-					bool bFailed = false;
-					MeshProcessorInstance->ProcessDynamicMesh(NewMesh, bFailed);
-					if (!bFailed) // on success, move the mesh back
-					{
-						bAnySuccess = true;
-						NewMesh->EditMesh([&MeshInfo](UE::Geometry::FDynamicMesh3& Mesh)
-						{
-							*MeshInfo.Mesh = MoveTemp(Mesh);
-						}, EDynamicMeshChangeType::GeneralEdit, EDynamicMeshAttributeChangeFlags::Unknown, true);
-					}
-					else // on failure, clear the mesh so it won't be written back
-					{
-						MeshInfo.TransformIndex = -1;
-						MeshInfo.Mesh = nullptr;
-					}
-				}
-
-				if (bAnySuccess)
-				{
-					if (TUniquePtr<FGeometryCollection> GeomCollection = TUniquePtr<FGeometryCollection>(InCollection.NewCopy<FGeometryCollection>()))
-					{
-						UE::Geometry::FGeometryCollectionToDynamicMeshes::FToCollectionOptions ToCollectionOptions;
-						ToCollectionOptions.bDefaultFaceInternal = false;
-						ToCollectionOptions.bDefaultFaceVisible = true;
-						CollectionToMeshes.UpdateGeometryCollection(*GeomCollection, ToCollectionOptions);
-						SetValue<const FManagedArrayCollection&>(Context, *GeomCollection, &Collection);
-						return;
-					}
-				}
-			}
-		}
-
-		SafeForwardInput(Context, &Collection, &Collection);
-	}
-}
-
diff --git a/Engine/Plugins/Experimental/GeometryCollectionPlugin/Source/GeometryCollectionNodes/Public/Dataflow/GeometryCollectionMeshNodes.h b/Engine/Plugins/Experimental/GeometryCollectionPlugin/Source/GeometryCollectionNodes/Public/Dataflow/GeometryCollectionMeshNodes.h
--- a/Engine/Plugins/Experimental/GeometryCollectionPlugin/Source/GeometryCollectionNodes/Public/Dataflow/GeometryCollectionMeshNodes.h
+++ b/Engine/Plugins/Experimental/GeometryCollectionPlugin/Source/GeometryCollectionNodes/Public/Dataflow/GeometryCollectionMeshNodes.h
@@ -134,46 +134,6 @@ public:
 	virtual void Evaluate(UE::Dataflow::FContext& Context, const FDataflowOutput* Out) const override;
 };
 
-/**
- * Apply a Geometry Script Mesh Processors to the geometry of selected transforms in a geometry collection
- */
-USTRUCT(meta = (DataflowGeometryCollection))
-struct FApplyMeshProcessorToGeometryCollectionDataflowNode final : public FMeshProcessorDataflowNodeBase
-{
-	GENERATED_USTRUCT_BODY()
-	DATAFLOW_NODE_DEFINE_INTERNAL(FApplyMeshProcessorToGeometryCollectionDataflowNode, "ApplyGeometryScriptToCollection", "Mesh|Utilities", "")
-	DATAFLOW_NODE_RENDER_TYPE("SurfaceRender", FGeometryCollection::StaticType(), "Collection")
-
-public:
-
-	/** Input/Output mesh */
-	UPROPERTY(meta = (DataflowInput, DataflowOutput, DataflowPassthrough = "Collection", DataflowIntrinsic))
-	FManagedArrayCollection Collection;
-
-	/** Selected bones will have geometry script processing applied (if they have geometry). If not connected, all geometry will be processed. */
-	UPROPERTY(meta = (DataflowInput))
-	FDataflowTransformSelection TransformSelection;
-
-	/** Whether the processed mesh will have edges at normal/UV/color seams welded so they are treated as one edge during processing. */
-	UPROPERTY(EditAnywhere, Category = Options)
-	bool bWeldVertices = true;
-
-	/** Whether to preserve isolated vertices which aren't used by any triangles. */
-	UPROPERTY(EditAnywhere, Category = Options)
-	bool bPreserveIsolatedVertices = true;
-
-	FApplyMeshProcessorToGeometryCollectionDataflowNode(const UE::Dataflow::FNodeParameters& InParam, FGuid InGuid = FGuid::NewGuid())
-		: FMeshProcessorDataflowNodeBase(InParam, InGuid)
-	{
-		RegisterInputConnection(&Collection);
-
-		RegisterInputConnection(&TransformSelection);
-		RegisterOutputConnection(&Collection, &Collection);
-	}
-
-	virtual void Evaluate(UE::Dataflow::FContext& Context, const FDataflowOutput* Out) const override;
-};
-
 /**
  *
  * Converts a BoundingBox into a DynamicMesh
diff --git a/Engine/Source/Runtime/MeshConversion/MeshConversion.Build.cs b/Engine/Source/Runtime/MeshConversion/MeshConversion.Build.cs
--- a/Engine/Source/Runtime/MeshConversion/MeshConversion.Build.cs
+++ b/Engine/Source/Runtime/MeshConversion/MeshConversion.Build.cs
@@ -16,7 +16,6 @@ public class MeshConversion : ModuleRules
         PublicDependencyModuleNames.AddRange(
 			new string[] {
                 "Core",
-				"Chaos", // for geometry collection
                 "MeshDescription",
 				"StaticMeshDescription",
 				"SkeletalMeshDescription",
diff --git a/Engine/Source/Runtime/MeshConversion/Private/GeometryCollectionToDynamicMesh.cpp b/Engine/Source/Runtime/MeshConversion/Private/GeometryCollectionToDynamicMesh.cpp
deleted file mode 100644
diff --git a/Engine/Source/Runtime/MeshConversion/Public/GeometryCollectionToDynamicMesh.h b/Engine/Source/Runtime/MeshConversion/Public/GeometryCollectionToDynamicMesh.h
deleted file mode 100644
