Dataflow : deprecate OnInputtypeChanged and OnoutputTypechanged
- Remove type constraints for Vector nodes using them as they are not necessary really
- Convert selection general nodes to use the new type dependency system

diff --git a/Engine/Plugins/Experimental/GeometryCollectionPlugin/Source/GeometryCollectionNodes/Private/Dataflow/GeometryCollectionSelectionNodes.cpp b/Engine/Plugins/Experimental/GeometryCollectionPlugin/Source/GeometryCollectionNodes/Private/Dataflow/GeometryCollectionSelectionNodes.cpp
--- a/Engine/Plugins/Experimental/GeometryCollectionPlugin/Source/GeometryCollectionNodes/Private/Dataflow/GeometryCollectionSelectionNodes.cpp
+++ b/Engine/Plugins/Experimental/GeometryCollectionPlugin/Source/GeometryCollectionNodes/Private/Dataflow/GeometryCollectionSelectionNodes.cpp
@@ -1516,9 +1516,13 @@ void FGeometrySelectionToVertexSelectionDataflowNode::Evaluate(UE::Dataflow::FCo
 FCollectionSelectionSetOperationDataflowNode::FCollectionSelectionSetOperationDataflowNode(const UE::Dataflow::FNodeParameters& InParam, FGuid InGuid)
 	: FDataflowNode(InParam, InGuid)
 {
-	RegisterInputConnection(&SelectionA);
-	RegisterInputConnection(&SelectionB);
-	RegisterOutputConnection(&Selection, &SelectionA);
+	static const FName TypeDependencyGroup("Main");
+	RegisterInputConnection(&SelectionA)
+		.SetTypeDependencyGroup(TypeDependencyGroup);
+	RegisterInputConnection(&SelectionB)
+		.SetTypeDependencyGroup(TypeDependencyGroup);
+	RegisterOutputConnection(&Selection, &SelectionA)
+		.SetTypeDependencyGroup(TypeDependencyGroup);
 }
 
 void FCollectionSelectionSetOperationDataflowNode::Evaluate(UE::Dataflow::FContext& Context, const FDataflowOutput* Out) const
@@ -1559,34 +1563,16 @@ void FCollectionSelectionSetOperationDataflowNode::Evaluate(UE::Dataflow::FConte
 	}
 }
 
-bool FCollectionSelectionSetOperationDataflowNode::OnInputTypeChanged(const FDataflowInput* Input)
-{
-	// using single bitwise | operator to avoid skipping calls to SetInputConcreteType because of the shortcircuiting of ||
-	// need to disable the warning for static analysis ( V792 warning )
-	return bool(
-		SetInputConcreteType(&SelectionA, Input->GetType())
-		| SetInputConcreteType(&SelectionB, Input->GetType()) //-V792
-		| SetOutputConcreteType(&Selection, Input->GetType()) //-V792
-		);
-}
-
-bool FCollectionSelectionSetOperationDataflowNode::OnOutputTypeChanged(const FDataflowOutput* Input)
-{
-	// using single bitwise | operator to avoid skipping calls to SetInputConcreteType because of the shortcircuiting of ||
-	// need to disable the warning for static analysis ( V792 warning )
-	return bool(
-		SetInputConcreteType(&SelectionA, Input->GetType())
-		| SetInputConcreteType(&SelectionB, Input->GetType()) //-V792
-		);
-}
-
 //////////////////////////////////////////////////////////////////////////////////////////////////////////////////////
 
 FCollectionSelectionInvertDataflowNode::FCollectionSelectionInvertDataflowNode(const UE::Dataflow::FNodeParameters& InParam, FGuid InGuid)
 	: FDataflowNode(InParam, InGuid)
 {
-	RegisterInputConnection(&Selection);
-	RegisterOutputConnection(&Selection, &Selection);
+	static const FName TypeDependencyGroup("Main");
+	RegisterInputConnection(&Selection)
+		.SetTypeDependencyGroup(TypeDependencyGroup);
+	RegisterOutputConnection(&Selection, &Selection)
+		.SetTypeDependencyGroup(TypeDependencyGroup);
 }
 
 void FCollectionSelectionInvertDataflowNode::Evaluate(UE::Dataflow::FContext& Context, const FDataflowOutput* Out) const
@@ -1601,14 +1587,4 @@ void FCollectionSelectionInvertDataflowNode::Evaluate(UE::Dataflow::FContext& Co
 	}
 }
 
-bool FCollectionSelectionInvertDataflowNode::OnInputTypeChanged(const FDataflowInput* Input)
-{
-	return SetOutputConcreteType(&Selection, Input->GetType());
-}
-
-bool FCollectionSelectionInvertDataflowNode::OnOutputTypeChanged(const FDataflowOutput* Input)
-{
-	return SetInputConcreteType(&Selection, Input->GetType());
-}
-
 //////////////////////////////////////////////////////////////////////////////////////////////////////////////////////
\ No newline at end of file
diff --git a/Engine/Plugins/Experimental/GeometryCollectionPlugin/Source/GeometryCollectionNodes/Public/Dataflow/GeometryCollectionSelectionNodes.h b/Engine/Plugins/Experimental/GeometryCollectionPlugin/Source/GeometryCollectionNodes/Public/Dataflow/GeometryCollectionSelectionNodes.h
--- a/Engine/Plugins/Experimental/GeometryCollectionPlugin/Source/GeometryCollectionNodes/Public/Dataflow/GeometryCollectionSelectionNodes.h
+++ b/Engine/Plugins/Experimental/GeometryCollectionPlugin/Source/GeometryCollectionNodes/Public/Dataflow/GeometryCollectionSelectionNodes.h
@@ -1576,9 +1576,6 @@ private:
 	UPROPERTY(meta = (DataflowOutput, DataflowPassthrough = "SelectionA"))
 	FDataflowSelectionTypes Selection;
 
-	virtual bool OnInputTypeChanged(const FDataflowInput* Input) override;
-	virtual bool OnOutputTypeChanged(const FDataflowOutput* Input) override;
-
 	virtual void Evaluate(UE::Dataflow::FContext& Context, const FDataflowOutput* Out) const override;
 
 public:
@@ -1602,9 +1599,6 @@ private:
 	UPROPERTY(meta = (DataflowInput, DataflowOutput, DataflowPassthrough = "Selection", DataflowIntrinsic))
 	FDataflowSelectionTypes Selection;
 
-	virtual bool OnInputTypeChanged(const FDataflowInput* Input) override;
-	virtual bool OnOutputTypeChanged(const FDataflowOutput* Input) override;
-
 	virtual void Evaluate(UE::Dataflow::FContext& Context, const FDataflowOutput* Out) const override;
 
 public:
diff --git a/Engine/Source/Runtime/Experimental/Dataflow/Core/Private/Dataflow/DataflowNode.cpp b/Engine/Source/Runtime/Experimental/Dataflow/Core/Private/Dataflow/DataflowNode.cpp
--- a/Engine/Source/Runtime/Experimental/Dataflow/Core/Private/Dataflow/DataflowNode.cpp
+++ b/Engine/Source/Runtime/Experimental/Dataflow/Core/Private/Dataflow/DataflowNode.cpp
@@ -1291,6 +1291,8 @@ void FDataflowNode::NotifyConnectionTypeChanged(FDataflowConnection* Connection)
 	if (Connection && Connection->IsAnyType())
 	{
 		OnConnectionTypeChanged(*Connection);
+
+PRAGMA_DISABLE_DEPRECATION_WARNINGS
 		if (Connection->GetDirection() == UE::Dataflow::FPin::EDirection::INPUT)
 		{
 			OnInputTypeChanged((FDataflowInput*)Connection);
@@ -1299,6 +1301,7 @@ void FDataflowNode::NotifyConnectionTypeChanged(FDataflowConnection* Connection)
 		{
 			OnOutputTypeChanged((FDataflowOutput*)Connection);
 		}
+PRAGMA_ENABLE_DEPRECATION_WARNINGS
 	}
 }
 
diff --git a/Engine/Source/Runtime/Experimental/Dataflow/Core/Private/Dataflow/DataflowVectorNodes.cpp b/Engine/Source/Runtime/Experimental/Dataflow/Core/Private/Dataflow/DataflowVectorNodes.cpp
--- a/Engine/Source/Runtime/Experimental/Dataflow/Core/Private/Dataflow/DataflowVectorNodes.cpp
+++ b/Engine/Source/Runtime/Experimental/Dataflow/Core/Private/Dataflow/DataflowVectorNodes.cpp
@@ -279,7 +279,7 @@ FDataflowVectorCrossProductNode::FDataflowVectorCrossProductNode(const UE::Dataf
 	RegisterInputConnection(&A);
 	RegisterInputConnection(&B);
 	RegisterOutputConnection(&CrossProduct);
-	// set default type to be double
+	// set default type to be vector 3
 	SetOutputConcreteType<FVector>(&CrossProduct);
 }
 
@@ -304,6 +304,8 @@ FDataflowVectorScaleNode::FDataflowVectorScaleNode(const UE::Dataflow::FNodePara
 	RegisterInputConnection(&V);
 	RegisterInputConnection(&Scale);
 	RegisterOutputConnection(&Scaled);
+	// set default type to be Vector4 for now
+	SetOutputConcreteType<FVector4>(&Scaled);
 }
 
 void FDataflowVectorScaleNode::Evaluate(UE::Dataflow::FContext& Context, const FDataflowOutput* Out) const
@@ -317,24 +319,6 @@ void FDataflowVectorScaleNode::Evaluate(UE::Dataflow::FContext& Context, const F
 	}
 }
 
-bool FDataflowVectorScaleNode::OnInputTypeChanged(const FDataflowInput* Input)
-{
-	if (Input->IsA(&V))
-	{
-		return SetOutputConcreteType(&Scaled, Input->GetType());
-	}
-	return false;
-}
-
-bool FDataflowVectorScaleNode::OnOutputTypeChanged(const FDataflowOutput* Output)
-{
-	if (Output->IsA(&Scaled))
-	{
-		return SetInputConcreteType(&V, Output->GetType());
-	}
-	return false;
-}
-
 //-----------------------------------------------------------
 
 FDataflowVectorNormalize::FDataflowVectorNormalize(const UE::Dataflow::FNodeParameters& InParam, FGuid InGuid)
@@ -342,6 +326,9 @@ FDataflowVectorNormalize::FDataflowVectorNormalize(const UE::Dataflow::FNodePara
 {
 	RegisterInputConnection(&V);
 	RegisterOutputConnection(&Normalized);
+	// set default type to be Vector4 for now
+	SetOutputConcreteType<FVector4>(&Normalized);
+
 }
 
 void FDataflowVectorNormalize::Evaluate(UE::Dataflow::FContext& Context, const FDataflowOutput* Out) const
@@ -353,21 +340,3 @@ void FDataflowVectorNormalize::Evaluate(UE::Dataflow::FContext& Context, const F
 		SetValue(Context, OutResult, &Normalized);
 	}
 }
-
-bool FDataflowVectorNormalize::OnInputTypeChanged(const FDataflowInput* Input)
-{
-	if (Input->IsA(&V))
-	{
-		return SetOutputConcreteType(&Normalized, Input->GetType());
-	}
-	return false;
-}
-
-bool FDataflowVectorNormalize::OnOutputTypeChanged(const FDataflowOutput* Output)
-{
-	if (Output->IsA(&Normalized))
-	{
-		return SetInputConcreteType(&V, Output->GetType());
-	}
-	return false;
-}
diff --git a/Engine/Source/Runtime/Experimental/Dataflow/Core/Public/Dataflow/DataflowNode.h b/Engine/Source/Runtime/Experimental/Dataflow/Core/Public/Dataflow/DataflowNode.h
--- a/Engine/Source/Runtime/Experimental/Dataflow/Core/Public/Dataflow/DataflowNode.h
+++ b/Engine/Source/Runtime/Experimental/Dataflow/Core/Public/Dataflow/DataflowNode.h
@@ -596,7 +596,11 @@ PRAGMA_ENABLE_DEPRECATION_WARNINGS
 
 protected:
 	bool OnConnectionTypeChanged(const FDataflowConnection& Connection);
+
+	UE_DEPRECATED(5.6, "Use inputs and outputs SetTypeDependencyGroup() to declare dependent types")
 	virtual bool OnInputTypeChanged(const FDataflowInput* Input) { return false; }
+
+	UE_DEPRECATED(5.6, "Use inputs and outputs SetTypeDependencyGroup() to declare dependent types")
 	virtual bool OnOutputTypeChanged(const FDataflowOutput* Output) { return false; }
 
 	template <typename T>
diff --git a/Engine/Source/Runtime/Experimental/Dataflow/Core/Public/Dataflow/DataflowVectorNodes.h b/Engine/Source/Runtime/Experimental/Dataflow/Core/Public/Dataflow/DataflowVectorNodes.h
--- a/Engine/Source/Runtime/Experimental/Dataflow/Core/Public/Dataflow/DataflowVectorNodes.h
+++ b/Engine/Source/Runtime/Experimental/Dataflow/Core/Public/Dataflow/DataflowVectorNodes.h
@@ -338,9 +338,6 @@ private:
 	FDataflowVectorTypes Scaled;
 
 	virtual void Evaluate(UE::Dataflow::FContext& Context, const FDataflowOutput* Out) const override;
-
-	virtual bool OnInputTypeChanged(const FDataflowInput* Input) override;
-	virtual bool OnOutputTypeChanged(const FDataflowOutput* Output) override;
 };
 
 /** Normalize a vector : Normalized = V/|V| */
@@ -363,9 +360,6 @@ private:
 	FDataflowVectorTypes Normalized;
 
 	virtual void Evaluate(UE::Dataflow::FContext& Context, const FDataflowOutput* Out) const override;
-
-	virtual bool OnInputTypeChanged(const FDataflowInput* Input) override;
-	virtual bool OnOutputTypeChanged(const FDataflowOutput* Output) override;
 };
 
 namespace UE::Dataflow
