Dataflow:
- Updated FGetCollectionAttributeDataTypedDataflowNode to use FDataflowNumericArrayTypes, FDataflowVectorArrayTypes, FDataflowStringArrayTypes
- new version is v2, old version was deprecated
- Added FSpheresToPointsDataflowNode
- Bug fixes

diff --git a/Engine/Plugins/Experimental/GeometryCollectionPlugin/Source/GeometryCollectionNodes/Private/Dataflow/GeometryCollectionNodes.cpp b/Engine/Plugins/Experimental/GeometryCollectionPlugin/Source/GeometryCollectionNodes/Private/Dataflow/GeometryCollectionNodes.cpp
--- a/Engine/Plugins/Experimental/GeometryCollectionPlugin/Source/GeometryCollectionNodes/Private/Dataflow/GeometryCollectionNodes.cpp
+++ b/Engine/Plugins/Experimental/GeometryCollectionPlugin/Source/GeometryCollectionNodes/Private/Dataflow/GeometryCollectionNodes.cpp
@@ -76,6 +76,8 @@ namespace UE::Dataflow
 		DATAFLOW_NODE_REGISTER_CREATION_FACTORY(FAddCustomCollectionAttributeDataflowNode);
 		DATAFLOW_NODE_REGISTER_CREATION_FACTORY(FGetNumElementsInCollectionGroupDataflowNode);
 		DATAFLOW_NODE_REGISTER_CREATION_FACTORY(FGetCollectionAttributeDataTypedDataflowNode);
+		// Commented out until AnyType outputs can properly change types
+//		DATAFLOW_NODE_REGISTER_CREATION_FACTORY(FGetCollectionAttributeDataTypedDataflowNode_v2);
 		DATAFLOW_NODE_REGISTER_CREATION_FACTORY(FSetCollectionAttributeDataTypedDataflowNode);
 		DATAFLOW_NODE_REGISTER_CREATION_FACTORY(FMultiplyTransformDataflowNode);
 		DATAFLOW_NODE_REGISTER_CREATION_FACTORY(FInvertTransformDataflowNode);
@@ -87,6 +89,7 @@ namespace UE::Dataflow
 		DATAFLOW_NODE_REGISTER_CREATION_FACTORY(FVisualizeTetrahedronsDataflowNode);
 		DATAFLOW_NODE_REGISTER_CREATION_FACTORY(FPointsToCollectionDataflowNode);
 		DATAFLOW_NODE_REGISTER_CREATION_FACTORY(FCollectionToPointsDataflowNode);
+		DATAFLOW_NODE_REGISTER_CREATION_FACTORY(FSpheresToPointsDataflowNode);
 
 		// GeometryCollection
 		DATAFLOW_NODE_REGISTER_CREATION_FACTORY_NODE_COLORS_BY_CATEGORY("GeometryCollection", FLinearColor(0.55f, 0.45f, 1.0f), CDefaultNodeBodyTintColor);
@@ -917,7 +920,6 @@ void FGetNumElementsInCollectionGroupDataflowNode::Evaluate(UE::Dataflow::FConte
 	}
 }
 
-
 void FGetCollectionAttributeDataTypedDataflowNode::Evaluate(UE::Dataflow::FContext& Context, const FDataflowOutput* Out) const
 {
 	if (Out->IsA<TArray<bool>>(&BoolAttributeData) ||
@@ -929,7 +931,7 @@ void FGetCollectionAttributeDataTypedDataflowNode::Evaluate(UE::Dataflow::FConte
 		Out->IsA<TArray<FVector3d>>(&Vector3dAttributeData))
 	{
 		const FManagedArrayCollection& InCollection = GetValue<FManagedArrayCollection>(Context, &Collection);
-		
+
 		FName InputGroupName;
 		if (GroupName != EStandardGroupNameEnum::Dataflow_EStandardGroupNameEnum_Custom)
 		{
@@ -1024,6 +1026,134 @@ void FGetCollectionAttributeDataTypedDataflowNode::Evaluate(UE::Dataflow::FConte
 	}
 }
 
+void FGetCollectionAttributeDataTypedDataflowNode_v2::Evaluate(UE::Dataflow::FContext& Context, const FDataflowOutput* Out) const
+{
+	if (Out->IsA(&Collection))
+	{
+		SafeForwardInput(Context, &Collection, &Collection);
+	}
+	else if (Out->IsA(&BoolAttributeData) ||
+		Out->IsA(&NumericArray) ||
+		Out->IsA(&VectorArray) ||
+		Out->IsA(&StringArray))
+	{
+		const FManagedArrayCollection& InCollection = GetValue<FManagedArrayCollection>(Context, &Collection);
+		
+		FName InputGroupName;
+		if (GroupName != EStandardGroupNameEnum::Dataflow_EStandardGroupNameEnum_Custom)
+		{
+			InputGroupName = GetGroupName(GroupName);
+		}
+		else
+		{
+			InputGroupName = FName(*CustomGroupName);
+		}
+
+		SetValue(Context, TArray<bool>(), &BoolAttributeData);
+		SetValue(Context, TArray<double>(), &NumericArray);
+		SetValue(Context, TArray<FVector4>(), &VectorArray);
+		SetValue(Context, TArray<FString>(), &StringArray);
+
+		FCollectionAttributeKey DefaultAttributeKey(AttrName, InputGroupName.ToString());
+		FCollectionAttributeKey AttributeKeyVal = GetValue(Context, &AttributeKey, DefaultAttributeKey);
+		FName GroupNameVal = FName(AttributeKeyVal.Group);
+		FName AttributeNameVal = FName(AttributeKeyVal.Attribute);
+
+		if (GroupNameVal.GetStringLength() > 0 && AttributeNameVal.GetStringLength() > 0)
+		{
+			if (InCollection.HasGroup(GroupNameVal))
+			{
+				if (InCollection.HasAttribute(AttributeNameVal, GroupNameVal))
+				{
+					FString TypeStr = GetArrayTypeString(InCollection.GetAttributeType(AttributeNameVal, GroupNameVal)).ToString();
+
+					if (TypeStr == FString("Bool"))
+					{
+						if (const TManagedArray<bool>* AttributeArr = InCollection.FindAttribute<bool>(AttributeNameVal, GroupNameVal))
+						{
+							TArray<bool> BoolArray = AttributeArr->GetAsBoolArray();
+							SetValue(Context, MoveTemp(BoolArray), &BoolAttributeData);
+						}
+					}
+					else if (TypeStr == FString("Float"))
+					{
+						if (const TManagedArray<float>* AttributeArr = InCollection.FindAttribute<float>(AttributeNameVal, GroupNameVal))
+						{
+							TArray<double> Arr; Arr.Reserve(AttributeArr->Num());
+							Algo::Copy(AttributeArr->GetConstArray(), Arr);
+
+							SetValue(Context, MoveTemp(Arr), &NumericArray);
+						}
+					}
+					else if (TypeStr == FString("Double"))
+					{
+						if (const TManagedArray<double>* AttributeArr = InCollection.FindAttribute<double>(AttributeNameVal, GroupNameVal))
+						{
+							SetValue(Context, AttributeArr->GetConstArray(), &NumericArray);
+						}
+					}
+					else if (TypeStr == FString("Int32"))
+					{
+						if (const TManagedArray<int32>* AttributeArr = InCollection.FindAttribute<int32>(AttributeNameVal, GroupNameVal))
+						{
+							TArray<double> Arr; Arr.Reserve(AttributeArr->Num());
+							Algo::Copy(AttributeArr->GetConstArray(), Arr);
+
+							SetValue(Context, MoveTemp(Arr), &NumericArray);
+						}
+					}
+					else if (TypeStr == FString("String"))
+					{
+						if (const TManagedArray<FString>* AttributeArr = InCollection.FindAttribute<FString>(AttributeNameVal, GroupNameVal))
+						{
+							SetValue(Context, AttributeArr->GetConstArray(), &StringArray);
+						}
+					}
+					else if (TypeStr == FString("Vector"))
+					{
+						if (const TManagedArray<FVector3f>* AttributeArr = InCollection.FindAttribute<FVector3f>(AttributeNameVal, GroupNameVal))
+						{
+							TArray<FVector4> Arr; Arr.Reserve(AttributeArr->Num());
+							for (int32 Idx = 0; Idx < AttributeArr->Num(); ++Idx)
+							{
+								const FVector3f Vec = AttributeArr->GetConstArray()[Idx];
+								Arr[Idx] = FVector4(Vec.X, Vec.Y, Vec.Z, 0.0);
+							}
+							SetValue(Context, MoveTemp(Arr), &VectorArray);
+						}
+					}
+					else if (TypeStr == FString("Vector3d"))
+					{
+						if (const TManagedArray<FVector3d>* AttributeArr = InCollection.FindAttribute<FVector3d>(AttributeNameVal, GroupNameVal))
+						{
+							TArray<FVector4> Arr; Arr.Reserve(AttributeArr->Num());
+							for (int32 Idx = 0; Idx < AttributeArr->Num(); ++Idx)
+							{
+								const FVector3d Vec = AttributeArr->GetConstArray()[Idx];
+								Arr[Idx] = FVector4(Vec.X, Vec.Y, Vec.Z, 0.0);
+							}
+							SetValue(Context, MoveTemp(Arr), &VectorArray);
+						}
+					}
+					else if (TypeStr == FString("LinearColor"))
+					{
+						if (const TManagedArray<FLinearColor>* AttributeArr = InCollection.FindAttribute<FLinearColor>(AttributeNameVal, GroupNameVal))
+						{
+							TArray<FVector4> Arr; Arr.Reserve(AttributeArr->Num());
+							for (int32 Idx = 0; Idx < AttributeArr->Num(); ++Idx)
+							{
+								const FLinearColor Vec = AttributeArr->GetConstArray()[Idx];
+								Arr[Idx] = FVector4(Vec.R, Vec.G, Vec.B, Vec.A);
+							}
+							SetValue(Context, MoveTemp(Arr), &VectorArray);
+						}
+					}
+				}
+			}
+		}
+	}
+}
+
 template<typename T>
 static void SetAttributeData(const FDataflowNode* DataflowNode, UE::Dataflow::FContext& Context, FManagedArrayCollection& InCollection, const TArray<T>& Property, FName AttributeName, FName GroupName)
 {
@@ -1385,4 +1515,37 @@ void FCollectionToPointsDataflowNode::Evaluate(UE::Dataflow::FContext& Context,
 
 		SetValue(Context, TArray<FVector>(), &Points);
 	}
+}
+
+void FSpheresToPointsDataflowNode::Evaluate(UE::Dataflow::FContext& Context, const FDataflowOutput* Out) const
+{
+	if (Out->IsA(&Points) ||
+		Out->IsA(&Radii))
+	{
+		if (IsConnected(&Spheres))
+		{
+			const TArray<FSphere>& InSpheres = GetValue(Context, &Spheres);
+
+			const int32 NumSpheres = InSpheres.Num();
+
+			if (NumSpheres > 0)
+			{
+				TArray<FVector> OutPoints; OutPoints.AddUninitialized(NumSpheres);
+				TArray<float> OutRadii; OutRadii.AddUninitialized(NumSpheres);
+
+				for (int32 Idx = 0; Idx < NumSpheres; ++Idx)
+				{
+					OutPoints[Idx] = InSpheres[Idx].Center;
+					OutRadii[Idx] = InSpheres[Idx].W;
+				}
+
+				SetValue(Context, MoveTemp(OutPoints), &Points);
+				SetValue(Context, MoveTemp(OutRadii), &Radii);
+				return;
+			}
+		}
+
+		SetValue(Context, TArray<FVector>(), &Points);
+		SetValue(Context, TArray<float>(), &Radii);
+	}
 }
\ No newline at end of file
diff --git a/Engine/Plugins/Experimental/GeometryCollectionPlugin/Source/GeometryCollectionNodes/Public/Dataflow/GeometryCollectionNodes.h b/Engine/Plugins/Experimental/GeometryCollectionPlugin/Source/GeometryCollectionNodes/Public/Dataflow/GeometryCollectionNodes.h
--- a/Engine/Plugins/Experimental/GeometryCollectionPlugin/Source/GeometryCollectionNodes/Public/Dataflow/GeometryCollectionNodes.h
+++ b/Engine/Plugins/Experimental/GeometryCollectionPlugin/Source/GeometryCollectionNodes/Public/Dataflow/GeometryCollectionNodes.h
@@ -1372,7 +1372,6 @@ public:
 
 };
 
-
 /**
  *
  * Get attribute data from a Collection
@@ -1456,6 +1455,69 @@ public:
 
 };
 
+/**
+ *
+ * Get attribute data from a Collection
+ *
+ */
+USTRUCT(meta = (DataflowGeometryCollection))
+struct FGetCollectionAttributeDataTypedDataflowNode_v2 : public FDataflowNode
+{
+	GENERATED_USTRUCT_BODY()
+	DATAFLOW_NODE_DEFINE_INTERNAL(FGetCollectionAttributeDataTypedDataflowNode_v2, "GetCollectionAttributeDataTyped", "GeometryCollection|Utilities", "")
+
+private:
+	/** Collection for the custom attribute */
+	UPROPERTY(meta = (DataflowInput, DataflowOutput, DataflowPassthrough = "Collection", DataflowIntrinsic))
+	FManagedArrayCollection Collection;
+
+	/** Input to drive the Attribute and Group name */
+	UPROPERTY(meta = (DataflowInput, DisplayName = "AttributeKey"))
+	FCollectionAttributeKey AttributeKey;
+
+	/** Standard group names */
+	UPROPERTY(EditAnywhere, Category = "Attribute", meta = (DisplayName = "Group"))
+	EStandardGroupNameEnum GroupName = EStandardGroupNameEnum::Dataflow_EStandardGroupNameEnum_Transform;
+
+	/** User specified group name */
+	UPROPERTY(EditAnywhere, Category = "Attribute", meta = (DisplayName = "Custom Group", EditCondition = "GroupName == EStandardGroupNameEnum::Dataflow_EStandardGroupNameEnum_Custom"))
+	FString CustomGroupName = FString("");
+
+	/** Attribute name */
+	UPROPERTY(EditAnywhere, Category = "Attribute", meta = (DisplayName = "Attribute"))
+	FString AttrName = FString("");
+
+	/** Bool type attribute data */
+	UPROPERTY(meta = (DataflowOutput, DisplayName = "Bool Array"));
+	TArray<bool> BoolAttributeData;
+
+	/** Numeric Array types */
+	UPROPERTY(meta = (DataflowOutput));
+	FDataflowNumericArrayTypes NumericArray;
+
+	/** Vector Array types */
+	UPROPERTY(meta = (DataflowOutput));
+	FDataflowVectorArrayTypes VectorArray;
+
+	/** String Array types */
+	UPROPERTY(meta = (DataflowOutput));
+	FDataflowStringArrayTypes StringArray;
+
+public:
+	FGetCollectionAttributeDataTypedDataflowNode_v2(const UE::Dataflow::FNodeParameters& InParam, FGuid InGuid = FGuid::NewGuid())
+		: FDataflowNode(InParam, InGuid)
+	{
+		RegisterInputConnection(&Collection);
+		RegisterInputConnection(&AttributeKey);
+		RegisterOutputConnection(&Collection, &Collection);
+		RegisterOutputConnection(&BoolAttributeData);
+		RegisterOutputConnection(&NumericArray);
+		RegisterOutputConnection(&VectorArray);
+		RegisterOutputConnection(&StringArray);
+	}
+
+	virtual void Evaluate(UE::Dataflow::FContext& Context, const FDataflowOutput* Out) const override;
+};
 
 /**
  *
@@ -1833,6 +1895,45 @@ public:
 	virtual void Evaluate(UE::Dataflow::FContext& Context, const FDataflowOutput* Out) const override;
 };
 
+/**
+ *
+ * Outputs Spheres as Points and radius values
+ *
+ */
+USTRUCT(meta = (DataflowGeometryCollection))
+struct FSpheresToPointsDataflowNode : public FDataflowNode
+{
+	GENERATED_USTRUCT_BODY()
+	DATAFLOW_NODE_DEFINE_INTERNAL(FSpheresToPointsDataflowNode, "SpheresToPoints", "GeometryCollection|Utilities", "")
+	DATAFLOW_NODE_RENDER_TYPE_START()
+	DATAFLOW_NODE_RENDER_TYPE_ADD("PointsRender", FName("TArray<FVector>"), "Points")
+	DATAFLOW_NODE_RENDER_TYPE_END()
+
+private:
+	/** Input Spheres  */
+	UPROPERTY(meta = (DataflowInputDataflowIntrinsic))
+	TArray<FSphere> Spheres;
+
+	/** Centers of the spheres */
+	UPROPERTY(meta = (DataflowOutput))
+	TArray<FVector> Points;
+
+	/** Radius values */
+	UPROPERTY(meta = (DataflowOutput))
+	TArray<float> Radii;
+
+public:
+	FSpheresToPointsDataflowNode(const UE::Dataflow::FNodeParameters& InParam, FGuid InGuid = FGuid::NewGuid())
+		: FDataflowNode(InParam, InGuid)
+	{
+		RegisterInputConnection(&Spheres);
+		RegisterOutputConnection(&Points);
+		RegisterOutputConnection(&Radii);
+	}
+
+	virtual void Evaluate(UE::Dataflow::FContext& Context, const FDataflowOutput* Out) const override;
+};
+
 namespace UE::Dataflow
 {
 	void GeometryCollectionEngineNodes();
