Dataflow : Fix assign to collection material node

diff --git a/Engine/Plugins/Experimental/GeometryCollectionPlugin/Source/GeometryCollectionNodes/Private/Dataflow/GeometryCollectionMaterialInterfaceNodes.cpp b/Engine/Plugins/Experimental/GeometryCollectionPlugin/Source/GeometryCollectionNodes/Private/Dataflow/GeometryCollectionMaterialInterfaceNodes.cpp
--- a/Engine/Plugins/Experimental/GeometryCollectionPlugin/Source/GeometryCollectionNodes/Private/Dataflow/GeometryCollectionMaterialInterfaceNodes.cpp
+++ b/Engine/Plugins/Experimental/GeometryCollectionPlugin/Source/GeometryCollectionNodes/Private/Dataflow/GeometryCollectionMaterialInterfaceNodes.cpp
@@ -275,7 +275,7 @@ void FAssignMaterialInterfaceToCollectionDataflowNode::Evaluate(UE::Dataflow::FC
 		// add the material to the array 
 		const int32 InMaterialIndex = AddOrMergeMaterialToArray(InMaterialArray, InMaterial);
 
-		const int32 NumFaces = Collection.NumElements(FGeometryCollection::FacesGroup);
+		const int32 NumFaces = InCollection.NumElements(FGeometryCollection::FacesGroup);
 		if (!IsConnected(&FaceSelection))
 		{
 			// if not connected select all faces
@@ -283,7 +283,7 @@ void FAssignMaterialInterfaceToCollectionDataflowNode::Evaluate(UE::Dataflow::FC
 		}
 
 		static const FName MaterialIdAttributeName = "MaterialID";
-		if (InFaceSelection.Num() == NumFaces && 
+		if (InFaceSelection.Num() <= NumFaces && 
 			InCollection.HasAttribute(MaterialIdAttributeName, FGeometryCollection::FacesGroup))
 		{
 			TManagedArray<int32>& MaterialIDs = InCollection.ModifyAttribute<int32>(MaterialIdAttributeName, FGeometryCollection::FacesGroup);
