GeometryCollection - Fixed a missing SetValue on the StaticMeshToMeshDataflow node, that is triggering an ensure everytime the node is added to the graph.

#rnx

diff --git a/Engine/Plugins/Experimental/GeometryCollectionPlugin/Source/GeometryCollectionNodes/Private/Dataflow/GeometryCollectionMeshNodes.cpp b/Engine/Plugins/Experimental/GeometryCollectionPlugin/Source/GeometryCollectionNodes/Private/Dataflow/GeometryCollectionMeshNodes.cpp
--- a/Engine/Plugins/Experimental/GeometryCollectionPlugin/Source/GeometryCollectionNodes/Private/Dataflow/GeometryCollectionMeshNodes.cpp
+++ b/Engine/Plugins/Experimental/GeometryCollectionPlugin/Source/GeometryCollectionNodes/Private/Dataflow/GeometryCollectionMeshNodes.cpp
@@ -239,12 +239,10 @@ void FStaticMeshToMeshDataflowNode::Evaluate(UE::Dataflow::FContext& Context, co
 				}
 
 				SetValue(Context, NewMesh, &Mesh);
-			}
-			else
-			{
-				SetValue(Context, TObjectPtr<UDynamicMesh>(NewObject<UDynamicMesh>()), &Mesh);
+				return;
 			}
 		}
+		SetValue(Context, TObjectPtr<UDynamicMesh>(NewObject<UDynamicMesh>()), &Mesh);
 	}
 #endif
 }
