Dataflow : Deprecate material nodes that could not handle material interface ( they have been replaced by the material interface related nodes )

diff --git a/Engine/Plugins/Experimental/GeometryCollectionPlugin/Source/GeometryCollectionNodes/Public/Dataflow/GeometryCollectionMaterialNodes.h b/Engine/Plugins/Experimental/GeometryCollectionPlugin/Source/GeometryCollectionNodes/Public/Dataflow/GeometryCollectionMaterialNodes.h
--- a/Engine/Plugins/Experimental/GeometryCollectionPlugin/Source/GeometryCollectionNodes/Public/Dataflow/GeometryCollectionMaterialNodes.h
+++ b/Engine/Plugins/Experimental/GeometryCollectionPlugin/Source/GeometryCollectionNodes/Public/Dataflow/GeometryCollectionMaterialNodes.h
@@ -27,9 +27,10 @@ class UMaterial;
 /**
  *
  * Adds Outside/Inside Materials to Outside/Inside faces
+ * Deprecated (5.6) - Use AssignMaterialToCollection node instead
  *
  */
-USTRUCT(meta = (DataflowGeometryCollection))
+USTRUCT(meta = (DataflowGeometryCollection, Deprecated = "5.6"))
 struct FAddMaterialToCollectionDataflowNode : public FDataflowNode
 {
 	GENERATED_USTRUCT_BODY()
@@ -84,9 +85,10 @@ public:
 /**
  *
  * Reassign existing material(s) to Outside/Inside faces
+ * Deprecated (5.6) - Use AssignMaterialToCollection node instead
  *
  */
-USTRUCT(meta = (DataflowGeometryCollection))
+USTRUCT(meta = (DataflowGeometryCollection, Deprecated = "5.6"))
 struct FReAssignMaterialInCollectionDataflowNode : public FDataflowNode
 {
 	GENERATED_USTRUCT_BODY()
@@ -141,9 +143,10 @@ public:
 /**
  *
  * Generates a formatted string of materials from the Materials array
+ * Deprecated (5.6) 
  *
  */
-USTRUCT(meta = (DataflowGeometryCollection))
+USTRUCT(meta = (DataflowGeometryCollection, Deprecated = "5.6"))
 struct FMaterialsInfoDataflowNode : public FDataflowNode
 {
 	GENERATED_USTRUCT_BODY()
@@ -173,9 +176,10 @@ public:
 /**
  *
  * Get a Material from a Materials array
+ * Deprecated (5.6) - Use GetFromMaterialArray instead
  *
  */
-USTRUCT(meta = (DataflowGeometryCollection))
+USTRUCT(meta = (DataflowGeometryCollection, Deprecated = "5.6"))
 struct FGetMaterialFromMaterialsArrayDataflowNode : public FDataflowNode
 {
 	GENERATED_USTRUCT_BODY()
@@ -220,9 +224,10 @@ enum class ESetMaterialOperationTypeEnum : uint8
 /**
  *
  * Set a Material in a Materials array
+ * Deprecated (5.6) - Use SetIntoMaterialArray instead
  *
  */
-USTRUCT(meta = (DataflowGeometryCollection))
+USTRUCT(meta = (DataflowGeometryCollection, Deprecated = "5.6"))
 struct FSetMaterialInMaterialsArrayDataflowNode : public FDataflowNode
 {
 	GENERATED_USTRUCT_BODY()
@@ -263,9 +268,10 @@ public:
 /**
  *
  * Makes a material
+ * Deprecated (5.6)
  *
  */
-USTRUCT()
+USTRUCT(meta=(Deprecated = "5.6"))
 struct FMakeMaterialDataflowNode : public FDataflowNode
 {
 	GENERATED_USTRUCT_BODY()
@@ -293,9 +299,10 @@ public:
 /**
  *
  * Makes an empty Materials array
+ * Deprecated (5.6)
  *
  */
-USTRUCT()
+USTRUCT(meta = (Deprecated = "5.6"))
 struct FMakeMaterialsArrayDataflowNode : public FDataflowNode
 {
 	GENERATED_USTRUCT_BODY()
