Dataflow : new versions of StaticMeshToCollection and GeometryCollectionTerminal nodes  to remove reference to UMaterial input sor outputs and only use UmaterialInterface ones
Older version of the nodes have been deprecated

diff --git a/Engine/Plugins/Experimental/GeometryCollectionPlugin/Source/GeometryCollectionNodes/Private/Dataflow/GeometryCollectionAssetNodes.cpp b/Engine/Plugins/Experimental/GeometryCollectionPlugin/Source/GeometryCollectionNodes/Private/Dataflow/GeometryCollectionAssetNodes.cpp
--- a/Engine/Plugins/Experimental/GeometryCollectionPlugin/Source/GeometryCollectionNodes/Private/Dataflow/GeometryCollectionAssetNodes.cpp
+++ b/Engine/Plugins/Experimental/GeometryCollectionPlugin/Source/GeometryCollectionNodes/Private/Dataflow/GeometryCollectionAssetNodes.cpp
@@ -19,18 +19,70 @@ namespace UE::Dataflow
 	{
 		static const FLinearColor CDefaultNodeBodyTintColor = FLinearColor(0.f, 0.f, 0.f, 0.5f);
 
-		DATAFLOW_NODE_REGISTER_CREATION_FACTORY(FGeometryCollectionTerminalDataflowNode);
+		DATAFLOW_NODE_REGISTER_CREATION_FACTORY(FGeometryCollectionTerminalDataflowNode_v2);
 		DATAFLOW_NODE_REGISTER_CREATION_FACTORY(FGetGeometryCollectionAssetDataflowNode);
 		DATAFLOW_NODE_REGISTER_CREATION_FACTORY(FGetGeometryCollectionSourcesDataflowNode);
 		DATAFLOW_NODE_REGISTER_CREATION_FACTORY(FCreateGeometryCollectionFromSourcesDataflowNode);
 		DATAFLOW_NODE_REGISTER_CREATION_FACTORY(FGeometryCollectionToCollectionDataflowNode);
 		DATAFLOW_NODE_REGISTER_CREATION_FACTORY(FBlueprintToCollectionDataflowNode);
 
+		// deprecated nodes (need to stay registered)
+		DATAFLOW_NODE_REGISTER_CREATION_FACTORY(FGeometryCollectionTerminalDataflowNode);
+
 		// Terminal
 		DATAFLOW_NODE_REGISTER_CREATION_FACTORY_NODE_COLORS_BY_CATEGORY("Terminal", FLinearColor(0.f, 0.f, 0.f), CDefaultNodeBodyTintColor);
 	}
 }
 
+// ===========================================================================================================================
+
+FGeometryCollectionTerminalDataflowNode_v2::FGeometryCollectionTerminalDataflowNode_v2(const UE::Dataflow::FNodeParameters& InParam, FGuid InGuid)
+	: FDataflowTerminalNode(InParam, InGuid)
+{
+	RegisterInputConnection(&Collection);
+	RegisterOutputConnection(&Collection, &Collection);
+	RegisterInputConnection(&Materials);
+	RegisterOutputConnection(&Materials, &Materials);
+	RegisterInputConnection(&InstancedMeshes);
+	RegisterOutputConnection(&InstancedMeshes, &InstancedMeshes);
+}
+
+void FGeometryCollectionTerminalDataflowNode_v2::SetAssetValue(TObjectPtr<UObject> Asset, UE::Dataflow::FContext& Context) const
+{
+	using FGeometryCollectionPtr = TSharedPtr<FGeometryCollection, ESPMode::ThreadSafe>;
+	using FMaterialArray = TArray<TObjectPtr<UMaterialInterface>>;
+	using FInstancedMeshesArray = TArray<FGeometryCollectionAutoInstanceMesh>;
+
+	if (UGeometryCollection* CollectionAsset = Cast<UGeometryCollection>(Asset.Get()))
+	{
+		if (FGeometryCollectionPtr GeometryCollection = CollectionAsset->GetGeometryCollection())
+		{
+			const FManagedArrayCollection& InCollection = GetValue(Context, &Collection);
+			const FMaterialArray& InMaterials = GetValue(Context, &Materials);
+			const FInstancedMeshesArray& InInstancedMeshes = GetValue(Context, &InstancedMeshes);
+
+			constexpr bool bHasInternalMaterial = false; // with dataflow there's no assumption of internal materials
+			CollectionAsset->ResetFrom(InCollection, InMaterials, bHasInternalMaterial);
+
+			CollectionAsset->SetAutoInstanceMeshes(InInstancedMeshes);
+
+#if WITH_EDITOR
+			// make sure we rebuild the render data when we are done setting everything 
+			CollectionAsset->RebuildRenderData();
+			// also make sure all components using it are getting a notification about it
+			CollectionAsset->PropagateTransformUpdateToComponents();
+#endif
+		}
+	}
+}
+
+void FGeometryCollectionTerminalDataflowNode_v2::Evaluate(UE::Dataflow::FContext& Context) const
+{
+	// simply forward all inputs to corresponding outputs
+	SafeForwardInput(Context, &Collection, &Collection);
+	SafeForwardInput(Context, &Materials, &Materials);
+	SafeForwardInput(Context, &InstancedMeshes, &InstancedMeshes);
+}
 
 // ===========================================================================================================================
 
diff --git a/Engine/Plugins/Experimental/GeometryCollectionPlugin/Source/GeometryCollectionNodes/Private/Dataflow/GeometryCollectionNodesPlugin.cpp b/Engine/Plugins/Experimental/GeometryCollectionPlugin/Source/GeometryCollectionNodes/Private/Dataflow/GeometryCollectionNodesPlugin.cpp
--- a/Engine/Plugins/Experimental/GeometryCollectionPlugin/Source/GeometryCollectionNodes/Private/Dataflow/GeometryCollectionNodesPlugin.cpp
+++ b/Engine/Plugins/Experimental/GeometryCollectionPlugin/Source/GeometryCollectionNodes/Private/Dataflow/GeometryCollectionNodesPlugin.cpp
@@ -38,7 +38,6 @@
 #include "Dataflow/SetVertexColorFromVertexIndicesNode.h"
 #include "Dataflow/SetVertexColorFromVertexSelectionNode.h"
 
-
 #define LOCTEXT_NAMESPACE "DataflowNodes"
 
 
@@ -66,6 +65,7 @@ void IGeometryCollectionNodesPlugin::StartupModule()
 	UE::Dataflow::RegisterGeometryCollectionUVNodes();
 	UE::Dataflow::RegisterGeometryCollectionTextureNodes();
 	UE::Dataflow::RegisterGeometryCollectionMaterialInterfaceNodes();
+	UE::Dataflow::RegisterGeometryStaticMeshToCollectionNodes();
 	DATAFLOW_NODE_REGISTER_CREATION_FACTORY(FCreateColorArrayFromFloatArrayDataflowNode);
 	DATAFLOW_NODE_REGISTER_CREATION_FACTORY(FGeometryCollectionVertexScalarToVertexIndicesNode);
 	DATAFLOW_NODE_REGISTER_CREATION_FACTORY(FSetVertexColorFromFloatArrayDataflowNode);
@@ -75,7 +75,6 @@ void IGeometryCollectionNodesPlugin::StartupModule()
 	DATAFLOW_NODE_REGISTER_CREATION_FACTORY(FGeometryCollectionTransferVertexAttributeNode);
 	DATAFLOW_NODE_REGISTER_CREATION_FACTORY(FGeometryCollectionTransferVertexSkinWeightsNode);
 	DATAFLOW_NODE_REGISTER_CREATION_FACTORY(FGeometryCollectionSetKinematicVertexSelectionNode);
-	DATAFLOW_NODE_REGISTER_CREATION_FACTORY(FStaticMeshToCollectionDataflowNode);
 	DATAFLOW_NODE_REGISTER_CREATION_FACTORY(FTriangleBoundaryIndicesNode);
 	// Deprecated nodes
 PRAGMA_DISABLE_DEPRECATION_WARNINGS
diff --git a/Engine/Plugins/Experimental/GeometryCollectionPlugin/Source/GeometryCollectionNodes/Private/Dataflow/GeometryCollectionStaticMeshToCollectionNode.cpp b/Engine/Plugins/Experimental/GeometryCollectionPlugin/Source/GeometryCollectionNodes/Private/Dataflow/GeometryCollectionStaticMeshToCollectionNode.cpp
--- a/Engine/Plugins/Experimental/GeometryCollectionPlugin/Source/GeometryCollectionNodes/Private/Dataflow/GeometryCollectionStaticMeshToCollectionNode.cpp
+++ b/Engine/Plugins/Experimental/GeometryCollectionPlugin/Source/GeometryCollectionNodes/Private/Dataflow/GeometryCollectionStaticMeshToCollectionNode.cpp
@@ -9,9 +9,54 @@
 #include "GeometryCollection/GeometryCollectionObject.h"
 #include "GeometryCollection/ManagedArrayCollection.h"
 
+namespace UE::Dataflow
+{
+	void RegisterGeometryStaticMeshToCollectionNodes()
+	{
+		DATAFLOW_NODE_REGISTER_CREATION_FACTORY(FStaticMeshToCollectionDataflowNode_v2);
+
+		// Deprecated nodes (Need to stay registered) 
+		DATAFLOW_NODE_REGISTER_CREATION_FACTORY(FStaticMeshToCollectionDataflowNode);
+	}
+}
+
+////////////////////////////////////////////////////////////////////////////////////////////////////////////////////////////////////////////
+
+FStaticMeshToCollectionDataflowNode_v2::FStaticMeshToCollectionDataflowNode_v2(const UE::Dataflow::FNodeParameters& InParam, FGuid InGuid)
+	: FDataflowNode(InParam, InGuid)
+{
+	RegisterInputConnection(&StaticMesh);
+	RegisterInputConnection(&MeshTransform);
+	RegisterOutputConnection(&Collection);
+	RegisterOutputConnection(&Materials);
+	RegisterOutputConnection(&InstancedMeshes);
+}
+
+void FStaticMeshToCollectionDataflowNode_v2::Evaluate(UE::Dataflow::FContext& Context, const FDataflowOutput* Out) const
+{
+	ensure(Out->IsA(&Collection) || Out->IsA(&Materials) || Out->IsA(&InstancedMeshes));
+
+	FManagedArrayCollection OutCollection;
+	TArray<TObjectPtr<UMaterialInterface>> OutMaterials;
+	TArray<FGeometryCollectionAutoInstanceMesh> OutInstancedMeshes;
+
+	TObjectPtr<UStaticMesh> StaticMeshVal = GetValue(Context, &StaticMesh, StaticMesh);
+	const FTransform& InMeshTransform = GetValue(Context, &MeshTransform, MeshTransform);
+	if (StaticMeshVal)
+	{
+		FGeometryCollectionEngineConversion::ConvertStaticMeshToGeometryCollection(StaticMeshVal, InMeshTransform, OutCollection, OutMaterials, OutInstancedMeshes, bSetInternalFromMaterialIndex, bSplitComponents);
+	}
+
+	// Set Outputs
+	SetValue(Context, MoveTemp(OutCollection), &Collection);
+	SetValue(Context, MoveTemp(OutMaterials), &Materials);
+	SetValue(Context, MoveTemp(OutInstancedMeshes), &InstancedMeshes);
+}
 
-// ===========================================================================================================================
 
+////////////////////////////////////////////////////////////////////////////////////////////////////////////////////////////////////////////
+//
+// Deprecated version 
 
 FStaticMeshToCollectionDataflowNode::FStaticMeshToCollectionDataflowNode(const UE::Dataflow::FNodeParameters& InParam, FGuid InGuid)
 	: FDataflowNode(InParam, InGuid)
diff --git a/Engine/Plugins/Experimental/GeometryCollectionPlugin/Source/GeometryCollectionNodes/Public/Dataflow/GeometryCollectionAssetNodes.h b/Engine/Plugins/Experimental/GeometryCollectionPlugin/Source/GeometryCollectionNodes/Public/Dataflow/GeometryCollectionAssetNodes.h
--- a/Engine/Plugins/Experimental/GeometryCollectionPlugin/Source/GeometryCollectionNodes/Public/Dataflow/GeometryCollectionAssetNodes.h
+++ b/Engine/Plugins/Experimental/GeometryCollectionPlugin/Source/GeometryCollectionNodes/Public/Dataflow/GeometryCollectionAssetNodes.h
@@ -19,7 +19,40 @@ namespace UE_DEPRECATED(5.5, "Use UE::Dataflow instead.") Dataflow {}
 class UGeometryCollection;
 class UBlueprint;
 
+/*
+* Geometry Collection asset terminal node
+*/
 USTRUCT(meta = (DataflowGeometryCollection, DataflowTerminal))
+struct FGeometryCollectionTerminalDataflowNode_v2 : public FDataflowTerminalNode
+{
+	GENERATED_USTRUCT_BODY()
+	DATAFLOW_NODE_DEFINE_INTERNAL(FGeometryCollectionTerminalDataflowNode_v2, "GeometryCollectionTerminal", "Terminal", "")
+		DATAFLOW_NODE_RENDER_TYPE("SurfaceRender", FGeometryCollection::StaticType(), "Collection")
+
+private:
+	UPROPERTY(meta = (DataflowInput, DataflowOutput, DataflowPassthrough = "Collection"))
+	FManagedArrayCollection Collection;
+
+	/** Materials to set on this asset */
+	UPROPERTY(meta = (DataflowInput, DataflowOutput, DataflowPassthrough = "Materials"))
+	TArray<TObjectPtr<UMaterialInterface>> Materials;
+
+	/** array of instanced meshes*/
+	UPROPERTY(meta = (DataflowInput, DataflowOutput, DataflowPassthrough = "InstancedMeshes"))
+	TArray<FGeometryCollectionAutoInstanceMesh> InstancedMeshes;
+
+	virtual void Evaluate(UE::Dataflow::FContext& Context) const override;
+	virtual void SetAssetValue(TObjectPtr<UObject> Asset, UE::Dataflow::FContext& Context) const override;
+
+public:
+	FGeometryCollectionTerminalDataflowNode_v2(const UE::Dataflow::FNodeParameters& InParam, FGuid InGuid = FGuid::NewGuid());
+};
+
+/*
+* Geometry Collection asset terminal node
+* Deprecated (5.6) - Use version 2 of the same node that only support material interface array as materials input 
+*/
+USTRUCT(meta = (DataflowTerminal, Deprecated = "5.6"))
 struct FGeometryCollectionTerminalDataflowNode : public FDataflowTerminalNode
 {
 	GENERATED_USTRUCT_BODY()
diff --git a/Engine/Plugins/Experimental/GeometryCollectionPlugin/Source/GeometryCollectionNodes/Public/Dataflow/GeometryCollectionStaticMeshToCollectionNode.h b/Engine/Plugins/Experimental/GeometryCollectionPlugin/Source/GeometryCollectionNodes/Public/Dataflow/GeometryCollectionStaticMeshToCollectionNode.h
--- a/Engine/Plugins/Experimental/GeometryCollectionPlugin/Source/GeometryCollectionNodes/Public/Dataflow/GeometryCollectionStaticMeshToCollectionNode.h
+++ b/Engine/Plugins/Experimental/GeometryCollectionPlugin/Source/GeometryCollectionNodes/Public/Dataflow/GeometryCollectionStaticMeshToCollectionNode.h
@@ -21,12 +21,58 @@ class UMaterial;
 /**
  * Create a geometry collection from a UStaticMesh
  */
-USTRUCT(meta = (DataflowContext = "GeometryCollection", DataflowGeometryCollection, DataflowTerminal))
+USTRUCT()
+struct FStaticMeshToCollectionDataflowNode_v2 : public FDataflowNode
+{
+	GENERATED_USTRUCT_BODY()
+	DATAFLOW_NODE_DEFINE_INTERNAL(FStaticMeshToCollectionDataflowNode_v2, "StaticMeshToCollection", "GeometryCollection|Asset", "")
+	DATAFLOW_NODE_RENDER_TYPE("SurfaceRender",FGeometryCollection::StaticType(),  "Collection")
+
+private:
+	/** Asset input */
+	UPROPERTY(EditAnywhere, Category = "Asset", meta = (DataflowInput));
+	TObjectPtr<UStaticMesh> StaticMesh;
+
+	/** Transform to apply to the mesh before converting it to a collection */
+	UPROPERTY(EditAnywhere, Category = "Asset", meta = (DataflowInput));
+	FTransform MeshTransform = FTransform::Identity;
+
+	/** Geometry collection newly created */
+	UPROPERTY(meta = (DataflowOutput))
+	FManagedArrayCollection Collection;
+
+	/** Material array from the static mesh */
+	UPROPERTY(meta = (DataflowOutput))
+	TArray<TObjectPtr<UMaterialInterface>> Materials;
+
+	/** Array of instanced meshes*/
+	UPROPERTY(meta = (DataflowOutput))
+	TArray<FGeometryCollectionAutoInstanceMesh> InstancedMeshes;
+
+	/** Set the internal faces from material index */
+	UPROPERTY(EditAnywhere, Category = "Asset");
+	bool bSetInternalFromMaterialIndex = true;
+
+	/** Split components - when enabled, each island of the mesh will be converted to an individual transform in the collection */
+	UPROPERTY(EditAnywhere, Category = "Asset");
+	bool bSplitComponents = false;
+
+	virtual void Evaluate(UE::Dataflow::FContext& Context, const FDataflowOutput* Out) const override;
+
+public:
+	FStaticMeshToCollectionDataflowNode_v2(const UE::Dataflow::FNodeParameters& InParam, FGuid InGuid = FGuid::NewGuid());
+};
+
+/**
+ * Create a geometry collection from a UStaticMesh
+ * Deprecated (5.6) - use the second version of StaticMeshToCollection node with only one materials ouput
+ */
+USTRUCT(meta = (Deprecated = "5.6"))
 struct FStaticMeshToCollectionDataflowNode : public FDataflowNode
 {
 	GENERATED_USTRUCT_BODY()
 	DATAFLOW_NODE_DEFINE_INTERNAL(FStaticMeshToCollectionDataflowNode, "StaticMeshToCollection", "GeometryCollection|Asset", "")
-	DATAFLOW_NODE_RENDER_TYPE("SurfaceRender",FGeometryCollection::StaticType(),  "Collection")
+		DATAFLOW_NODE_RENDER_TYPE("SurfaceRender", FGeometryCollection::StaticType(), "Collection")
 
 public:
 	/** Asset input */
@@ -64,3 +110,8 @@ public:
 	FStaticMeshToCollectionDataflowNode(const UE::Dataflow::FNodeParameters& InParam, FGuid InGuid = FGuid::NewGuid());
 	virtual void Evaluate(UE::Dataflow::FContext& Context, const FDataflowOutput* Out) const override;
 };
+
+namespace UE::Dataflow
+{
+	void RegisterGeometryStaticMeshToCollectionNodes();
+}
