Dataflow : texture Bake node improvements
- Check for UVChannel validity and output an log error message if not valid
- Expose all baking parameters as outputs ( hidden by default )

diff --git a/Engine/Plugins/Experimental/GeometryCollectionPlugin/Source/GeometryCollectionNodes/Private/Dataflow/GeometryCollectionTextureNodes.cpp b/Engine/Plugins/Experimental/GeometryCollectionPlugin/Source/GeometryCollectionNodes/Private/Dataflow/GeometryCollectionTextureNodes.cpp
--- a/Engine/Plugins/Experimental/GeometryCollectionPlugin/Source/GeometryCollectionNodes/Private/Dataflow/GeometryCollectionTextureNodes.cpp
+++ b/Engine/Plugins/Experimental/GeometryCollectionPlugin/Source/GeometryCollectionNodes/Private/Dataflow/GeometryCollectionTextureNodes.cpp
@@ -2,6 +2,7 @@
 
 #include "Dataflow/GeometryCollectionTextureNodes.h"
 #include "GeometryCollection/GeometryCollectionObject.h"
+#include "ChaosLog.h"
 
 namespace UE::Dataflow
 {
@@ -25,6 +26,16 @@ FBakeTextureFromCollectionDataflowNode::FBakeTextureFromCollectionDataflowNode(c
 	RegisterInputConnection(&GutterSize);
 	RegisterInputConnection(&Resolution);
 
+	// hidden by default inputs
+	RegisterInputConnection(&MaxDistance).SetCanHidePin(true).SetPinIsHidden(true);
+	RegisterInputConnection(&OcclusionRays).SetCanHidePin(true).SetPinIsHidden(true);
+	RegisterInputConnection(&OcclusionBlurRadius).SetCanHidePin(true).SetPinIsHidden(true);
+	RegisterInputConnection(&CurvatureBlurRadius).SetCanHidePin(true).SetPinIsHidden(true);
+	RegisterInputConnection(&VoxelResolution).SetCanHidePin(true).SetPinIsHidden(true);
+	RegisterInputConnection(&SmoothingIterations).SetCanHidePin(true).SetPinIsHidden(true);
+	RegisterInputConnection(&ThicknessFactor).SetCanHidePin(true).SetPinIsHidden(true);
+	RegisterInputConnection(&MaxCurvature).SetCanHidePin(true).SetPinIsHidden(true);
+	
 	RegisterOutputConnection(&Collection, &Collection);
 	RegisterOutputConnection(&Image);
 }
@@ -45,17 +56,31 @@ void FBakeTextureFromCollectionDataflowNode::Evaluate(UE::Dataflow::FContext& Co
 		TUniquePtr<FGeometryCollection> GeomCollection = TUniquePtr<FGeometryCollection>(InCollection.NewCopy<FGeometryCollection>());
 		check(GeomCollection.IsValid());
 
+		const bool bIsValidUVChannel = (InUVChannel >= 0 && InUVChannel < GeomCollection->NumUVLayers());
+		if (!bIsValidUVChannel)
+		{
+			// log until we have proper error reporting in context 
+			UE_LOG(LogChaosDataflow, Error, TEXT("Dataflow Node [%s] : Invalid UV channel [%d] - The collection has %d UV channels"), *GetName().ToString(), InUVChannel, GeomCollection->NumUVLayers());
+
+			// Write out a red 16x16 image as a sign of error
+			FDataflowImage DefaultImage;
+			DefaultImage.CreateFromColor(EDataflowImageResolution::Resolution16, FLinearColor::Red);
+			SetValue(Context, MoveTemp(DefaultImage), &Image);
+
+			return;
+		}
+
 		const UE::Geometry::FIndex4i BakeAttributes{ (int32)RedChannel, (int32)GreenChannel, (int32)BlueChannel, (int32)AlphaChannel };
 
 		UE::PlanarCut::FTextureAttributeSettings BakeAttributeSettings;
-		BakeAttributeSettings.ToExternal_MaxDistance = MaxDistance;
-		BakeAttributeSettings.AO_Rays = OcclusionRays;
-		BakeAttributeSettings.AO_BlurRadius = OcclusionBlurRadius;
-		BakeAttributeSettings.Curvature_BlurRadius = CurvatureBlurRadius;
-		BakeAttributeSettings.Curvature_SmoothingSteps = SmoothingIterations;
-		BakeAttributeSettings.Curvature_VoxelRes = VoxelResolution;
-		BakeAttributeSettings.Curvature_ThicknessFactor = ThicknessFactor;
-		BakeAttributeSettings.Curvature_MaxValue = MaxCurvature;
+		BakeAttributeSettings.ToExternal_MaxDistance = GetValue(Context, &MaxDistance);;
+		BakeAttributeSettings.AO_Rays = GetValue(Context, &OcclusionRays);
+		BakeAttributeSettings.AO_BlurRadius = GetValue(Context, &OcclusionBlurRadius);
+		BakeAttributeSettings.Curvature_BlurRadius = GetValue(Context, &CurvatureBlurRadius);
+		BakeAttributeSettings.Curvature_SmoothingSteps = GetValue(Context, &SmoothingIterations);
+		BakeAttributeSettings.Curvature_VoxelRes = GetValue(Context, &VoxelResolution);
+		BakeAttributeSettings.Curvature_ThicknessFactor = GetValue(Context, &ThicknessFactor);
+		BakeAttributeSettings.Curvature_MaxValue = GetValue(Context, &MaxCurvature);
 		BakeAttributeSettings.ClearGutterChannel = 3; // default clear the gutters for the alpha channel, so it shows more clearly the island boundaries
 
 		UE::Geometry::TImageBuilder<FVector4f> ImageBuilder;
diff --git a/Engine/Plugins/Experimental/GeometryCollectionPlugin/Source/GeometryCollectionNodes/Public/Dataflow/GeometryCollectionTextureNodes.h b/Engine/Plugins/Experimental/GeometryCollectionPlugin/Source/GeometryCollectionNodes/Public/Dataflow/GeometryCollectionTextureNodes.h
--- a/Engine/Plugins/Experimental/GeometryCollectionPlugin/Source/GeometryCollectionNodes/Public/Dataflow/GeometryCollectionTextureNodes.h
+++ b/Engine/Plugins/Experimental/GeometryCollectionPlugin/Source/GeometryCollectionNodes/Public/Dataflow/GeometryCollectionTextureNodes.h
@@ -99,51 +99,59 @@ private:
 
 	/** Max distance to search for the outer mesh surface */
 	UPROPERTY(EditAnywhere, Category = DistToOuterSettings, meta = (
+		DataflowInput,
 		EditCondition = "RedChannel == ECollectionBakeTextureAttribute::DistanceToExternal || GreenChannel == ECollectionBakeTextureAttribute::DistanceToExternal || BlueChannel == ECollectionBakeTextureAttribute::DistanceToExternal || AlphaChannel == ECollectionBakeTextureAttribute::DistanceToExternal", 
 		UIMin = "1", UIMax = "100", ClampMin = ".01", ClampMax = "1000"))
-	double MaxDistance = 100;
+	float MaxDistance = 100;
 
 	/** Number of occlusion rays */
 	UPROPERTY(EditAnywhere, Category = AmbientOcclusionSettings, meta = (
+		DataflowInput,
 		EditCondition = "RedChannel == ECollectionBakeTextureAttribute::AmbientOcclusion || GreenChannel == ECollectionBakeTextureAttribute::AmbientOcclusion || BlueChannel == ECollectionBakeTextureAttribute::AmbientOcclusion || AlphaChannel == ECollectionBakeTextureAttribute::AmbientOcclusion",
 		UIMin = "1", UIMax = "1024", ClampMin = "0", ClampMax = "50000"))
-	int OcclusionRays = 16;
+	int32 OcclusionRays = 16;
 
 	/** Pixel Radius of Gaussian Blur Kernel applied to AO map (0 will apply no blur) */
 	UPROPERTY(EditAnywhere, Category = AmbientOcclusionSettings, meta = (
+		DataflowInput,
 		EditCondition = "RedChannel == ECollectionBakeTextureAttribute::AmbientOcclusion || GreenChannel == ECollectionBakeTextureAttribute::AmbientOcclusion || BlueChannel == ECollectionBakeTextureAttribute::AmbientOcclusion || AlphaChannel == ECollectionBakeTextureAttribute::AmbientOcclusion",
 		UIMin = "0", UIMax = "10.0", ClampMin = "0", ClampMax = "100.0"))
-	double OcclusionBlurRadius = 2.25;
+	float OcclusionBlurRadius = 2.25;
 
 	/** Pixel Radius of Gaussian Blur Kernel applied to Curvature map (0 will apply no blur) */
 	UPROPERTY(EditAnywhere, Category = SmoothedCurvatureSettings, meta = (
+		DataflowInput,
 		EditCondition = "RedChannel == ECollectionBakeTextureAttribute::Curvature || GreenChannel == ECollectionBakeTextureAttribute::Curvature || BlueChannel == ECollectionBakeTextureAttribute::Curvature || AlphaChannel == ECollectionBakeTextureAttribute::Curvature",
 		UIMin = "0", UIMax = "10.0", ClampMin = "0", ClampMax = "100.0"))
-	double CurvatureBlurRadius = 2.25;
+	float CurvatureBlurRadius = 2.25;
 
 	/** Voxel resolution of smoothed shape representation */
 	UPROPERTY(EditAnywhere, Category = SmoothedCurvatureSettings, meta = (
+		DataflowInput,
 		EditCondition = "RedChannel == ECollectionBakeTextureAttribute::Curvature || GreenChannel == ECollectionBakeTextureAttribute::Curvature || BlueChannel == ECollectionBakeTextureAttribute::Curvature || AlphaChannel == ECollectionBakeTextureAttribute::Curvature",
 		UIMin = "8", UIMax = "512", ClampMin = "4", ClampMax = "1024"))
-	int VoxelResolution = 128;
+	int32 VoxelResolution = 128;
 
 	/** Amount of smoothing iterations to apply before computing curvature */
 	UPROPERTY(EditAnywhere, Category = SmoothedCurvatureSettings, meta = (
+		DataflowInput,
 		EditCondition = "RedChannel == ECollectionBakeTextureAttribute::Curvature || GreenChannel == ECollectionBakeTextureAttribute::Curvature || BlueChannel == ECollectionBakeTextureAttribute::Curvature || AlphaChannel == ECollectionBakeTextureAttribute::Curvature",
 		UIMin = "2", UIMax = "100", ClampMin = "2", ClampMax = "1000"))
-	int SmoothingIterations = 10;
+	int32 SmoothingIterations = 10;
 
 	/** Distance to search for correspondence between fractured shape and smoothed shape, as factor of voxel size */
 	UPROPERTY(EditAnywhere, Category = SmoothedCurvatureSettings, meta = (
+		DataflowInput,
 		EditCondition = "RedChannel == ECollectionBakeTextureAttribute::Curvature || GreenChannel == ECollectionBakeTextureAttribute::Curvature || BlueChannel == ECollectionBakeTextureAttribute::Curvature || AlphaChannel == ECollectionBakeTextureAttribute::Curvature",
 		UIMin = "2", UIMax = "10.0", ClampMin = "1", ClampMax = "100.0"))
-	double ThicknessFactor = 4;
+	float ThicknessFactor = 4;
 
 	/** Curvatures in the range [-MaxCurvature, MaxCurvature] will be mapped from [0,1]. Values outside that range will be clamped */
 	UPROPERTY(EditAnywhere, Category = SmoothedCurvatureSettings, meta = (
+		DataflowInput,
 		EditCondition = "RedChannel == ECollectionBakeTextureAttribute::Curvature || GreenChannel == ECollectionBakeTextureAttribute::Curvature || BlueChannel == ECollectionBakeTextureAttribute::Curvature || AlphaChannel == ECollectionBakeTextureAttribute::Curvature",
 		UIMin = ".01", UIMax = "1", ClampMin = ".0001", ClampMax = "10"))
-	double MaxCurvature = .1;
+	float MaxCurvature = .1;
 
 	virtual void Evaluate(UE::Dataflow::FContext& Context, const FDataflowOutput* Out) const override;
 
