Dataflow : Add boxProject UV node

diff --git a/Engine/Plugins/Experimental/GeometryCollectionPlugin/Source/GeometryCollectionNodes/Private/Dataflow/GeometryCollectionUVNodes.cpp b/Engine/Plugins/Experimental/GeometryCollectionPlugin/Source/GeometryCollectionNodes/Private/Dataflow/GeometryCollectionUVNodes.cpp
--- a/Engine/Plugins/Experimental/GeometryCollectionPlugin/Source/GeometryCollectionNodes/Private/Dataflow/GeometryCollectionUVNodes.cpp
+++ b/Engine/Plugins/Experimental/GeometryCollectionPlugin/Source/GeometryCollectionNodes/Private/Dataflow/GeometryCollectionUVNodes.cpp
@@ -11,6 +11,7 @@ namespace UE::Dataflow
 	{
 		DATAFLOW_NODE_REGISTER_CREATION_FACTORY(FAddUVChannelDataflowNode);
 		DATAFLOW_NODE_REGISTER_CREATION_FACTORY(FAutoUnwrapUVDataflowNode);
+		DATAFLOW_NODE_REGISTER_CREATION_FACTORY(FBoxProjectUVDataflowNode);
 
 		static const FLinearColor CDefaultNodeBodyTintColor = FLinearColor(0.f, 0.f, 0.f, 0.5f);
 		DATAFLOW_NODE_REGISTER_CREATION_FACTORY_NODE_COLORS_BY_CATEGORY("GeometryCollection|UV", FLinearColor(.1f, 1.f, 0.6f), CDefaultNodeBodyTintColor);
@@ -70,6 +71,7 @@ FAutoUnwrapUVDataflowNode::FAutoUnwrapUVDataflowNode(const UE::Dataflow::FNodePa
 	RegisterInputConnection(&Collection);
 	RegisterInputConnection(&FaceSelection);
 	RegisterInputConnection(&UVChannel);
+	RegisterInputConnection(&GutterSize).SetCanHidePin(true).SetPinIsHidden(true);
 	RegisterOutputConnection(&Collection, &Collection);
 }
 
@@ -87,6 +89,7 @@ void FAutoUnwrapUVDataflowNode::Evaluate(UE::Dataflow::FContext& Context, const
 		}
 
 		const int32 InUVChannel = GetValue(Context, &UVChannel);
+		const int32 InGutterSize = FMath::Max(1, GetValue(Context, &GutterSize));
 
 		TUniquePtr<FGeometryCollection> GeomCollection = TUniquePtr<FGeometryCollection>(InCollection.NewCopy<FGeometryCollection>());
 		check(GeomCollection.IsValid());
@@ -107,14 +110,14 @@ void FAutoUnwrapUVDataflowNode::Evaluate(UE::Dataflow::FContext& Context, const
 					FacesToUnwrap[Index] = InFaceSelection.IsSelected(Index);
 				}
 
-				UE::PlanarCut::UVLayout(InUVChannel, *GeomCollection, ReferenceResolution, GutterSize,
+				UE::PlanarCut::UVLayout(InUVChannel, *GeomCollection, ReferenceResolution, InGutterSize,
 					FacesToUnwrap, bRecreateUVsForDegenerateIslands, nullptr /* progress */);
 			}
 		}
 		else
 		{
 			// All faces
-			UE::PlanarCut::UVLayout(InUVChannel, *GeomCollection, ReferenceResolution, GutterSize,
+			UE::PlanarCut::UVLayout(InUVChannel, *GeomCollection, ReferenceResolution, InGutterSize,
 				UE::PlanarCut::ETargetFaces::AllFaces, {}, bRecreateUVsForDegenerateIslands, nullptr /* progress */);
 		}
 		
@@ -122,3 +125,81 @@ void FAutoUnwrapUVDataflowNode::Evaluate(UE::Dataflow::FContext& Context, const
 	}
 }
 
+////////////////////////////////////////////////////////////////////////////////////////////
+
+FBoxProjectUVDataflowNode::FBoxProjectUVDataflowNode(const UE::Dataflow::FNodeParameters& InParam, FGuid InGuid)
+	: FDataflowNode(InParam, InGuid)
+{
+	RegisterInputConnection(&Collection);
+	// RegisterInputConnection(&FaceSelection); not supported yet
+	RegisterInputConnection(&UVChannel);
+	RegisterInputConnection(&ProjectionScale);
+	RegisterInputConnection(&UVOffset);
+	RegisterInputConnection(&bAutoFitToBounds);
+	RegisterInputConnection(&bCenterBoxAtPivot);
+	RegisterInputConnection(&bUniformProjectionScale);
+	RegisterInputConnection(&GutterSize).SetCanHidePin(true).SetPinIsHidden(true);
+	RegisterOutputConnection(&Collection, &Collection);
+}
+
+void FBoxProjectUVDataflowNode::Evaluate(UE::Dataflow::FContext& Context, const FDataflowOutput* Out) const
+{
+	if (Out->IsA(&Collection))
+	{
+		const FManagedArrayCollection& InCollection = GetValue(Context, &Collection);
+		const int32 FaceCount = InCollection.NumElements(FGeometryCollection::FacesGroup);
+		if (FaceCount == 0)
+		{
+			// nothing to do : forward and exit 
+			SafeForwardInput(Context, &Collection, &Collection);
+			return;
+		}
+
+		const int32 InUVChannel = GetValue(Context, &UVChannel);
+		const int32 InGutterSize = FMath::Max(1, GetValue(Context, &GutterSize));
+		const FVector InProjectionScale = GetValue(Context, &ProjectionScale);
+		const FVector2D InUVOffset = GetValue(Context, &UVOffset);
+		const bool bInAutoFitToBounds = GetValue(Context, &bAutoFitToBounds);
+		const bool bInCenterBoxAtPivot = GetValue(Context, &bCenterBoxAtPivot);
+		const bool bInUniformProjectionScale = GetValue(Context, &bUniformProjectionScale);
+
+		TUniquePtr<FGeometryCollection> GeomCollection = TUniquePtr<FGeometryCollection>(InCollection.NewCopy<FGeometryCollection>());
+		check(GeomCollection.IsValid());
+
+		// this is used as reference for the gutter size that is expressed in pixels for a specific resolution 
+		constexpr int32 ReferenceResolution = 512;
+		constexpr bool bRecreateUVsForDegenerateIslands = true;
+
+		//if (IsConnected(&FaceSelection))
+		//{
+		//	const FDataflowFaceSelection& InFaceSelection = GetValue(Context, &FaceSelection);
+		//	if (InFaceSelection.Num() > 0)
+		//	{
+		//		TArray<bool> FacesToUnwrap;
+		//		FacesToUnwrap.SetNumUninitialized(InFaceSelection.Num());
+		//		for (int32 Index = 0; Index < InFaceSelection.Num(); ++Index)
+		//		{
+		//			FacesToUnwrap[Index] = InFaceSelection.IsSelected(Index);
+		//		}
+
+		//		UE::PlanarCut::UVLayout(InUVChannel, *GeomCollection, ReferenceResolution, InGutterSize,
+		//			FacesToUnwrap, bRecreateUVsForDegenerateIslands, nullptr /* progress */);
+		//	}
+		//}
+		//else
+		{
+			// All faces
+			UE::PlanarCut::BoxProjectUVs(
+				InUVChannel, *GeomCollection, 
+				InProjectionScale, 
+				UE::PlanarCut::ETargetFaces::AllFaces, {}, 
+				(FVector2f)InUVOffset, 
+				bAutoFitToBounds,
+				bCenterBoxAtPivot,
+				bUniformProjectionScale
+				);
+		}
+
+		SetValue(Context, (const FManagedArrayCollection&)(*GeomCollection), &Collection);
+	}
+}
diff --git a/Engine/Plugins/Experimental/GeometryCollectionPlugin/Source/GeometryCollectionNodes/Public/Dataflow/GeometryCollectionUVNodes.h b/Engine/Plugins/Experimental/GeometryCollectionPlugin/Source/GeometryCollectionNodes/Public/Dataflow/GeometryCollectionUVNodes.h
--- a/Engine/Plugins/Experimental/GeometryCollectionPlugin/Source/GeometryCollectionNodes/Public/Dataflow/GeometryCollectionUVNodes.h
+++ b/Engine/Plugins/Experimental/GeometryCollectionPlugin/Source/GeometryCollectionNodes/Public/Dataflow/GeometryCollectionUVNodes.h
@@ -60,7 +60,7 @@ private:
 	int32 UVChannel = 0;
 
 	/** Approximate space to leave between UV islands, measured in texels for 512x512 texture */
-	UPROPERTY(EditAnywhere, Category = Options, meta = (ClampMin = 1, UIMin = 1))
+	UPROPERTY(EditAnywhere, Category = Options, meta = (DataflowInput, ClampMin = 1, UIMin = 1))
 	int32 GutterSize = 1;
 
 	virtual void Evaluate(UE::Dataflow::FContext& Context, const FDataflowOutput* Out) const override;
@@ -69,6 +69,53 @@ public:
 	FAutoUnwrapUVDataflowNode(const UE::Dataflow::FNodeParameters& InParam, FGuid InGuid = FGuid::NewGuid());
 };
 
+/**
+ * Generates UVs using a box projection
+ */
+USTRUCT()
+struct FBoxProjectUVDataflowNode : public FDataflowNode
+{
+	GENERATED_USTRUCT_BODY()
+	DATAFLOW_NODE_DEFINE_INTERNAL(FBoxProjectUVDataflowNode, "BoxProjectUV", "GeometryCollection|UV", "")
+
+private:
+	/** Target Collection */
+	UPROPERTY(meta = (DataflowInput, DataflowOutput, DataflowPassthrough = "Collection", DataflowIntrinsic))
+	FManagedArrayCollection Collection;
+
+	/** Faces to auto unwrap, no selection means all faces */
+	//UPROPERTY(meta = (DataflowInput))
+	//FDataflowFaceSelection FaceSelection;
+
+	/** UV channel to unwrap into ( 0 by default ) */
+	UPROPERTY(EditAnywhere, Category = Options, meta = (DataflowInput, ClampMin = 0, ClampMax = 7, UIMin = 0, UIMax = 7))
+	int32 UVChannel = 0;
+
+	/** Approximate space to leave between UV islands, measured in texels for 512x512 texture */
+	UPROPERTY(EditAnywhere, Category = Options, meta = (DataflowInput, ClampMin = 1, UIMin = 1))
+	int32 GutterSize = 1;
+
+	UPROPERTY(EditAnywhere, Category = Options, meta = (DataflowInput))
+	FVector ProjectionScale;
+
+	UPROPERTY(EditAnywhere, Category = Options, meta = (DataflowInput))
+	FVector2D UVOffset;
+
+	UPROPERTY(EditAnywhere, Category = Options, meta = (DataflowInput))
+	bool bAutoFitToBounds;
+
+	UPROPERTY(EditAnywhere, Category = Options, meta = (DataflowInput))
+	bool bCenterBoxAtPivot;
+
+	UPROPERTY(EditAnywhere, Category = Options, meta = (DataflowInput))
+	bool bUniformProjectionScale;
+
+	virtual void Evaluate(UE::Dataflow::FContext& Context, const FDataflowOutput* Out) const override;
+
+public:
+	FBoxProjectUVDataflowNode(const UE::Dataflow::FNodeParameters& InParam, FGuid InGuid = FGuid::NewGuid());
+};
+
 namespace UE::Dataflow
 {
 	void RegisterGeometryCollectionUVNodes();
