Dataflow : add propert default values to box project UV node

diff --git a/Engine/Plugins/Experimental/GeometryCollectionPlugin/Source/GeometryCollectionNodes/Private/Dataflow/GeometryCollectionUVNodes.cpp b/Engine/Plugins/Experimental/GeometryCollectionPlugin/Source/GeometryCollectionNodes/Private/Dataflow/GeometryCollectionUVNodes.cpp
--- a/Engine/Plugins/Experimental/GeometryCollectionPlugin/Source/GeometryCollectionNodes/Private/Dataflow/GeometryCollectionUVNodes.cpp
+++ b/Engine/Plugins/Experimental/GeometryCollectionPlugin/Source/GeometryCollectionNodes/Private/Dataflow/GeometryCollectionUVNodes.cpp
@@ -158,7 +158,7 @@ void FBoxProjectUVDataflowNode::Evaluate(UE::Dataflow::FContext& Context, const
 		const int32 InUVChannel = GetValue(Context, &UVChannel);
 		const int32 InGutterSize = FMath::Max(1, GetValue(Context, &GutterSize));
 		const FVector InProjectionScale = GetValue(Context, &ProjectionScale);
-		const FVector2D InUVOffset = GetValue(Context, &UVOffset);
+		const FVector2f InUVOffset = GetValue(Context, &UVOffset);
 		const bool bInAutoFitToBounds = GetValue(Context, &bAutoFitToBounds);
 		const bool bInCenterBoxAtPivot = GetValue(Context, &bCenterBoxAtPivot);
 		const bool bInUniformProjectionScale = GetValue(Context, &bUniformProjectionScale);
@@ -193,7 +193,7 @@ void FBoxProjectUVDataflowNode::Evaluate(UE::Dataflow::FContext& Context, const
 				InUVChannel, *GeomCollection, 
 				InProjectionScale, 
 				UE::PlanarCut::ETargetFaces::AllFaces, {}, 
-				(FVector2f)InUVOffset, 
+				InUVOffset, 
 				bAutoFitToBounds,
 				bCenterBoxAtPivot,
 				bUniformProjectionScale
diff --git a/Engine/Plugins/Experimental/GeometryCollectionPlugin/Source/GeometryCollectionNodes/Public/Dataflow/GeometryCollectionUVNodes.h b/Engine/Plugins/Experimental/GeometryCollectionPlugin/Source/GeometryCollectionNodes/Public/Dataflow/GeometryCollectionUVNodes.h
--- a/Engine/Plugins/Experimental/GeometryCollectionPlugin/Source/GeometryCollectionNodes/Public/Dataflow/GeometryCollectionUVNodes.h
+++ b/Engine/Plugins/Experimental/GeometryCollectionPlugin/Source/GeometryCollectionNodes/Public/Dataflow/GeometryCollectionUVNodes.h
@@ -96,19 +96,19 @@ private:
 	int32 GutterSize = 1;
 
 	UPROPERTY(EditAnywhere, Category = Options, meta = (DataflowInput))
-	FVector ProjectionScale;
+	FVector ProjectionScale = { 100, 100, 100 };
 
 	UPROPERTY(EditAnywhere, Category = Options, meta = (DataflowInput))
-	FVector2D UVOffset;
+	FVector2f UVOffset = { 0.5, 0.5 };
 
 	UPROPERTY(EditAnywhere, Category = Options, meta = (DataflowInput))
-	bool bAutoFitToBounds;
+	bool bAutoFitToBounds = false;
 
 	UPROPERTY(EditAnywhere, Category = Options, meta = (DataflowInput))
-	bool bCenterBoxAtPivot;
+	bool bCenterBoxAtPivot = false;
 
 	UPROPERTY(EditAnywhere, Category = Options, meta = (DataflowInput))
-	bool bUniformProjectionScale;
+	bool bUniformProjectionScale = false;
 
 	virtual void Evaluate(UE::Dataflow::FContext& Context, const FDataflowOutput* Out) const override;
 
