Dataflow nodes : fix issues with some of the nodes
* invalid inputs properties
* wrong property condition name
* non-returning evaluation path

diff --git a/Engine/Plugins/Experimental/GeometryCollectionPlugin/Source/GeometryCollectionNodes/Private/Dataflow/GeometryCollectionNodes.cpp b/Engine/Plugins/Experimental/GeometryCollectionPlugin/Source/GeometryCollectionNodes/Private/Dataflow/GeometryCollectionNodes.cpp
--- a/Engine/Plugins/Experimental/GeometryCollectionPlugin/Source/GeometryCollectionNodes/Private/Dataflow/GeometryCollectionNodes.cpp
+++ b/Engine/Plugins/Experimental/GeometryCollectionPlugin/Source/GeometryCollectionNodes/Private/Dataflow/GeometryCollectionNodes.cpp
@@ -541,6 +541,9 @@ void FBranchCollectionDataflowNode::Evaluate(UE::Dataflow::FContext& Context, co
 				return;
 			}
 		}
+
+		// default empty collection 
+		SetValue(Context, FManagedArrayCollection(), &ChosenCollection);
 	}
 }
 
diff --git a/Engine/Plugins/Experimental/GeometryCollectionPlugin/Source/GeometryCollectionNodes/Public/Dataflow/GeometryCollectionMathNodes.h b/Engine/Plugins/Experimental/GeometryCollectionPlugin/Source/GeometryCollectionNodes/Public/Dataflow/GeometryCollectionMathNodes.h
--- a/Engine/Plugins/Experimental/GeometryCollectionPlugin/Source/GeometryCollectionNodes/Public/Dataflow/GeometryCollectionMathNodes.h
+++ b/Engine/Plugins/Experimental/GeometryCollectionPlugin/Source/GeometryCollectionNodes/Public/Dataflow/GeometryCollectionMathNodes.h
@@ -1474,7 +1474,7 @@ public:
 	UPROPERTY(EditAnywhere, Category = "Seed")
 	bool bDeterministic = false;
 
-	UPROPERTY(EditAnywhere, Category = "Seed", meta = (DataflowInput, EditCondition = "Deterministic"))
+	UPROPERTY(EditAnywhere, Category = "Seed", meta = (DataflowInput, EditCondition = "bDeterministic"))
 	float RandomSeed = 0.f;
 
 	UPROPERTY(meta = (DataflowOutput))
@@ -1507,7 +1507,7 @@ public:
 	UPROPERTY(EditAnywhere, Category = "Seed")
 	bool bDeterministic = false;
 
-	UPROPERTY(EditAnywhere, Category = "Seed", meta = (DataflowInput, EditCondition = "Deterministic"))
+	UPROPERTY(EditAnywhere, Category = "Seed", meta = (DataflowInput, EditCondition = "bDeterministic"))
 	float RandomSeed = 0.f;
 
 	UPROPERTY(EditAnywhere, Category = "Random", meta = (DataflowInput))
@@ -1549,7 +1549,7 @@ public:
 	UPROPERTY(EditAnywhere, Category = "Seed")
 	bool bDeterministic = false;
 
-	UPROPERTY(EditAnywhere, Category = "Seed", meta = (DataflowInput, EditCondition = "Deterministic"))
+	UPROPERTY(EditAnywhere, Category = "Seed", meta = (DataflowInput, EditCondition = "bDeterministic"))
 	float RandomSeed = 0.f;
 
 	UPROPERTY(meta = (DataflowOutput))
@@ -1583,7 +1583,7 @@ public:
 	UPROPERTY(EditAnywhere, Category = "Seed")
 	bool bDeterministic = false;
 
-	UPROPERTY(EditAnywhere, Category = "Seed", meta = (DataflowInput, EditCondition = "Deterministic"))
+	UPROPERTY(EditAnywhere, Category = "Seed", meta = (DataflowInput, EditCondition = "bDeterministic"))
 	float RandomSeed = 0.f;
 
 	/** The base "center" direction of the cone */
diff --git a/Engine/Source/Runtime/Experimental/Dataflow/Core/Public/Dataflow/DataflowMathNodes.h b/Engine/Source/Runtime/Experimental/Dataflow/Core/Public/Dataflow/DataflowMathNodes.h
--- a/Engine/Source/Runtime/Experimental/Dataflow/Core/Public/Dataflow/DataflowMathNodes.h
+++ b/Engine/Source/Runtime/Experimental/Dataflow/Core/Public/Dataflow/DataflowMathNodes.h
@@ -222,9 +222,6 @@ struct FDataflowMathNegateNode : public FDataflowMathOneInputOperatorNode
 	GENERATED_USTRUCT_BODY()
 	DATAFLOW_NODE_DEFINE_INTERNAL(FDataflowMathNegateNode, "Negate", DATAFLOW_MATH_NODES_CATEGORY, "")
 
-	UPROPERTY(EditAnywhere, Category = SafeDivide, meta = (DataflowInput));
-	FDataflowNumericTypes Fallback;
-
 public:
 	FDataflowMathNegateNode(const UE::Dataflow::FNodeParameters& InParam, FGuid InGuid = FGuid::NewGuid());
 	virtual double ComputeResult(UE::Dataflow::FContext& Context, double InA) const override;
@@ -237,9 +234,6 @@ struct FDataflowMathAbsNode : public FDataflowMathOneInputOperatorNode
 	GENERATED_USTRUCT_BODY()
 	DATAFLOW_NODE_DEFINE_INTERNAL(FDataflowMathAbsNode, "Abs", DATAFLOW_MATH_NODES_CATEGORY, "")
 
-	UPROPERTY(EditAnywhere, Category = SafeDivide, meta = (DataflowInput));
-	FDataflowNumericTypes Fallback;
-
 public:
 	FDataflowMathAbsNode(const UE::Dataflow::FNodeParameters& InParam, FGuid InGuid = FGuid::NewGuid());
 	virtual double ComputeResult(UE::Dataflow::FContext& Context, double InA) const override;
