Sampling nodes : fix issues where node do not return Set values when unplugged

diff --git a/Engine/Plugins/Experimental/GeometryCollectionPlugin/Source/GeometryCollectionNodes/Private/Dataflow/GeometryCollectionSamplingNodes.cpp b/Engine/Plugins/Experimental/GeometryCollectionPlugin/Source/GeometryCollectionNodes/Private/Dataflow/GeometryCollectionSamplingNodes.cpp
--- a/Engine/Plugins/Experimental/GeometryCollectionPlugin/Source/GeometryCollectionNodes/Private/Dataflow/GeometryCollectionSamplingNodes.cpp
+++ b/Engine/Plugins/Experimental/GeometryCollectionPlugin/Source/GeometryCollectionNodes/Private/Dataflow/GeometryCollectionSamplingNodes.cpp
@@ -21,18 +21,20 @@ void FUniformPointSamplingDataflowNode::Evaluate(UE::Dataflow::FContext& Context
 {
 	if (Out->IsA(&SamplePoints) || 
 		Out->IsA(&SampleTriangleIDs) ||
+		Out->IsA(&SampleBarycentricCoords) ||
 		Out->IsA(&NumSamplePoints))
 	{
+		TArray<FVector> OutPoints;
+		TArray<FTransform> OutSamples;
+		TArray<int32> OutTriangleIDs;
+		TArray<FVector> OutBarycentricCoords;
+
 		if (TObjectPtr<UDynamicMesh> InTargetMesh = GetValue(Context, &TargetMesh))
 		{
 			const UE::Geometry::FDynamicMesh3& InDynTargetMesh = InTargetMesh->GetMeshRef();
 
 			if (InDynTargetMesh.VertexCount() > 0)
 			{
-				TArray<FTransform> OutSamples;
-				TArray<int32> OutTriangleIDs;
-				TArray<FVector> OutBarycentricCoords;
-
 				FFractureEngineSampling::ComputeUniformPointSampling(InDynTargetMesh,
 					GetValue(Context, &SamplingRadius),
 					GetValue(Context, &MaxNumSamples),
@@ -44,20 +46,19 @@ void FUniformPointSamplingDataflowNode::Evaluate(UE::Dataflow::FContext& Context
 
 				const int32 NumSamples = OutSamples.Num();
 
-				TArray<FVector> OutPoints;
 				OutPoints.AddUninitialized(NumSamples);
 
 				for (int32 Idx = 0; Idx < NumSamples; ++Idx)
 				{
 					OutPoints[Idx] = OutSamples[Idx].GetTranslation();
 				}
-
-				SetValue(Context, OutPoints, &SamplePoints);
-				SetValue(Context, OutTriangleIDs, &SampleTriangleIDs);
-				SetValue(Context, OutBarycentricCoords, &SampleBarycentricCoords);
-				SetValue(Context, OutPoints.Num(), &NumSamplePoints);
 			}
 		}
+
+		SetValue(Context, MoveTemp(OutPoints), &SamplePoints);
+		SetValue(Context, MoveTemp(OutTriangleIDs), &SampleTriangleIDs);
+		SetValue(Context, MoveTemp(OutBarycentricCoords), &SampleBarycentricCoords);
+		SetValue(Context, OutPoints.Num(), &NumSamplePoints);
 	}
 }
 
@@ -66,19 +67,21 @@ void FNonUniformPointSamplingDataflowNode::Evaluate(UE::Dataflow::FContext& Cont
 	if (Out->IsA(&SamplePoints) ||
 		Out->IsA(&SampleRadii) ||
 		Out->IsA(&SampleTriangleIDs) ||
+		Out->IsA(&SampleBarycentricCoords) ||
 		Out->IsA(&NumSamplePoints))
 	{
+		TArray<FVector> OutPoints;
+		TArray<FTransform> OutSamples;
+		TArray<float> OutSampleRadii;
+		TArray<int32> OutTriangleIDs;
+		TArray<FVector> OutBarycentricCoords;
+
 		if (TObjectPtr<UDynamicMesh> InTargetMesh = GetValue(Context, &TargetMesh))
 		{
 			const UE::Geometry::FDynamicMesh3& InDynTargetMesh = InTargetMesh->GetMeshRef();
 
 			if (InDynTargetMesh.VertexCount() > 0)
 			{
-				TArray<FTransform> OutSamples;
-				TArray<float> OutSampleRadii;
-				TArray<int32> OutTriangleIDs;
-				TArray<FVector> OutBarycentricCoords;
-
 				FFractureEngineSampling::ComputeNonUniformPointSampling(InDynTargetMesh,
 					GetValue(Context, &SamplingRadius),
 					GetValue(Context, &MaxNumSamples),
@@ -94,21 +97,20 @@ void FNonUniformPointSamplingDataflowNode::Evaluate(UE::Dataflow::FContext& Cont
 
 				const int32 NumSamples = OutSamples.Num();
 
-				TArray<FVector> OutPoints;
 				OutPoints.AddUninitialized(NumSamples);
 
 				for (int32 Idx = 0; Idx < NumSamples; ++Idx)
 				{
 					OutPoints[Idx] = OutSamples[Idx].GetTranslation();
 				}
-
-				SetValue(Context, OutPoints, &SamplePoints);
-				SetValue(Context, OutSampleRadii, &SampleRadii);
-				SetValue(Context, OutTriangleIDs, &SampleTriangleIDs);
-				SetValue(Context, OutBarycentricCoords, &SampleBarycentricCoords);
-				SetValue(Context, OutPoints.Num(), &NumSamplePoints);
 			}
 		}
+
+		SetValue(Context, MoveTemp(OutPoints), &SamplePoints);
+		SetValue(Context, MoveTemp(OutSampleRadii), &SampleRadii);
+		SetValue(Context, MoveTemp(OutTriangleIDs), &SampleTriangleIDs);
+		SetValue(Context, MoveTemp(OutBarycentricCoords), &SampleBarycentricCoords);
+		SetValue(Context, OutPoints.Num(), &NumSamplePoints);
 	}
 }
 
@@ -117,8 +119,15 @@ void FVertexWeightedPointSamplingDataflowNode::Evaluate(UE::Dataflow::FContext&
 	if (Out->IsA(&SamplePoints) ||
 		Out->IsA(&SampleRadii) ||
 		Out->IsA(&SampleTriangleIDs) ||
+		Out->IsA(&SampleBarycentricCoords) ||
 		Out->IsA(&NumSamplePoints))
 	{
+		TArray<FVector> OutPoints;
+		TArray<FTransform> OutSamples;
+		TArray<float> OutSampleRadii;
+		TArray<int32> OutTriangleIDs;
+		TArray<FVector> OutBarycentricCoords;
+
 		if (TObjectPtr<UDynamicMesh> InTargetMesh = GetValue(Context, &TargetMesh))
 		{
 			const UE::Geometry::FDynamicMesh3& InDynTargetMesh = InTargetMesh->GetMeshRef();
@@ -127,11 +136,6 @@ void FVertexWeightedPointSamplingDataflowNode::Evaluate(UE::Dataflow::FContext&
 			{
 				if (IsConnected(&VertexWeights))
 				{
-					TArray<FTransform> OutSamples;
-					TArray<float> OutSampleRadii;
-					TArray<int32> OutTriangleIDs;
-					TArray<FVector> OutBarycentricCoords;
-
 					FFractureEngineSampling::ComputeVertexWeightedPointSampling(InDynTargetMesh,
 						GetValue(Context, &VertexWeights),
 						GetValue(Context, &SamplingRadius),
@@ -150,7 +154,6 @@ void FVertexWeightedPointSamplingDataflowNode::Evaluate(UE::Dataflow::FContext&
 
 					const int32 NumSamples = OutSamples.Num();
 
-					TArray<FVector> OutPoints;
 					OutPoints.AddUninitialized(NumSamples);
 
 					for (int32 Idx = 0; Idx < NumSamples; ++Idx)
@@ -158,14 +161,15 @@ void FVertexWeightedPointSamplingDataflowNode::Evaluate(UE::Dataflow::FContext&
 						OutPoints[Idx] = OutSamples[Idx].GetTranslation();
 					}
 
-					SetValue(Context, OutPoints, &SamplePoints);
-					SetValue(Context, OutSampleRadii, &SampleRadii);
-					SetValue(Context, OutTriangleIDs, &SampleTriangleIDs);
-					SetValue(Context, OutBarycentricCoords, &SampleBarycentricCoords);
-					SetValue(Context, OutPoints.Num(), &NumSamplePoints);
 				}
 			}
 		}
+
+		SetValue(Context, MoveTemp(OutPoints), &SamplePoints);
+		SetValue(Context, MoveTemp(OutSampleRadii), &SampleRadii);
+		SetValue(Context, MoveTemp(OutTriangleIDs), &SampleTriangleIDs);
+		SetValue(Context, MoveTemp(OutBarycentricCoords), &SampleBarycentricCoords);
+		SetValue(Context, OutPoints.Num(), &NumSamplePoints);
 	}
 }
 
