Dataflow : add Min / Max for selecting from attribute

diff --git a/Engine/Plugins/Experimental/GeometryCollectionPlugin/Source/GeometryCollectionNodes/Private/Dataflow/GeometryCollectionSelectionNodes.cpp b/Engine/Plugins/Experimental/GeometryCollectionPlugin/Source/GeometryCollectionNodes/Private/Dataflow/GeometryCollectionSelectionNodes.cpp
--- a/Engine/Plugins/Experimental/GeometryCollectionPlugin/Source/GeometryCollectionNodes/Private/Dataflow/GeometryCollectionSelectionNodes.cpp
+++ b/Engine/Plugins/Experimental/GeometryCollectionPlugin/Source/GeometryCollectionNodes/Private/Dataflow/GeometryCollectionSelectionNodes.cpp
@@ -1241,7 +1241,6 @@ void FCollectionVertexSelectionByPercentageDataflowNode::Evaluate(UE::Dataflow::
 	}
 }
 
-
 void FCollectionVertexSelectionSetOperationDataflowNode::Evaluate(UE::Dataflow::FContext& Context, const FDataflowOutput* Out) const
 {
 	if (Out->IsA<FDataflowVertexSelection>(&VertexSelection))
@@ -1291,23 +1290,64 @@ static void CreateSelectionFromAttr(const FManagedArrayCollection& InCollection,
 	const FManagedArrayCollection::EArrayType ArrayType = InCollection.GetAttributeType(InAttribute, InGroup);
 	const int32 NumElements = InCollection.NumElements(InGroup);
 
+	const bool bIsMinOrMaxOperation = (InOperation == ESelectionByAttrOperation::Maximum) || (InOperation == ESelectionByAttrOperation::Minimum);
+
 	if (ArrayType == FManagedArrayCollection::EArrayType::FFloatType)
 	{
 		const TManagedArray<float>* const Array = InCollection.FindAttributeTyped<float>(InAttribute, InGroup);
-		if (InValue.IsNumeric())
+		if (InValue.IsNumeric() || bIsMinOrMaxOperation)
 		{
-			float FloatValue = FCString::Atof(*InValue);
-
-			for (int32 Idx = 0; Idx < NumElements; ++Idx)
+			if (InOperation == ESelectionByAttrOperation::Maximum)
+			{
+				float MaxValue = TNumericLimits<float>::Lowest();
+				int32 MaxIndex = INDEX_NONE;
+				for (int32 Idx = 0; Idx < NumElements; ++Idx)
+				{
+					const float ValueAtIdx = (*Array)[Idx];
+					if (ValueAtIdx >= MaxValue)
+					{
+						MaxValue = ValueAtIdx;
+						MaxIndex = Idx;
+					}
+				}
+				if (MaxIndex != INDEX_NONE)
+				{
+					OutSelection.SetSelected(MaxIndex);
+				}
+			}
+			else if (InOperation == ESelectionByAttrOperation::Minimum)
+			{
+				float MinValue = TNumericLimits<float>::Max();
+				int32 MinIndex = INDEX_NONE;
+				for (int32 Idx = 0; Idx < NumElements; ++Idx)
+				{
+					const float ValueAtIdx = (*Array)[Idx];
+					if (ValueAtIdx <= MinValue)
+					{
+						MinValue = ValueAtIdx;
+						MinIndex = Idx;
+					}
+				}
+				if (MinIndex != INDEX_NONE)
+				{
+					OutSelection.SetSelected(MinIndex);
+				}
+			}
+			else
 			{
-				if ((InOperation == ESelectionByAttrOperation::Equal && (*Array)[Idx] == FloatValue) ||
-					(InOperation == ESelectionByAttrOperation::NotEqual && (*Array)[Idx] != FloatValue) ||
-					(InOperation == ESelectionByAttrOperation::Greater && (*Array)[Idx] > FloatValue) ||
-					(InOperation == ESelectionByAttrOperation::GreaterOrEqual && (*Array)[Idx] >= FloatValue) ||
-					(InOperation == ESelectionByAttrOperation::Smaller && (*Array)[Idx] < FloatValue) ||
-					(InOperation == ESelectionByAttrOperation::SmallerOrEqual && (*Array)[Idx] <= FloatValue))
+				float FloatValue = FCString::Atof(*InValue);
+
+				for (int32 Idx = 0; Idx < NumElements; ++Idx)
 				{
-					OutSelection.SetSelected(Idx);
+					if ((InOperation == ESelectionByAttrOperation::Equal && (*Array)[Idx] == FloatValue) ||
+						(InOperation == ESelectionByAttrOperation::NotEqual && (*Array)[Idx] != FloatValue) ||
+						(InOperation == ESelectionByAttrOperation::Greater && (*Array)[Idx] > FloatValue) ||
+						(InOperation == ESelectionByAttrOperation::GreaterOrEqual && (*Array)[Idx] >= FloatValue) ||
+						(InOperation == ESelectionByAttrOperation::Smaller && (*Array)[Idx] < FloatValue) ||
+						(InOperation == ESelectionByAttrOperation::SmallerOrEqual && (*Array)[Idx] <= FloatValue))
+					{
+						OutSelection.SetSelected(Idx);
+					}
 				}
 			}
 		}
@@ -1320,20 +1360,59 @@ static void CreateSelectionFromAttr(const FManagedArrayCollection& InCollection,
 	else if (ArrayType == FManagedArrayCollection::EArrayType::FInt32Type)
 	{
 		const TManagedArray<int32>* const Array = InCollection.FindAttributeTyped<int32>(InAttribute, InGroup);
-		if (InValue.IsNumeric())
+		if (InValue.IsNumeric() || bIsMinOrMaxOperation)
 		{
-			float IntValue = FCString::Atoi(*InValue);
-
-			for (int32 Idx = 0; Idx < NumElements; ++Idx)
+			if (InOperation == ESelectionByAttrOperation::Maximum)
+			{
+				int32 MaxValue = TNumericLimits<int32>::Lowest();
+				int32 MaxIndex = INDEX_NONE;
+				for (int32 Idx = 0; Idx < NumElements; ++Idx)
+				{
+					const int32 ValueAtIdx = (*Array)[Idx];
+					if (ValueAtIdx >= MaxValue)
+					{
+						MaxValue = ValueAtIdx;
+						MaxIndex = Idx;
+					}
+				}
+				if (MaxIndex != INDEX_NONE)
+				{
+					OutSelection.SetSelected(MaxIndex);
+				}
+			}
+			else if (InOperation == ESelectionByAttrOperation::Minimum)
+			{
+				int32 MinValue = TNumericLimits<int32>::Max();
+				int32 MinIndex = INDEX_NONE;
+				for (int32 Idx = 0; Idx < NumElements; ++Idx)
+				{
+					const int32 ValueAtIdx = (*Array)[Idx];
+					if (ValueAtIdx <= MinValue)
+					{
+						MinValue = ValueAtIdx;
+						MinIndex = Idx;
+					}
+				}
+				if (MinIndex != INDEX_NONE)
+				{
+					OutSelection.SetSelected(MinIndex);
+				}
+			}
+			else
 			{
-				if ((InOperation == ESelectionByAttrOperation::Equal && (*Array)[Idx] == IntValue) ||
-					(InOperation == ESelectionByAttrOperation::NotEqual && (*Array)[Idx] != IntValue) ||
-					(InOperation == ESelectionByAttrOperation::Greater && (*Array)[Idx] > IntValue) ||
-					(InOperation == ESelectionByAttrOperation::GreaterOrEqual && (*Array)[Idx] >= IntValue) ||
-					(InOperation == ESelectionByAttrOperation::Smaller && (*Array)[Idx] < IntValue) ||
-					(InOperation == ESelectionByAttrOperation::SmallerOrEqual && (*Array)[Idx] <= IntValue))
+				float IntValue = FCString::Atoi(*InValue);
+
+				for (int32 Idx = 0; Idx < NumElements; ++Idx)
 				{
-					OutSelection.SetSelected(Idx);
+					if ((InOperation == ESelectionByAttrOperation::Equal && (*Array)[Idx] == IntValue) ||
+						(InOperation == ESelectionByAttrOperation::NotEqual && (*Array)[Idx] != IntValue) ||
+						(InOperation == ESelectionByAttrOperation::Greater && (*Array)[Idx] > IntValue) ||
+						(InOperation == ESelectionByAttrOperation::GreaterOrEqual && (*Array)[Idx] >= IntValue) ||
+						(InOperation == ESelectionByAttrOperation::Smaller && (*Array)[Idx] < IntValue) ||
+						(InOperation == ESelectionByAttrOperation::SmallerOrEqual && (*Array)[Idx] <= IntValue))
+					{
+						OutSelection.SetSelected(Idx);
+					}
 				}
 			}
 		}
@@ -1388,7 +1467,6 @@ static void CreateSelectionFromAttr(const FManagedArrayCollection& InCollection,
 	}
 }
 
-
 void FCollectionSelectionByAttrDataflowNode::Evaluate(UE::Dataflow::FContext& Context, const FDataflowOutput* Out) const
 {
 	if (Out->IsA<FDataflowVertexSelection>(&VertexSelection) ||
diff --git a/Engine/Plugins/Experimental/GeometryCollectionPlugin/Source/GeometryCollectionNodes/Public/Dataflow/GeometryCollectionSelectionNodes.h b/Engine/Plugins/Experimental/GeometryCollectionPlugin/Source/GeometryCollectionNodes/Public/Dataflow/GeometryCollectionSelectionNodes.h
--- a/Engine/Plugins/Experimental/GeometryCollectionPlugin/Source/GeometryCollectionNodes/Public/Dataflow/GeometryCollectionSelectionNodes.h
+++ b/Engine/Plugins/Experimental/GeometryCollectionPlugin/Source/GeometryCollectionNodes/Public/Dataflow/GeometryCollectionSelectionNodes.h
@@ -1434,7 +1434,11 @@ enum class ESelectionByAttrOperation : uint8
 	/** Select faces which attribute value smaller than specified value */
 	Smaller UMETA(DisplayName = "<"),
 	/** Select faces which attribute value greater than specified value */
-	SmallerOrEqual UMETA(DisplayName = "<=")
+	SmallerOrEqual UMETA(DisplayName = "<="),
+	/** Select faces which attribute value greater than specified value */
+	Maximum UMETA(DisplayName = "Max"),
+	/** Select faces which attribute value greater than specified value */
+	Minimum UMETA(DisplayName = "Min")
 };
 
 /**
