Dataflow:
- opening a GC asset now will offer to create/assign a Dataflow asset

diff --git a/Engine/Plugins/Experimental/GeometryCollectionPlugin/Source/GeometryCollectionEditor/Private/GeometryCollection/AssetDefinition_GeometryCollection.cpp b/Engine/Plugins/Experimental/GeometryCollectionPlugin/Source/GeometryCollectionEditor/Private/GeometryCollection/AssetDefinition_GeometryCollection.cpp
--- a/Engine/Plugins/Experimental/GeometryCollectionPlugin/Source/GeometryCollectionEditor/Private/GeometryCollection/AssetDefinition_GeometryCollection.cpp
+++ b/Engine/Plugins/Experimental/GeometryCollectionPlugin/Source/GeometryCollectionEditor/Private/GeometryCollection/AssetDefinition_GeometryCollection.cpp
@@ -72,6 +72,14 @@ EAssetCommandResult UAssetDefinition_GeometryCollection::OpenAssets(const FAsset
 		// Validate the asset
 		if (UGeometryCollection* const GeometryCollection = Cast<UGeometryCollection>(GeometryCollectionObjects[0]))
 		{
+			if (!FDataflowEditorToolkit::HasDataflowAsset(GeometryCollection))
+			{
+				if (UDataflow* const NewDataflowAsset = Cast<UDataflow>(UE::DataflowAssetDefinitionHelpers::NewOrOpenDataflowAsset(GeometryCollection)))
+				{
+					GeometryCollection->DataflowInstance.SetDataflowAsset(NewDataflowAsset);
+				}
+			}
+
 			if (FDataflowEditorToolkit::HasDataflowAsset(GeometryCollection))
 			{
 				UAssetEditorSubsystem* const AssetEditorSubsystem = GEditor->GetEditorSubsystem<UAssetEditorSubsystem>();
@@ -91,5 +99,4 @@ EAssetCommandResult UAssetDefinition_GeometryCollection::OpenAssets(const FAsset
 	return EAssetCommandResult::Unhandled;
 }
 
-
 #undef LOCTEXT_NAMESPACE
