fix missing include of materials headers for geometry collection static mesh convert node

diff --git a/Engine/Plugins/Experimental/GeometryCollectionPlugin/Source/GeometryCollectionNodes/Private/Dataflow/GeometryCollectionStaticMeshToCollectionNode.cpp b/Engine/Plugins/Experimental/GeometryCollectionPlugin/Source/GeometryCollectionNodes/Private/Dataflow/GeometryCollectionStaticMeshToCollectionNode.cpp
--- a/Engine/Plugins/Experimental/GeometryCollectionPlugin/Source/GeometryCollectionNodes/Private/Dataflow/GeometryCollectionStaticMeshToCollectionNode.cpp
+++ b/Engine/Plugins/Experimental/GeometryCollectionPlugin/Source/GeometryCollectionNodes/Private/Dataflow/GeometryCollectionStaticMeshToCollectionNode.cpp
@@ -4,6 +4,8 @@
 #include "Dataflow/DataflowCore.h"
 
 #include "Engine/StaticMesh.h"
+#include "Materials/MaterialInterface.h"
+#include "Materials/Material.h"
 #include "GeometryCollection/GeometryCollection.h"
 #include "GeometryCollection/GeometryCollectionEngineConversion.h"
 #include "GeometryCollection/GeometryCollectionObject.h"
