Add dataflow nodes for geometry collection selection -> array of meshes and for appending an array of meshes to a geometry collection

diff --git a/Engine/Plugins/Experimental/GeometryCollectionPlugin/Source/GeometryCollectionNodes/Private/Dataflow/GeometryCollectionMeshNodes.cpp b/Engine/Plugins/Experimental/GeometryCollectionPlugin/Source/GeometryCollectionNodes/Private/Dataflow/GeometryCollectionMeshNodes.cpp
--- a/Engine/Plugins/Experimental/GeometryCollectionPlugin/Source/GeometryCollectionNodes/Private/Dataflow/GeometryCollectionMeshNodes.cpp
+++ b/Engine/Plugins/Experimental/GeometryCollectionPlugin/Source/GeometryCollectionNodes/Private/Dataflow/GeometryCollectionMeshNodes.cpp
@@ -5,12 +5,14 @@
 
 #include "Engine/Blueprint.h"
 #include "Engine/StaticMesh.h"
+#include "Materials/MaterialInterface.h"
 #include "GeometryCollection/GeometryCollectionObject.h"
 #include "GeometryCollection/ManagedArrayCollection.h"
 #include "GeometryCollection/GeometryCollection.h"
 #include "GeometryCollection/GeometryCollectionEngineUtility.h"
 #include "GeometryCollection/GeometryCollectionEngineConversion.h"
 #include "GeometryCollection/Facades/CollectionTransformSelectionFacade.h"
+#include "GeometryCollection/Facades/CollectionHierarchyFacade.h"
 #include "Logging/LogMacros.h"
 #include "Templates/SharedPointer.h"
 #include "UObject/UnrealTypePrivate.h"
@@ -54,6 +56,8 @@ namespace UE::Dataflow
 		DATAFLOW_NODE_REGISTER_CREATION_FACTORY(FGetMeshDataDataflowNode);
 		DATAFLOW_NODE_REGISTER_CREATION_FACTORY(FApplyMeshProcessorToMeshDataflowNode);
 		DATAFLOW_NODE_REGISTER_CREATION_FACTORY(FApplyMeshProcessorToGeometryCollectionDataflowNode);
+		DATAFLOW_NODE_REGISTER_CREATION_FACTORY(FCollectionSelectionToMeshesDataflowNode);
+		DATAFLOW_NODE_REGISTER_CREATION_FACTORY(FAppendMeshesToCollectionDataflowNode);
 	}
 }
 
@@ -612,3 +616,192 @@ void FApplyMeshProcessorToGeometryCollectionDataflowNode::Evaluate(UE::Dataflow:
 	}
 }
 
+FAppendMeshesToCollectionDataflowNode::FAppendMeshesToCollectionDataflowNode(const UE::Dataflow::FNodeParameters& InParam, FGuid InGuid)
+	: FDataflowNode(InParam, InGuid)
+{
+	RegisterInputConnection(&Collection);
+	RegisterInputConnection(&Meshes);
+	RegisterInputConnection(&ParentIndex);
+	RegisterOutputConnection(&Collection, &Collection);
+	RegisterOutputConnection(&AddedSelection);
+}
+
+void FAppendMeshesToCollectionDataflowNode::Evaluate(UE::Dataflow::FContext& Context, const FDataflowOutput* Out) const
+{
+	if (Out->IsA<FManagedArrayCollection>(&Collection) ||
+		Out->IsA<FDataflowTransformSelection>(&AddedSelection))
+	{
+		if (!IsConnected(&Collection))
+		{
+			SafeForwardInput(Context, &Collection, &Collection);
+			return;
+		}
+
+		int32 UseParentIndex = GetValue(Context, &ParentIndex);
+		const FManagedArrayCollection& InCollection = GetValue(Context, &Collection);
+		const TArray<TObjectPtr<UDynamicMesh>> InMeshes = GetValue(Context, &Meshes);
+
+		FDataflowTransformSelection NewSelection;
+
+		if (TUniquePtr<FGeometryCollection> GeomCollection = TUniquePtr<FGeometryCollection>(InCollection.NewCopy<FGeometryCollection>()))
+		{
+			bool bModifiedCollection = false;
+			int32 FirstNewTransformIndex = INDEX_NONE;
+			for (const TObjectPtr<UDynamicMesh>& MeshObject : InMeshes)
+			{
+				if (MeshObject)
+				{
+					MeshObject->ProcessMesh([&GeomCollection, &bModifiedCollection, UseParentIndex, &FirstNewTransformIndex](const FDynamicMesh3& Mesh)
+					{
+						UE::Geometry::FGeometryCollectionToDynamicMeshes::FToCollectionOptions Options;
+						Options.NewMeshParentIndex = UseParentIndex;
+						int32 AddedIdx = UE::Geometry::FGeometryCollectionToDynamicMeshes::AppendMeshToCollection(*GeomCollection, Mesh, FTransform::Identity, Options);
+						if (AddedIdx != INDEX_NONE)
+						{
+							if (!bModifiedCollection)
+							{
+								FirstNewTransformIndex = AddedIdx;
+							}
+							bModifiedCollection = true;
+						}
+					});
+				}
+			}
+			if (bModifiedCollection)
+			{
+				NewSelection.Initialize(GeomCollection->Transform.Num(), false);
+				for (int32 Idx = FirstNewTransformIndex; Idx < NewSelection.Num(); ++Idx)
+				{
+					NewSelection.SetSelected(Idx);
+				}
+				SetValue<const FManagedArrayCollection&>(Context, MoveTemp(*GeomCollection), &Collection);
+				SetValue(Context, NewSelection, &AddedSelection);
+				return;
+			}
+		}
+
+		SafeForwardInput(Context, &Collection, &Collection);
+		SetValue(Context, NewSelection, &AddedSelection);
+	}
+}
+
+FCollectionSelectionToMeshesDataflowNode::FCollectionSelectionToMeshesDataflowNode(const UE::Dataflow::FNodeParameters& InParam, FGuid InGuid)
+	: FDataflowNode(InParam, InGuid)
+{
+	RegisterInputConnection(&Collection);
+	RegisterInputConnection(&TransformSelection);
+	RegisterOutputConnection(&Meshes);
+}
+
+void FCollectionSelectionToMeshesDataflowNode::Evaluate(UE::Dataflow::FContext& Context, const FDataflowOutput* Out) const
+{
+	if (Out->IsA(&Meshes))
+	{
+		TArray<TObjectPtr<UDynamicMesh>> NewMeshes;
+		FDataflowTransformSelection InTransformSelection = GetValue(Context, &TransformSelection);
+
+		//
+		// If not connected select everything by default
+		//
+		if (!IsConnected(&TransformSelection))
+		{
+			const FManagedArrayCollection& InCollection = GetValue(Context, &Collection);
+
+			InTransformSelection.Initialize(InCollection.NumElements(FGeometryCollection::TransformGroup), true);
+		}
+
+		if (InTransformSelection.AnySelected())
+		{
+			const FManagedArrayCollection& InCollection = GetValue(Context, &Collection);
+
+			TArray<int32> TransformSelectionArray = InTransformSelection.AsArray();
+
+			GeometryCollection::Facades::FCollectionTransformSelectionFacade SelectionFacade(InCollection);
+			Chaos::Facades::FCollectionHierarchyFacade HierarchyFacade(InCollection);
+			TArray<int32> LeafSelectionArray = TransformSelectionArray;
+			SelectionFacade.ConvertSelectionToRigidNodes(LeafSelectionArray);
+
+			UE::Geometry::FGeometryCollectionToDynamicMeshes CollectionToMeshes;
+			UE::Geometry::FGeometryCollectionToDynamicMeshes::FToMeshOptions ToMeshOptions;
+			ToMeshOptions.bWeldVertices = bWeldVertices;
+			ToMeshOptions.bSaveIsolatedVertices = bPreserveIsolatedVertices;
+			if (CollectionToMeshes.InitFromTransformSelection(InCollection, LeafSelectionArray, ToMeshOptions)
+				&& !CollectionToMeshes.Meshes.IsEmpty())
+			{
+				NewMeshes.Reserve(bConvertSelectionToLeaves ? CollectionToMeshes.Meshes.Num() : TransformSelectionArray.Num());
+				using FMeshInfo = UE::Geometry::FGeometryCollectionToDynamicMeshes::FMeshInfo;
+
+				if (bConvertSelectionToLeaves)
+				{
+					for (FMeshInfo& MeshInfo : CollectionToMeshes.Meshes)
+					{
+						TObjectPtr<UDynamicMesh>& NewMesh = NewMeshes.Add_GetRef(NewObject<UDynamicMesh>());
+						NewMesh->SetMesh(MoveTemp(*MeshInfo.Mesh));
+					}
+				}
+				else
+				{
+					TMap<int32, FMeshInfo*> BoneToMeshInfo;
+					TMap<int32, const FDynamicMesh3*> BoneToMesh;
+					for (FMeshInfo& MeshInfo : CollectionToMeshes.Meshes)
+					{
+						BoneToMeshInfo.Add(MeshInfo.TransformIndex, &MeshInfo);
+					}
+					for (int32 BoneIdx : TransformSelectionArray)
+					{
+						if (BoneToMeshInfo.Contains(BoneIdx))
+						{
+							TObjectPtr<UDynamicMesh>& NewMesh = NewMeshes.Add_GetRef(NewObject<UDynamicMesh>());
+							// move the mesh out of the collection and add the pointer to the BoneToMesh map instead
+							// (in case we also have to make a cluster node using the same mesh)
+							NewMesh->SetMesh(MoveTemp(*BoneToMeshInfo[BoneIdx]->Mesh));
+							BoneToMeshInfo.Remove(BoneIdx);
+							BoneToMesh.Add(BoneIdx, NewMesh->GetMeshPtr());
+						}
+						else
+						{
+							TObjectPtr<UDynamicMesh>& NewMesh = NewMeshes.Add_GetRef(NewObject<UDynamicMesh>());
+							FDynamicMesh3& Mesh = NewMesh->GetMeshRef();
+							UE::Geometry::FDynamicMeshEditor Editor(&Mesh);
+							TArray<int32> SearchBones;
+							SearchBones.Add(BoneIdx);
+							while (!SearchBones.IsEmpty())
+							{
+								int32 SearchBoneIdx = SearchBones.Pop(EAllowShrinking::No);
+								const FDynamicMesh3* FoundMesh = nullptr;
+								if (BoneToMeshInfo.Contains(SearchBoneIdx))
+								{
+									FoundMesh = BoneToMeshInfo[SearchBoneIdx]->Mesh.Get();
+								}
+								else if (BoneToMesh.Contains(SearchBoneIdx))
+								{
+									FoundMesh = BoneToMesh[SearchBoneIdx];
+								}
+
+								if (FoundMesh)
+								{
+									Mesh.EnableMatchingAttributes(*FoundMesh);
+									UE::Geometry::FMeshIndexMappings Unused;
+									Editor.AppendMesh(FoundMesh, Unused);
+								}
+								else
+								{
+									// No mesh for this bone; search the children for meshes
+									const TSet<int32>* Children = HierarchyFacade.FindChildren(SearchBoneIdx);
+									if (Children)
+									{
+										SearchBones.Append(Children->Array());
+									}
+								}
+							}
+							// add the built mesh to the map, in case we want to build a parent of it later
+							BoneToMesh.Add(BoneIdx, &Mesh);
+						}
+					}
+				}
+			}
+		}
+
+		SetValue(Context, NewMeshes, &Meshes);
+	}
+}
diff --git a/Engine/Plugins/Experimental/GeometryCollectionPlugin/Source/GeometryCollectionNodes/Public/Dataflow/GeometryCollectionMeshNodes.h b/Engine/Plugins/Experimental/GeometryCollectionPlugin/Source/GeometryCollectionNodes/Public/Dataflow/GeometryCollectionMeshNodes.h
--- a/Engine/Plugins/Experimental/GeometryCollectionPlugin/Source/GeometryCollectionNodes/Public/Dataflow/GeometryCollectionMeshNodes.h
+++ b/Engine/Plugins/Experimental/GeometryCollectionPlugin/Source/GeometryCollectionNodes/Public/Dataflow/GeometryCollectionMeshNodes.h
@@ -175,6 +175,85 @@ public:
 	virtual void Evaluate(UE::Dataflow::FContext& Context, const FDataflowOutput* Out) const override;
 };
 
+/**
+ *
+ * Converts a Collection to a set of Dynamic Meshes per selected Transform
+ *
+ */
+USTRUCT(meta = (DataflowGeometryCollection))
+struct FCollectionSelectionToMeshesDataflowNode : public FDataflowNode
+{
+	GENERATED_USTRUCT_BODY()
+	DATAFLOW_NODE_DEFINE_INTERNAL(FCollectionSelectionToMeshesDataflowNode, "CollectionSelectionToMeshes", "Mesh|Utilities", "")
+	//DATAFLOW_NODE_RENDER_TYPE("SurfaceRender", FName("FDynamicMesh3"), "Meshes")
+
+public:
+	/** Collection to convert*/
+	UPROPERTY(meta = (DataflowInput, DataflowIntrinsic))
+	FManagedArrayCollection Collection;
+
+	/** Selected bones will have geometry script processing applied (if they have geometry). If not connected, all geometry will be processed. */
+	UPROPERTY(meta = (DataflowInput))
+	FDataflowTransformSelection TransformSelection;
+
+	/** Whether to convert the input selection to only leaves, which may directly store geometry. Otherwise, meshes for selected cluster nodes will be generated by appending together geometry from leaf nodes. */
+	UPROPERTY(EditAnywhere, Category = Options)
+	bool bConvertSelectionToLeaves = false;
+
+	/** Whether the processed mesh will have edges at normal/UV/color seams welded so they are treated as one edge during processing. */
+	UPROPERTY(EditAnywhere, Category = Options)
+	bool bWeldVertices = true;
+
+	/** Whether to preserve isolated vertices which aren't used by any triangles. */
+	UPROPERTY(EditAnywhere, Category = Options)
+	bool bPreserveIsolatedVertices = true;
+
+	/** Output DynamicMesh */
+	UPROPERTY(meta = (DataflowOutput))
+	TArray<TObjectPtr<UDynamicMesh>> Meshes;
+
+	FCollectionSelectionToMeshesDataflowNode(const UE::Dataflow::FNodeParameters& InParam, FGuid InGuid = FGuid::NewGuid());
+
+	virtual void Evaluate(UE::Dataflow::FContext& Context, const FDataflowOutput* Out) const override;
+
+};
+
+/**
+ *
+ * Append Array of Meshes to Collection
+ *
+ */
+USTRUCT(meta = (DataflowGeometryCollection))
+struct FAppendMeshesToCollectionDataflowNode : public FDataflowNode
+{
+	GENERATED_USTRUCT_BODY()
+	DATAFLOW_NODE_DEFINE_INTERNAL(FAppendMeshesToCollectionDataflowNode, "AppendMeshesToCollection", "Mesh|Utilities", "")
+	DATAFLOW_NODE_RENDER_TYPE("SurfaceRender", FGeometryCollection::StaticType(), "Collection")
+
+public:
+	/** Meshes will be appended to this collection */
+	UPROPERTY(meta = (DataflowInput, DataflowOutput, DataflowPassthrough = "Collection", DataflowIntrinsic))
+	FManagedArrayCollection Collection;
+	
+	/** Selection of added transforms */
+	UPROPERTY(meta = (DataflowOutput))
+	FDataflowTransformSelection AddedSelection;
+
+	/** Dynamic Meshes to append */
+	UPROPERTY(meta = (DataflowInput))
+	TArray<TObjectPtr<UDynamicMesh>> Meshes;
+
+	/** Index of parent bone for appended meshes. If invalid, meshes will be appended to a root node. */
+	UPROPERTY(EditAnywhere, Category = Options, meta = (DataflowInput))
+	int32 ParentIndex = -1;
+
+	FAppendMeshesToCollectionDataflowNode(const UE::Dataflow::FNodeParameters& InParam, FGuid InGuid = FGuid::NewGuid());
+
+	virtual void Evaluate(UE::Dataflow::FContext& Context, const FDataflowOutput* Out) const override;
+
+};
+
+
 /**
  *
  * Converts a BoundingBox into a DynamicMesh
diff --git a/Engine/Source/Runtime/MeshConversionEngineTypes/Private/GeometryCollectionToDynamicMesh.cpp b/Engine/Source/Runtime/MeshConversionEngineTypes/Private/GeometryCollectionToDynamicMesh.cpp
--- a/Engine/Source/Runtime/MeshConversionEngineTypes/Private/GeometryCollectionToDynamicMesh.cpp
+++ b/Engine/Source/Runtime/MeshConversionEngineTypes/Private/GeometryCollectionToDynamicMesh.cpp
@@ -16,6 +16,7 @@
 #include "GeometryCollection/Facades/CollectionMeshFacade.h"
 #include "GeometryCollection/Facades/CollectionTransformFacade.h"
 #include "GeometryCollection/Facades/CollectionUVFacade.h"
+#include "GeometryCollection/GeometryCollectionClusteringUtility.h"
 
 namespace UE::Geometry
 {
@@ -285,7 +286,7 @@ bool FGeometryCollectionToDynamicMeshes::Init(const FManagedArrayCollection& Col
 // Convert geometry at selected transform indices in the given Collection to dynamic meshes
 bool FGeometryCollectionToDynamicMeshes::InitFromTransformSelection(const FManagedArrayCollection& Collection, TConstArrayView<int32> TransformIndices, const FToMeshOptions& Options)
 {
-	return InitHelper(Collection, true, GetCollectionTransforms(Collection), true, TransformIndices, Options);
+	return InitHelper(Collection, true, GetCollectionTransforms(Collection), false, TransformIndices, Options);
 }
 
 namespace
@@ -397,19 +398,27 @@ bool FGeometryCollectionToDynamicMeshes::UpdateGeometryCollection(FGeometryColle
 		}
 		return Count;
 	};
+	int32 NeedsAppend = 0;
 	for (int MeshIdx = 0; MeshIdx < Meshes.Num(); MeshIdx++)
 	{
 		const FMeshInfo& MeshData = Meshes[MeshIdx];
-		if (MeshData.Mesh.IsValid() && Collection.TransformToGeometryIndex.IsValidIndex(MeshData.TransformIndex))
+		if (MeshData.Mesh.IsValid())
 		{
-			int32 GeomIdx = Collection.TransformToGeometryIndex[MeshData.TransformIndex];
-			int32 PrevTriCount = NewFaceCounts[GeomIdx];
-			NewFaceCounts[GeomIdx] = MeshData.Mesh->TriangleCount();
-			bNeedsResize = bNeedsResize || NewFaceCounts[GeomIdx] != PrevTriCount;
-			int32 PrevVertexCount = NewVertexCounts[GeomIdx];
-			NewVertexCounts[GeomIdx] = CountUniqueVerts(*MeshData.Mesh);
-			bNeedsSplits = bNeedsSplits || NewVertexCounts[GeomIdx] != MeshData.Mesh->VertexCount();
-			bNeedsResize = bNeedsResize || NewVertexCounts[GeomIdx] != PrevVertexCount;
+			if (Collection.TransformToGeometryIndex.IsValidIndex(MeshData.TransformIndex))
+			{
+				int32 GeomIdx = Collection.TransformToGeometryIndex[MeshData.TransformIndex];
+				int32 PrevTriCount = NewFaceCounts[GeomIdx];
+				NewFaceCounts[GeomIdx] = MeshData.Mesh->TriangleCount();
+				bNeedsResize = bNeedsResize || NewFaceCounts[GeomIdx] != PrevTriCount;
+				int32 PrevVertexCount = NewVertexCounts[GeomIdx];
+				NewVertexCounts[GeomIdx] = CountUniqueVerts(*MeshData.Mesh);
+				bNeedsSplits = bNeedsSplits || NewVertexCounts[GeomIdx] != MeshData.Mesh->VertexCount();
+				bNeedsResize = bNeedsResize || NewVertexCounts[GeomIdx] != PrevVertexCount;
+			}
+			else
+			{
+				++NeedsAppend;
+			}
 		}
 	}
 
@@ -425,15 +434,106 @@ bool FGeometryCollectionToDynamicMeshes::UpdateGeometryCollection(FGeometryColle
 	for (int MeshIdx = 0; MeshIdx < Meshes.Num(); MeshIdx++)
 	{
 		const FMeshInfo& MeshData = Meshes[MeshIdx];
+		if (!MeshData.Mesh.IsValid()
+			|| !Collection.TransformToGeometryIndex.IsValidIndex(MeshData.TransformIndex))
+		{
+			continue;
+		}
 		const FDynamicMesh3& Mesh = *MeshData.Mesh;
 		int32 GeometryIdx = Collection.TransformToGeometryIndex[MeshData.TransformIndex];
 		bool bSucceeded = UpdateCollection(MeshData.Transform, Mesh, GeometryIdx, Collection, Options);
 		bAllSucceeded &= bSucceeded;
 	}
+	if (NeedsAppend > 0)
+	{
+		for (int MeshIdx = 0; MeshIdx < Meshes.Num(); MeshIdx++)
+		{
+			const FMeshInfo& MeshData = Meshes[MeshIdx];
+			if (!MeshData.Mesh.IsValid()
+				|| Collection.TransformToGeometryIndex.IsValidIndex(MeshData.TransformIndex))
+			{
+				continue;
+			}
+			const FDynamicMesh3& Mesh = *MeshData.Mesh;
+			const int32 NewTransformIdx = AppendMeshToCollection(Collection, Mesh, MeshData.Transform, Options);
+
+			bool bSucceeded = NewTransformIdx != INDEX_NONE;
+			bAllSucceeded &= bSucceeded;
+		}
+	}
 
 	return bAllSucceeded;
 }
 
+int32 FGeometryCollectionToDynamicMeshes::AppendMeshToCollection(FGeometryCollection& Collection, const FDynamicMesh3& Mesh, const FTransform& MeshTransform, const FToCollectionOptions& Options)
+{
+	TArray<TUniqueMeshVertex<>> Uniques;
+	auto CountUniqueVerts = [&Uniques](const FDynamicMesh3& Mesh)
+	{
+		int32 Count = 0;
+		for (int32 VID : Mesh.VertexIndicesItr())
+		{
+			Uniques.Reset();
+			for (int32 TID : Mesh.VtxTrianglesItr(VID))
+			{
+				FIndex3i Tri = Mesh.GetTriangle(TID);
+				int32 SubIdx = Tri.IndexOf(VID);
+				Uniques.AddUnique(TUniqueMeshVertex<>(Mesh, TID, SubIdx));
+			}
+			Count += FMath::Max(1, Uniques.Num());
+		}
+		return Count;
+	};
+
+	const int32 GeometryIdx = Collection.AddElements(1, FGeometryCollection::GeometryGroup);
+	const int32 TransformIdx = Collection.AddElements(1, FGeometryCollection::TransformGroup);
+	const int32 NumTriangles = Mesh.TriangleCount();
+	const int32 NumVertices = CountUniqueVerts(Mesh);
+	const int32 FacesStart = Collection.AddElements(NumTriangles, FGeometryCollection::FacesGroup);
+	const int32 VerticesStart = Collection.AddElements(NumVertices, FGeometryCollection::VerticesGroup);
+
+	Collection.FaceCount[GeometryIdx] = NumTriangles;
+	Collection.FaceStart[GeometryIdx] = FacesStart;
+	Collection.VertexCount[GeometryIdx] = NumVertices;
+	Collection.VertexStart[GeometryIdx] = VerticesStart;
+	Collection.TransformIndex[GeometryIdx] = TransformIdx;
+	Collection.TransformToGeometryIndex[TransformIdx] = GeometryIdx;
+	int32 TransformParent = Options.NewMeshParentIndex;
+	// If parent is INDEX_NONE, search for a root to use as parent
+	if (TransformParent == INDEX_NONE)
+	{
+		for (int32 Idx = 0; Idx < Collection.Parent.Num(); ++Idx)
+		{
+			if (Collection.Parent[Idx] == -1)
+			{
+				TransformParent = Idx;
+			}
+		}
+	}
+	// If still haven't found a valid parent, add a new root as parent (can occur if collection was empty)
+	if (!Collection.Parent.IsValidIndex(TransformParent))
+	{
+		TransformParent = Collection.AddElements(1, FGeometryCollection::TransformGroup);
+		Collection.Parent[TransformParent] = INDEX_NONE;
+		Collection.BoneColor[TransformParent] = FLinearColor::White;
+	}
+	Collection.BoneName[TransformIdx] = Collection.BoneName[TransformParent] + "_" + FString::FromInt(Collection.Children[TransformParent].Num());
+	Collection.BoneColor[TransformIdx] = Collection.BoneColor[TransformParent];
+	Collection.Parent[TransformIdx] = TransformParent;
+	Collection.Children[TransformParent].Add(TransformIdx);
+	Collection.SimulationType[TransformParent] = FGeometryCollection::ESimulationTypes::FST_Clustered;
+
+	Collection.Transform[TransformIdx] = FTransform3f::Identity;
+	Collection.SimulationType[TransformIdx] = FGeometryCollection::ESimulationTypes::FST_Rigid;
+
+	bool bSuccess = UpdateCollection(MeshTransform, Mesh, GeometryIdx, Collection, Options);
+	if (!ensure(bSuccess)) // Note: we should have set up the geometry collection above such that UpdateCollection won't fail
+	{
+		return INDEX_NONE;
+	}
+	return TransformIdx;
+}
+
 
 // Update an existing geometry in a collection w/ a new mesh (w/ the same number of faces and vertices!)
 bool FGeometryCollectionToDynamicMeshes::UpdateCollection(const FTransform& FromCollection, const FDynamicMesh3& Mesh, int32 GeometryIdx, FGeometryCollection& Output, const FToCollectionOptions& Options)
diff --git a/Engine/Source/Runtime/MeshConversionEngineTypes/Public/GeometryCollectionToDynamicMesh.h b/Engine/Source/Runtime/MeshConversionEngineTypes/Public/GeometryCollectionToDynamicMesh.h
--- a/Engine/Source/Runtime/MeshConversionEngineTypes/Public/GeometryCollectionToDynamicMesh.h
+++ b/Engine/Source/Runtime/MeshConversionEngineTypes/Public/GeometryCollectionToDynamicMesh.h
@@ -49,6 +49,9 @@ struct FGeometryCollectionToDynamicMeshes
 		bool bDefaultFaceVisible = true;
 		// Whether to set faces with no 'internal' tag as internal (if true) or external (if false)
 		bool bDefaultFaceInternal = false;
+
+		// Parent index to use if adding a new mesh (with no existing transform) -- if invalid, will add to root
+		int32 NewMeshParentIndex = -1;
 	};
 
 	// Convert all meshes in the given collection to dynamic meshes
@@ -61,6 +64,11 @@ struct FGeometryCollectionToDynamicMeshes
 	// Note this updates a Geometry Collection rather than the more general FManagedArrayCollection because some of the updating code is specific to geometry collection
 	MESHCONVERSIONENGINETYPES_API bool UpdateGeometryCollection(FGeometryCollection& Collection, const FToCollectionOptions& Options) const;
 
+	// Add a new mesh to the geometry collection
+	// Note this updates a Geometry Collection rather than the more general FManagedArrayCollection because some of the updating code is specific to geometry collection
+	// @return Index of added transform
+	MESHCONVERSIONENGINETYPES_API static int32 AppendMeshToCollection(FGeometryCollection& Collection, const FDynamicMesh3& Mesh, const FTransform& MeshTransform, const FToCollectionOptions& Options);
+
 
 
 	// Get the name used for the polygroup that we optionally set on the dynamic mesh, corresponding to the internal face tags of the geometry collection triangles
