Add include required for non-unity builds.
#fixes the non-unity build Issue 722321

diff --git a/Engine/Plugins/Experimental/GeometryCollectionPlugin/Source/GeometryCollectionNodes/Private/Dataflow/GeometryCollectionMaterialInterfaceNodes.cpp b/Engine/Plugins/Experimental/GeometryCollectionPlugin/Source/GeometryCollectionNodes/Private/Dataflow/GeometryCollectionMaterialInterfaceNodes.cpp
--- a/Engine/Plugins/Experimental/GeometryCollectionPlugin/Source/GeometryCollectionNodes/Private/Dataflow/GeometryCollectionMaterialInterfaceNodes.cpp
+++ b/Engine/Plugins/Experimental/GeometryCollectionPlugin/Source/GeometryCollectionNodes/Private/Dataflow/GeometryCollectionMaterialInterfaceNodes.cpp
@@ -3,6 +3,7 @@
 #include "Dataflow/GeometryCollectionMaterialInterfaceNodes.h"
 #include "Dataflow/DataflowCore.h"
 #include "GeometryCollection/GeometryCollection.h"
+#include "Materials/MaterialInterface.h"
 
 namespace UE::Dataflow
 {
