Dataflow:
- Added pin hiding to all the Fracture/Sampling nodes

diff --git a/Engine/Plugins/Experimental/GeometryCollectionPlugin/Source/GeometryCollectionNodes/Public/Dataflow/GeometryCollectionFracturingNodes.h b/Engine/Plugins/Experimental/GeometryCollectionPlugin/Source/GeometryCollectionNodes/Public/Dataflow/GeometryCollectionFracturingNodes.h
--- a/Engine/Plugins/Experimental/GeometryCollectionPlugin/Source/GeometryCollectionNodes/Public/Dataflow/GeometryCollectionFracturingNodes.h
+++ b/Engine/Plugins/Experimental/GeometryCollectionPlugin/Source/GeometryCollectionNodes/Public/Dataflow/GeometryCollectionFracturingNodes.h
@@ -87,9 +87,9 @@ public:
 		: FDataflowNode(InParam, InGuid)
 	{
 		RegisterInputConnection(&BoundingBox);
-		RegisterInputConnection(&MinNumberOfPoints);
-		RegisterInputConnection(&MaxNumberOfPoints);
-		RegisterInputConnection(&RandomSeed);
+		RegisterInputConnection(&MinNumberOfPoints) .SetCanHidePin(true) .SetPinIsHidden(true);
+		RegisterInputConnection(&MaxNumberOfPoints) .SetCanHidePin(true).SetPinIsHidden(true);
+		RegisterInputConnection(&RandomSeed) .SetCanHidePin(true).SetPinIsHidden(true);
 		RegisterOutputConnection(&Points);
 	}
 
@@ -226,19 +226,19 @@ public:
 	{
 		RegisterInputConnection(&BoundingBox);
 		RegisterInputConnection(&Center);
-		RegisterInputConnection(&Normal);
-		RegisterInputConnection(&RandomSeed);
-		RegisterInputConnection(&AngularSteps);
-		RegisterInputConnection(&AngleOffset);
-		RegisterInputConnection(&AngularNoise);
-		RegisterInputConnection(&Radius);
-		RegisterInputConnection(&RadialSteps);
-		RegisterInputConnection(&RadialStepExponent);
-		RegisterInputConnection(&RadialMinStep);
-		RegisterInputConnection(&RadialNoise);
-		RegisterInputConnection(&RadialVariability);
-		RegisterInputConnection(&AngularVariability);
-		RegisterInputConnection(&AxialVariability);
+		RegisterInputConnection(&Normal) .SetCanHidePin(true).SetPinIsHidden(true);
+		RegisterInputConnection(&RandomSeed) .SetCanHidePin(true).SetPinIsHidden(true);
+		RegisterInputConnection(&AngularSteps) .SetCanHidePin(true).SetPinIsHidden(true);
+		RegisterInputConnection(&AngleOffset) .SetCanHidePin(true).SetPinIsHidden(true);
+		RegisterInputConnection(&AngularNoise) .SetCanHidePin(true).SetPinIsHidden(true);
+		RegisterInputConnection(&Radius) .SetCanHidePin(true).SetPinIsHidden(true);
+		RegisterInputConnection(&RadialSteps) .SetCanHidePin(true).SetPinIsHidden(true);
+		RegisterInputConnection(&RadialStepExponent) .SetCanHidePin(true).SetPinIsHidden(true);
+		RegisterInputConnection(&RadialMinStep) .SetCanHidePin(true).SetPinIsHidden(true);
+		RegisterInputConnection(&RadialNoise) .SetCanHidePin(true).SetPinIsHidden(true);
+		RegisterInputConnection(&RadialVariability) .SetCanHidePin(true).SetPinIsHidden(true);
+		RegisterInputConnection(&AngularVariability) .SetCanHidePin(true).SetPinIsHidden(true);
+		RegisterInputConnection(&AxialVariability) .SetCanHidePin(true).SetPinIsHidden(true);
 		RegisterOutputConnection(&Points);
 	}
 
@@ -293,13 +293,13 @@ public:
 		: FDataflowNode(InParam, InGuid)
 	{
 		RegisterInputConnection(&BoundingBox);
-		RegisterInputConnection(&NumberOfPointsInX);
-		RegisterInputConnection(&NumberOfPointsInY);
-		RegisterInputConnection(&NumberOfPointsInZ);
-		RegisterInputConnection(&RandomSeed);
-		RegisterInputConnection(&MaxRandomDisplacementX);
-		RegisterInputConnection(&MaxRandomDisplacementY);
-		RegisterInputConnection(&MaxRandomDisplacementZ);
+		RegisterInputConnection(&NumberOfPointsInX) .SetCanHidePin(true).SetPinIsHidden(true);
+		RegisterInputConnection(&NumberOfPointsInY) .SetCanHidePin(true).SetPinIsHidden(true);
+		RegisterInputConnection(&NumberOfPointsInZ) .SetCanHidePin(true).SetPinIsHidden(true);
+		RegisterInputConnection(&RandomSeed) .SetCanHidePin(true).SetPinIsHidden(true);
+		RegisterInputConnection(&MaxRandomDisplacementX) .SetCanHidePin(true).SetPinIsHidden(true);
+		RegisterInputConnection(&MaxRandomDisplacementY) .SetCanHidePin(true).SetPinIsHidden(true);
+		RegisterInputConnection(&MaxRandomDisplacementZ) .SetCanHidePin(true).SetPinIsHidden(true);
 		RegisterOutputConnection(&Points);
 	}
 
@@ -487,16 +487,16 @@ public:
 		RegisterInputConnection(&Collection);
 		RegisterInputConnection(&Points);
 		RegisterInputConnection(&TransformSelection);
-		RegisterInputConnection(&Transform);
-		RegisterInputConnection(&ChanceToFracture);
-		RegisterInputConnection(&Grout);
-		RegisterInputConnection(&Amplitude);
-		RegisterInputConnection(&Frequency);
-		RegisterInputConnection(&Persistence);
-		RegisterInputConnection(&Lacunarity);
-		RegisterInputConnection(&OctaveNumber);
-		RegisterInputConnection(&PointSpacing);
-		RegisterInputConnection(&CollisionSampleSpacing);
+		RegisterInputConnection(&Transform) .SetCanHidePin(true).SetPinIsHidden(true);
+		RegisterInputConnection(&ChanceToFracture) .SetCanHidePin(true).SetPinIsHidden(true);
+		RegisterInputConnection(&Grout) .SetCanHidePin(true).SetPinIsHidden(true);
+		RegisterInputConnection(&Amplitude) .SetCanHidePin(true).SetPinIsHidden(true);
+		RegisterInputConnection(&Frequency) .SetCanHidePin(true).SetPinIsHidden(true);
+		RegisterInputConnection(&Persistence) .SetCanHidePin(true).SetPinIsHidden(true);
+		RegisterInputConnection(&Lacunarity) .SetCanHidePin(true).SetPinIsHidden(true);
+		RegisterInputConnection(&OctaveNumber) .SetCanHidePin(true).SetPinIsHidden(true);
+		RegisterInputConnection(&PointSpacing) .SetCanHidePin(true).SetPinIsHidden(true);
+		RegisterInputConnection(&CollisionSampleSpacing) .SetCanHidePin(true).SetPinIsHidden(true);
 		RegisterOutputConnection(&Collection, &Collection);
 		RegisterOutputConnection(&TransformSelection, &TransformSelection);
 		RegisterOutputConnection(&NewGeometryTransformSelection);
@@ -686,17 +686,17 @@ public:
 		RegisterInputConnection(&Collection);
 		RegisterInputConnection(&BoundingBox);
 		RegisterInputConnection(&TransformSelection);
-		RegisterInputConnection(&Transform);
-		RegisterInputConnection(&NumPlanes);
-		RegisterInputConnection(&RandomSeed);
-		RegisterInputConnection(&ChanceToFracture);
-		RegisterInputConnection(&Grout);
-		RegisterInputConnection(&Amplitude);
-		RegisterInputConnection(&Frequency);
-		RegisterInputConnection(&Persistence);
-		RegisterInputConnection(&Lacunarity);
-		RegisterInputConnection(&OctaveNumber);
-		RegisterInputConnection(&PointSpacing);
+		RegisterInputConnection(&Transform) .SetCanHidePin(true).SetPinIsHidden(true);
+		RegisterInputConnection(&NumPlanes) .SetCanHidePin(true).SetPinIsHidden(true);
+		RegisterInputConnection(&RandomSeed) .SetCanHidePin(true).SetPinIsHidden(true);
+		RegisterInputConnection(&ChanceToFracture) .SetCanHidePin(true).SetPinIsHidden(true);
+		RegisterInputConnection(&Grout) .SetCanHidePin(true).SetPinIsHidden(true);
+		RegisterInputConnection(&Amplitude) .SetCanHidePin(true).SetPinIsHidden(true);
+		RegisterInputConnection(&Frequency) .SetCanHidePin(true).SetPinIsHidden(true);
+		RegisterInputConnection(&Persistence) .SetCanHidePin(true).SetPinIsHidden(true);
+		RegisterInputConnection(&Lacunarity) .SetCanHidePin(true).SetPinIsHidden(true);
+		RegisterInputConnection(&OctaveNumber) .SetCanHidePin(true).SetPinIsHidden(true);
+		RegisterInputConnection(&PointSpacing) .SetCanHidePin(true).SetPinIsHidden(true);
 		RegisterInputConnection(&CollisionSampleSpacing);
 		RegisterOutputConnection(&Collection, &Collection);
 		RegisterOutputConnection(&TransformSelection, &TransformSelection);
@@ -871,21 +871,21 @@ public:
 		RegisterInputConnection(&Collection);
 		RegisterInputConnection(&BoundingBox);
 		RegisterInputConnection(&TransformSelection);
-		RegisterInputConnection(&SlicesX);
-		RegisterInputConnection(&SlicesY);
-		RegisterInputConnection(&SlicesZ);
-		RegisterInputConnection(&SliceAngleVariation);
-		RegisterInputConnection(&SliceOffsetVariation);
-		RegisterInputConnection(&RandomSeed);
-		RegisterInputConnection(&ChanceToFracture);
-		RegisterInputConnection(&Grout);
-		RegisterInputConnection(&Amplitude);
-		RegisterInputConnection(&Frequency);
-		RegisterInputConnection(&Persistence);
-		RegisterInputConnection(&Lacunarity);
-		RegisterInputConnection(&OctaveNumber);
-		RegisterInputConnection(&PointSpacing);
-		RegisterInputConnection(&CollisionSampleSpacing);
+		RegisterInputConnection(&SlicesX) .SetCanHidePin(true).SetPinIsHidden(true);
+		RegisterInputConnection(&SlicesY) .SetCanHidePin(true).SetPinIsHidden(true);
+		RegisterInputConnection(&SlicesZ) .SetCanHidePin(true).SetPinIsHidden(true);
+		RegisterInputConnection(&SliceAngleVariation) .SetCanHidePin(true).SetPinIsHidden(true);
+		RegisterInputConnection(&SliceOffsetVariation) .SetCanHidePin(true).SetPinIsHidden(true);
+		RegisterInputConnection(&RandomSeed) .SetCanHidePin(true).SetPinIsHidden(true);
+		RegisterInputConnection(&ChanceToFracture) .SetCanHidePin(true).SetPinIsHidden(true);
+		RegisterInputConnection(&Grout) .SetCanHidePin(true).SetPinIsHidden(true);
+		RegisterInputConnection(&Amplitude) .SetCanHidePin(true).SetPinIsHidden(true);
+		RegisterInputConnection(&Frequency) .SetCanHidePin(true).SetPinIsHidden(true);
+		RegisterInputConnection(&Persistence) .SetCanHidePin(true).SetPinIsHidden(true);
+		RegisterInputConnection(&Lacunarity) .SetCanHidePin(true).SetPinIsHidden(true);
+		RegisterInputConnection(&OctaveNumber) .SetCanHidePin(true).SetPinIsHidden(true);
+		RegisterInputConnection(&PointSpacing) .SetCanHidePin(true).SetPinIsHidden(true);
+		RegisterInputConnection(&CollisionSampleSpacing) .SetCanHidePin(true).SetPinIsHidden(true);
 		RegisterOutputConnection(&Collection, &Collection);
 		RegisterOutputConnection(&TransformSelection, &TransformSelection);
 		RegisterOutputConnection(&NewGeometryTransformSelection);
@@ -1009,20 +1009,20 @@ public:
 		RegisterInputConnection(&Collection);
 		RegisterInputConnection(&BoundingBox);
 		RegisterInputConnection(&TransformSelection);
-		RegisterInputConnection(&Transform);
-		RegisterInputConnection(&BrickLength);
-		RegisterInputConnection(&BrickHeight);
-		RegisterInputConnection(&BrickDepth);
-		RegisterInputConnection(&RandomSeed);
-		RegisterInputConnection(&ChanceToFracture);
-		RegisterInputConnection(&Grout);
-		RegisterInputConnection(&Amplitude);
-		RegisterInputConnection(&Frequency);
-		RegisterInputConnection(&Persistence);
-		RegisterInputConnection(&Lacunarity);
-		RegisterInputConnection(&OctaveNumber);
-		RegisterInputConnection(&PointSpacing);
-		RegisterInputConnection(&CollisionSampleSpacing);
+		RegisterInputConnection(&Transform) .SetCanHidePin(true).SetPinIsHidden(true);
+		RegisterInputConnection(&BrickLength) .SetCanHidePin(true).SetPinIsHidden(true);
+		RegisterInputConnection(&BrickHeight) .SetCanHidePin(true).SetPinIsHidden(true);
+		RegisterInputConnection(&BrickDepth) .SetCanHidePin(true).SetPinIsHidden(true);
+		RegisterInputConnection(&RandomSeed) .SetCanHidePin(true).SetPinIsHidden(true);
+		RegisterInputConnection(&ChanceToFracture) .SetCanHidePin(true).SetPinIsHidden(true);
+		RegisterInputConnection(&Grout) .SetCanHidePin(true).SetPinIsHidden(true);
+		RegisterInputConnection(&Amplitude) .SetCanHidePin(true).SetPinIsHidden(true);
+		RegisterInputConnection(&Frequency) .SetCanHidePin(true).SetPinIsHidden(true);
+		RegisterInputConnection(&Persistence) .SetCanHidePin(true).SetPinIsHidden(true);
+		RegisterInputConnection(&Lacunarity) .SetCanHidePin(true).SetPinIsHidden(true);
+		RegisterInputConnection(&OctaveNumber) .SetCanHidePin(true).SetPinIsHidden(true);
+		RegisterInputConnection(&PointSpacing) .SetCanHidePin(true).SetPinIsHidden(true);
+		RegisterInputConnection(&CollisionSampleSpacing) .SetCanHidePin(true).SetPinIsHidden(true);
 		RegisterOutputConnection(&Collection, &Collection);
 		RegisterOutputConnection(&TransformSelection, &TransformSelection);
 		RegisterOutputConnection(&NewGeometryTransformSelection);
@@ -1158,21 +1158,21 @@ public:
 		RegisterInputConnection(&Collection);
 		RegisterInputConnection(&BoundingBox);
 		RegisterInputConnection(&TransformSelection);
-		RegisterInputConnection(&Transform);
-		RegisterInputConnection(&CuttingStaticMesh);
-		RegisterInputConnection(&NumberToScatter);
-		RegisterInputConnection(&GridX);
-		RegisterInputConnection(&GridY);
-		RegisterInputConnection(&GridZ);
-		RegisterInputConnection(&Variability);
-		RegisterInputConnection(&MinScaleFactor);
-		RegisterInputConnection(&MaxScaleFactor);
-		RegisterInputConnection(&RollRange);
-		RegisterInputConnection(&PitchRange);
-		RegisterInputConnection(&YawRange);
-		RegisterInputConnection(&RandomSeed);
-		RegisterInputConnection(&ChanceToFracture);
-		RegisterInputConnection(&CollisionSampleSpacing);
+		RegisterInputConnection(&Transform) .SetCanHidePin(true).SetPinIsHidden(true);
+		RegisterInputConnection(&CuttingStaticMesh) .SetCanHidePin(true).SetPinIsHidden(true);
+		RegisterInputConnection(&NumberToScatter) .SetCanHidePin(true).SetPinIsHidden(true);
+		RegisterInputConnection(&GridX) .SetCanHidePin(true).SetPinIsHidden(true);
+		RegisterInputConnection(&GridY) .SetCanHidePin(true).SetPinIsHidden(true);
+		RegisterInputConnection(&GridZ) .SetCanHidePin(true).SetPinIsHidden(true);
+		RegisterInputConnection(&Variability) .SetCanHidePin(true).SetPinIsHidden(true);
+		RegisterInputConnection(&MinScaleFactor) .SetCanHidePin(true).SetPinIsHidden(true);
+		RegisterInputConnection(&MaxScaleFactor) .SetCanHidePin(true).SetPinIsHidden(true);
+		RegisterInputConnection(&RollRange) .SetCanHidePin(true).SetPinIsHidden(true);
+		RegisterInputConnection(&PitchRange) .SetCanHidePin(true).SetPinIsHidden(true);
+		RegisterInputConnection(&YawRange) .SetCanHidePin(true).SetPinIsHidden(true);
+		RegisterInputConnection(&RandomSeed) .SetCanHidePin(true).SetPinIsHidden(true);
+		RegisterInputConnection(&ChanceToFracture) .SetCanHidePin(true).SetPinIsHidden(true);
+		RegisterInputConnection(&CollisionSampleSpacing) .SetCanHidePin(true).SetPinIsHidden(true);
 		RegisterOutputConnection(&Collection, &Collection);
 		RegisterOutputConnection(&TransformSelection, &TransformSelection);
 		RegisterOutputConnection(&NewGeometryTransformSelection);
diff --git a/Engine/Plugins/Experimental/GeometryCollectionPlugin/Source/GeometryCollectionNodes/Public/Dataflow/GeometryCollectionSamplingNodes.h b/Engine/Plugins/Experimental/GeometryCollectionPlugin/Source/GeometryCollectionNodes/Public/Dataflow/GeometryCollectionSamplingNodes.h
--- a/Engine/Plugins/Experimental/GeometryCollectionPlugin/Source/GeometryCollectionNodes/Public/Dataflow/GeometryCollectionSamplingNodes.h
+++ b/Engine/Plugins/Experimental/GeometryCollectionPlugin/Source/GeometryCollectionNodes/Public/Dataflow/GeometryCollectionSamplingNodes.h
@@ -69,10 +69,10 @@ public:
 		: FDataflowNode(InParam, InGuid)
 	{
 		RegisterInputConnection(&TargetMesh);
-		RegisterInputConnection(&SamplingRadius);
-		RegisterInputConnection(&MaxNumSamples);
-		RegisterInputConnection(&SubSampleDensity);
-		RegisterInputConnection(&RandomSeed);
+		RegisterInputConnection(&SamplingRadius) .SetCanHidePin(true).SetPinIsHidden(true);
+		RegisterInputConnection(&MaxNumSamples) .SetCanHidePin(true).SetPinIsHidden(true);
+		RegisterInputConnection(&SubSampleDensity) .SetCanHidePin(true).SetPinIsHidden(true);
+		RegisterInputConnection(&RandomSeed) .SetCanHidePin(true).SetPinIsHidden(true);
 		RegisterOutputConnection(&SamplePoints);
 		RegisterOutputConnection(&SampleTriangleIDs);
 		RegisterOutputConnection(&SampleBarycentricCoords);
@@ -151,12 +151,12 @@ public:
 		: FDataflowNode(InParam, InGuid)
 	{
 		RegisterInputConnection(&TargetMesh);
-		RegisterInputConnection(&SamplingRadius);
-		RegisterInputConnection(&MaxNumSamples);
-		RegisterInputConnection(&SubSampleDensity);
-		RegisterInputConnection(&RandomSeed);
-		RegisterInputConnection(&MaxSamplingRadius);
-		RegisterInputConnection(&SizeDistributionPower);
+		RegisterInputConnection(&SamplingRadius) .SetCanHidePin(true).SetPinIsHidden(true);
+		RegisterInputConnection(&MaxNumSamples) .SetCanHidePin(true).SetPinIsHidden(true);
+		RegisterInputConnection(&SubSampleDensity) .SetCanHidePin(true).SetPinIsHidden(true);
+		RegisterInputConnection(&RandomSeed) .SetCanHidePin(true).SetPinIsHidden(true);
+		RegisterInputConnection(&MaxSamplingRadius) .SetCanHidePin(true).SetPinIsHidden(true);
+		RegisterInputConnection(&SizeDistributionPower) .SetCanHidePin(true).SetPinIsHidden(true);
 		RegisterOutputConnection(&SamplePoints);
 		RegisterOutputConnection(&SampleRadii);
 		RegisterOutputConnection(&SampleTriangleIDs);
@@ -247,12 +247,12 @@ public:
 	{
 		RegisterInputConnection(&TargetMesh);
 		RegisterInputConnection(&VertexWeights);
-		RegisterInputConnection(&SamplingRadius);
-		RegisterInputConnection(&MaxNumSamples);
-		RegisterInputConnection(&SubSampleDensity);
-		RegisterInputConnection(&RandomSeed);
-		RegisterInputConnection(&MaxSamplingRadius);
-		RegisterInputConnection(&SizeDistributionPower);
+		RegisterInputConnection(&SamplingRadius) .SetCanHidePin(true).SetPinIsHidden(true);
+		RegisterInputConnection(&MaxNumSamples) .SetCanHidePin(true).SetPinIsHidden(true);
+		RegisterInputConnection(&SubSampleDensity) .SetCanHidePin(true).SetPinIsHidden(true);
+		RegisterInputConnection(&RandomSeed) .SetCanHidePin(true).SetPinIsHidden(true);
+		RegisterInputConnection(&MaxSamplingRadius) .SetCanHidePin(true).SetPinIsHidden(true);
+		RegisterInputConnection(&SizeDistributionPower) .SetCanHidePin(true).SetPinIsHidden(true);
 		RegisterOutputConnection(&SamplePoints);
 		RegisterOutputConnection(&SampleRadii);
 		RegisterOutputConnection(&SampleTriangleIDs);
