Add dynamic mesh array as an optional input (vs static mesh) to the geometry collection fracturing nodes

diff --git a/Engine/Plugins/Experimental/Fracture/Source/FractureEngine/Private/FractureEngineFracturing.cpp b/Engine/Plugins/Experimental/Fracture/Source/FractureEngine/Private/FractureEngineFracturing.cpp
--- a/Engine/Plugins/Experimental/Fracture/Source/FractureEngine/Private/FractureEngineFracturing.cpp
+++ b/Engine/Plugins/Experimental/Fracture/Source/FractureEngine/Private/FractureEngineFracturing.cpp
@@ -1004,11 +1004,11 @@ void FFractureEngineFracturing::GenerateMeshTransforms(TArray<FTransform>& MeshT
 	}
 }
 
-int32 FFractureEngineFracturing::MeshCutter(TArray<FTransform>& MeshTransforms,
+int32 FFractureEngineFracturing::MeshArrayCutter(TArray<FTransform>& MeshTransforms,
 	FManagedArrayCollection& InOutCollection,
-	FDataflowTransformSelection InTransformSelection,
-	const UE::Geometry::FDynamicMesh3& InDynCuttingMesh,
-	const FTransform& InTransform,
+	const FDataflowTransformSelection& InTransformSelectionConst,
+	TArrayView<const UE::Geometry::FDynamicMesh3*> InDynCuttingMeshes,
+	const EMeshCutterPerCutMeshSelection PerCutMeshSelection,
 	const int32 InRandomSeed,
 	const float InChanceToFracture,
 	const bool InSplitIslands,
@@ -1030,33 +1030,60 @@ int32 FFractureEngineFracturing::MeshCutter(TArray<FTransform>& MeshTransforms,
 
 		const int32 OriginalNumTransforms = InOutCollection.NumElements(FGeometryCollection::TransformGroup);
 
+		FRandomStream RandStream(InRandomSeed);
+
+		FDataflowTransformSelection InTransformSelection = InTransformSelectionConst;
 		RandomReduceSelection(InTransformSelection, InRandomSeed, InChanceToFracture);
 		UE::Private::FractureHelpers::ConvertToLeafSelection(*GeomCollection, InTransformSelection);
 		TArray<int32> TransformSelectionArr = InTransformSelection.AsArray();
 
+		int32 SequentialMeshIndex = 0;
 		for (const FTransform& ScatterTransform : MeshTransforms)
 		{
-			constexpr bool bSetDefaultInternalMaterialsFromCollection = true;
-			int32 Index = CutWithMesh(InDynCuttingMesh, ScatterTransform, InternalSurfaceMaterials, *GeomCollection, TransformSelectionArr, InCollisionSampleSpacing, FTransform::Identity, bSetDefaultInternalMaterialsFromCollection, nullptr, InSplitIslands);
+			auto ApplyCut = 
+				[&ScatterTransform, &InternalSurfaceMaterials, &GeomCollection, &TransformSelectionArr, InCollisionSampleSpacing, InSplitIslands, &ResultGeometryIndex]
+				(const UE::Geometry::FDynamicMesh3& CuttingMesh)
+			{
+				constexpr bool bSetDefaultInternalMaterialsFromCollection = true;
+				int32 Index = CutWithMesh(CuttingMesh, ScatterTransform, InternalSurfaceMaterials, *GeomCollection, TransformSelectionArr, InCollisionSampleSpacing, FTransform::Identity, bSetDefaultInternalMaterialsFromCollection, nullptr, InSplitIslands);
 
-			int32 NewLen = Algo::RemoveIf(TransformSelectionArr, [&](int32 Bone)
+				int32 NewLen = Algo::RemoveIf(TransformSelectionArr, [&GeomCollection](int32 Bone)
 				{
 					return !GeomCollection->IsVisible(Bone); // remove already-fractured pieces from the to-cut list
 				});
-			TransformSelectionArr.SetNum(NewLen);
+				TransformSelectionArr.SetNum(NewLen);
 
-			if (ResultGeometryIndex == -1)
-			{
-				ResultGeometryIndex = Index;
-			}
-			if (Index > -1)
+				if (ResultGeometryIndex == -1)
+				{
+					ResultGeometryIndex = Index;
+				}
+				if (Index > -1)
+				{
+					const int32 TransformIdx = GeomCollection->TransformIndex[Index];
+					// after a successful cut, also consider any new bones added by the cut
+					for (int32 NewBoneIdx = TransformIdx; NewBoneIdx < GeomCollection->NumElements(FGeometryCollection::TransformGroup); NewBoneIdx++)
+					{
+						TransformSelectionArr.Add(NewBoneIdx);
+					}
+				}
+			};
+			switch (PerCutMeshSelection)
 			{
-				const int32 TransformIdx = GeomCollection->TransformIndex[Index];
-				// after a successful cut, also consider any new bones added by the cut
-				for (int32 NewBoneIdx = TransformIdx; NewBoneIdx < GeomCollection->NumElements(FGeometryCollection::TransformGroup); NewBoneIdx++)
+			case EMeshCutterPerCutMeshSelection::All:
+				for (const UE::Geometry::FDynamicMesh3* Mesh : InDynCuttingMeshes)
 				{
-					TransformSelectionArr.Add(NewBoneIdx);
+					ApplyCut(*Mesh);
 				}
+				break;
+			case EMeshCutterPerCutMeshSelection::Random:
+				ApplyCut(*InDynCuttingMeshes[RandStream.RandHelper(InDynCuttingMeshes.Num())]);
+				break;
+			case EMeshCutterPerCutMeshSelection::Sequential:
+				ApplyCut(*InDynCuttingMeshes[SequentialMeshIndex]);
+				SequentialMeshIndex = (SequentialMeshIndex + 1) % InDynCuttingMeshes.Num();
+				break;
+			default:
+				checkNoEntry();
 			}
 		}
 
@@ -1095,6 +1122,21 @@ int32 FFractureEngineFracturing::MeshCutter(TArray<FTransform>& MeshTransforms,
 	return INDEX_NONE;
 }
 
+int32 FFractureEngineFracturing::MeshCutter(TArray<FTransform>& MeshTransforms,
+	FManagedArrayCollection& InOutCollection,
+	const FDataflowTransformSelection& InTransformSelection,
+	const UE::Geometry::FDynamicMesh3& InDynCuttingMesh,
+	const int32 InRandomSeed,
+	const float InChanceToFracture,
+	const bool InSplitIslands,
+	const float InCollisionSampleSpacing)
+{
+	const UE::Geometry::FDynamicMesh3* LocalCutterPtr = &InDynCuttingMesh;
+	return MeshArrayCutter(MeshTransforms, InOutCollection, InTransformSelection, 
+		TArrayView<const UE::Geometry::FDynamicMesh3*>(&LocalCutterPtr, 1), EMeshCutterPerCutMeshSelection::All,
+		InRandomSeed, InChanceToFracture, InSplitIslands, InCollisionSampleSpacing);
+}
+
 namespace UE::Dataflow::Private
 {
 	static void GenerateVoronoiSites(const FBox& InBoundingBox,
diff --git a/Engine/Plugins/Experimental/Fracture/Source/FractureEngine/Public/FractureEngineFracturing.h b/Engine/Plugins/Experimental/Fracture/Source/FractureEngine/Public/FractureEngineFracturing.h
--- a/Engine/Plugins/Experimental/Fracture/Source/FractureEngine/Public/FractureEngineFracturing.h
+++ b/Engine/Plugins/Experimental/Fracture/Source/FractureEngine/Public/FractureEngineFracturing.h
@@ -34,6 +34,17 @@ enum class EMeshCutterCutDistribution : uint8
 	Grid UMETA(DisplayName = "Grid"),
 };
 
+UENUM(BlueprintType)
+enum class EMeshCutterPerCutMeshSelection : uint8
+{
+	// Use all cut meshes for every cut
+	All,
+	// Choose which cut mesh to use from the mesh array randomly, per cut
+	Random,
+	// Choose which cut mesh to use from the mesh array sequentially, starting with the first mesh and cycling through, per cut
+	Sequential,
+};
+
 struct FUniformFractureSettings
 {
 	FTransform Transform;
@@ -191,9 +202,18 @@ public:
 
 	static int32 MeshCutter(TArray<FTransform>& MeshTransforms,
 		FManagedArrayCollection& InOutCollection,
-		FDataflowTransformSelection InTransformSelection,
+		const FDataflowTransformSelection& InTransformSelection,
 		const UE::Geometry::FDynamicMesh3& InDynCuttingMesh,
-		const FTransform& InTransform,
+		const int32 InRandomSeed,
+		const float InChanceToFracture,
+		const bool InSplitIslands,
+		const float InCollisionSampleSpacing);
+
+	static int32 MeshArrayCutter(TArray<FTransform>& MeshTransforms,
+		FManagedArrayCollection& InOutCollection,
+		const FDataflowTransformSelection& InTransformSelection,
+		TArrayView<const UE::Geometry::FDynamicMesh3*> InDynCuttingMeshes,
+		const EMeshCutterPerCutMeshSelection PerCutMeshSelection,
 		const int32 InRandomSeed,
 		const float InChanceToFracture,
 		const bool InSplitIslands,
diff --git a/Engine/Plugins/Experimental/GeometryCollectionPlugin/Source/GeometryCollectionNodes/Private/Dataflow/GeometryCollectionFracturingNodes.cpp b/Engine/Plugins/Experimental/GeometryCollectionPlugin/Source/GeometryCollectionNodes/Private/Dataflow/GeometryCollectionFracturingNodes.cpp
--- a/Engine/Plugins/Experimental/GeometryCollectionPlugin/Source/GeometryCollectionNodes/Private/Dataflow/GeometryCollectionFracturingNodes.cpp
+++ b/Engine/Plugins/Experimental/GeometryCollectionPlugin/Source/GeometryCollectionNodes/Private/Dataflow/GeometryCollectionFracturingNodes.cpp
@@ -807,7 +807,6 @@ void FMeshCutterDataflowNode::Evaluate(UE::Dataflow::FContext& Context, const FD
 		Out->IsA(&TransformSelection) ||
 		Out->IsA(&NewGeometryTransformSelection))
 	{
-#if WITH_EDITORONLY_DATA
 		FDataflowTransformSelection InTransformSelection = GetValue(Context, &TransformSelection);
 		//
 		// If not connected select everything by default
@@ -844,8 +843,12 @@ void FMeshCutterDataflowNode::Evaluate(UE::Dataflow::FContext& Context, const FD
 		{
 			FManagedArrayCollection InCollection = GetValue(Context, &Collection);
 
+			FDynamicMesh3 LocalMesh; // can be used for local storage of a converted static mesh
+			TArray<const FDynamicMesh3*> UseMeshes;
+
 			if (TObjectPtr<UStaticMesh> InCuttingMesh = GetValue(Context, &CuttingStaticMesh))
 			{
+#if WITH_EDITORONLY_DATA
 				if (FMeshDescription* MeshDescription = bUseHiRes ? InCuttingMesh->GetHiResMeshDescription() : InCuttingMesh->GetMeshDescription(LODLevel))
 				{
 					// If HiRes is empty then use LoRes
@@ -856,110 +859,122 @@ void FMeshCutterDataflowNode::Evaluate(UE::Dataflow::FContext& Context, const FD
 
 					if (MeshDescription->Vertices().Num() > 0)
 					{
-						TObjectPtr<UDynamicMesh> NewMesh = NewObject<UDynamicMesh>();
-						NewMesh->Reset();
+						FMeshDescriptionToDynamicMesh ConverterToDynamicMesh;
+						ConverterToDynamicMesh.Convert(MeshDescription, LocalMesh);
+						UseMeshes.Add(&LocalMesh);
+					}
 
-						UE::Geometry::FDynamicMesh3& DynCuttingMesh = NewMesh->GetMeshRef();
-						{
-							FMeshDescriptionToDynamicMesh ConverterToDynamicMesh;
-							ConverterToDynamicMesh.Convert(MeshDescription, DynCuttingMesh);
-						}
+				}
+#else
+				// TODO: for runtime usage, could try to fallback to the render mesh (if available on CPU)
+				ensureMsgf(false, TEXT("FMeshCutterDataflowNode's Static Mesh support is currently editor-only."));
+#endif
+			}
 
-						const int32 InRandomSeed = GetValue(Context, &RandomSeed);
-						const int32 InNumberToScatter = GetValue(Context, &NumberToScatter);
-						const int32 InGridX = GetValue(Context, &GridX);
-						const int32 InGridY = GetValue(Context, &GridY);
-						const int32 InGridZ = GetValue(Context, &GridZ);
-						const float InVariability = GetValue(Context, &Variability);
-						const float InMinScaleFactor = GetValue(Context, &MinScaleFactor);
-						const float InMaxScaleFactor = GetValue(Context, &MaxScaleFactor);
-						const float InRollRange = GetValue(Context, &RollRange);
-						const float InPitchRange = GetValue(Context, &PitchRange);
-						const float InYawRange = GetValue(Context, &YawRange);
-						const FTransform InTransform = GetValue(Context, &Transform);
-						const float InChanceToFracture = GetValue(Context, &ChanceToFracture);
-						const float InCollisionSampleSpacing = GetValue(Context, &CollisionSampleSpacing);
-
-						TArray<FTransform> MeshTransforms;
-
-						if (CutDistribution == EMeshCutterCutDistribution::SingleCut)
-						{
-							MeshTransforms.Add(InTransform);
-						}
-						else
-						{
-							FFractureEngineFracturing::GenerateMeshTransforms(MeshTransforms,
-								InBoundingBox,
-								InRandomSeed,
-								CutDistribution,
-								InNumberToScatter,
-								InGridX,
-								InGridY,
-								InGridZ,
-								InVariability,
-								InMinScaleFactor,
-								InMaxScaleFactor,
-								bRandomOrientation,
-								InRollRange,
-								InPitchRange,
-								InYawRange);
-						}
+			const TArray<TObjectPtr<UDynamicMesh>>& InDynamicMeshes = GetValue(Context, &CuttingDynamicMeshes);
+			for (TObjectPtr<UDynamicMesh> MeshObj : InDynamicMeshes)
+			{
+				if (MeshObj && MeshObj->GetMeshPtr())
+				{
+					UseMeshes.Add(MeshObj->GetMeshPtr());
+				}
+			}
+
+			if (UseMeshes.Num() > 0)
+			{
+				const int32 InRandomSeed = GetValue(Context, &RandomSeed);
+				const int32 InNumberToScatter = GetValue(Context, &NumberToScatter);
+				const int32 InGridX = GetValue(Context, &GridX);
+				const int32 InGridY = GetValue(Context, &GridY);
+				const int32 InGridZ = GetValue(Context, &GridZ);
+				const float InVariability = GetValue(Context, &Variability);
+				const float InMinScaleFactor = GetValue(Context, &MinScaleFactor);
+				const float InMaxScaleFactor = GetValue(Context, &MaxScaleFactor);
+				const float InRollRange = GetValue(Context, &RollRange);
+				const float InPitchRange = GetValue(Context, &PitchRange);
+				const float InYawRange = GetValue(Context, &YawRange);
+				const FTransform InTransform = GetValue(Context, &Transform);
+				const float InChanceToFracture = GetValue(Context, &ChanceToFracture);
+				const float InCollisionSampleSpacing = GetValue(Context, &CollisionSampleSpacing);
+
+				// Note: per-cut mesh selection is not currently a dataflow input
+				const EMeshCutterPerCutMeshSelection InPerCutMeshSelection = PerCutMeshSelection;
+
+				TArray<FTransform> MeshTransforms;
+
+				if (CutDistribution == EMeshCutterCutDistribution::SingleCut)
+				{
+					MeshTransforms.Add(InTransform);
+				}
+				else
+				{
+					FFractureEngineFracturing::GenerateMeshTransforms(MeshTransforms,
+						InBoundingBox,
+						InRandomSeed,
+						CutDistribution,
+						InNumberToScatter,
+						InGridX,
+						InGridY,
+						InGridZ,
+						InVariability,
+						InMinScaleFactor,
+						InMaxScaleFactor,
+						bRandomOrientation,
+						InRollRange,
+						InPitchRange,
+						InYawRange);
+				}
 
-						int32 ResultGeometryIndex = FFractureEngineFracturing::MeshCutter(MeshTransforms,
-							InCollection,
-							InTransformSelection,
-							DynCuttingMesh,
-							InTransform,
-							InRandomSeed,
-							InChanceToFracture,
-							SplitIslands,
-							InCollisionSampleSpacing);
+				int32 ResultGeometryIndex = FFractureEngineFracturing::MeshArrayCutter(MeshTransforms,
+					InCollection,
+					InTransformSelection,
+					UseMeshes,
+					InPerCutMeshSelection,
+					InRandomSeed,
+					InChanceToFracture,
+					SplitIslands,
+					InCollisionSampleSpacing);
 
-						FDataflowTransformSelection NewSelection;
-						FDataflowTransformSelection OriginalSelection;
+				FDataflowTransformSelection NewSelection;
+				FDataflowTransformSelection OriginalSelection;
 
-						if (ResultGeometryIndex != INDEX_NONE)
+				if (ResultGeometryIndex != INDEX_NONE)
+				{
+					if (InCollection.HasAttribute("TransformIndex", FGeometryCollection::GeometryGroup))
+					{
+						const TManagedArray<int32>& TransformIndices = InCollection.GetAttribute<int32>("TransformIndex", FGeometryCollection::GeometryGroup);
+
+						NewSelection.Initialize(TransformIndices.Num(), false);
+						OriginalSelection.Initialize(TransformIndices.Num(), false);
+
+						// The newly fractured pieces are added to the end of the transform array (starting position is ResultGeometryIndex)
+						for (int32 Idx = ResultGeometryIndex; Idx < TransformIndices.Num(); ++Idx)
 						{
-							if (InCollection.HasAttribute("TransformIndex", FGeometryCollection::GeometryGroup))
+							int32 BoneIdx = TransformIndices[Idx];
+							NewSelection.SetSelected(BoneIdx);
+						}
+
+						for (int32 Idx = 0; Idx < InTransformSelection.Num(); ++Idx)
+						{
+							if (InTransformSelection.IsSelected(Idx))
 							{
-								const TManagedArray<int32>& TransformIndices = InCollection.GetAttribute<int32>("TransformIndex", FGeometryCollection::GeometryGroup);
-
-								NewSelection.Initialize(TransformIndices.Num(), false);
-								OriginalSelection.Initialize(TransformIndices.Num(), false);
-
-								// The newly fractured pieces are added to the end of the transform array (starting position is ResultGeometryIndex)
-								for (int32 Idx = ResultGeometryIndex; Idx < TransformIndices.Num(); ++Idx)
-								{
-									int32 BoneIdx = TransformIndices[Idx];
-									NewSelection.SetSelected(BoneIdx);
-								}
-
-								for (int32 Idx = 0; Idx < InTransformSelection.Num(); ++Idx)
-								{
-									if (InTransformSelection.IsSelected(Idx))
-									{
-										OriginalSelection.SetSelected(Idx);
-									}
-								}
+								OriginalSelection.SetSelected(Idx);
 							}
 						}
-
-						SetValue(Context, MoveTemp(InCollection), &Collection);
-						SetValue(Context, OriginalSelection, &TransformSelection);
-						SetValue(Context, NewSelection, &NewGeometryTransformSelection);
-
-						return;
 					}
 				}
+
+				SetValue(Context, MoveTemp(InCollection), &Collection);
+				SetValue(Context, OriginalSelection, &TransformSelection);
+				SetValue(Context, NewSelection, &NewGeometryTransformSelection);
+
+				return;
 			}
 		}
 
 		SafeForwardInput(Context, &Collection, &Collection);
 		SetValue(Context, InTransformSelection, &TransformSelection);
 		SetValue(Context, FDataflowTransformSelection(), &NewGeometryTransformSelection);
-#else
-		ensureMsgf(false, TEXT("FMeshCutterDataflowNode is an editor only node."));
-#endif
 	}
 }
 
diff --git a/Engine/Plugins/Experimental/GeometryCollectionPlugin/Source/GeometryCollectionNodes/Public/Dataflow/GeometryCollectionFracturingNodes.h b/Engine/Plugins/Experimental/GeometryCollectionPlugin/Source/GeometryCollectionNodes/Public/Dataflow/GeometryCollectionFracturingNodes.h
--- a/Engine/Plugins/Experimental/GeometryCollectionPlugin/Source/GeometryCollectionNodes/Public/Dataflow/GeometryCollectionFracturingNodes.h
+++ b/Engine/Plugins/Experimental/GeometryCollectionPlugin/Source/GeometryCollectionNodes/Public/Dataflow/GeometryCollectionFracturingNodes.h
@@ -18,6 +18,7 @@ namespace UE_DEPRECATED(5.5, "Use UE::Dataflow instead.") Dataflow {}
 
 class FGeometryCollection;
 class UStaticMesh;
+class UDynamicMesh;
 
 USTRUCT(meta = (DataflowGeometryCollection, Deprecated = "5.5"))
 struct FUniformScatterPointsDataflowNode : public FDataflowNode
@@ -1034,7 +1035,7 @@ public:
 /**
  * 
  * Editor Fracture Mode / Fracture / Mesh tool
- * Fracture using the shape of a chosen static mesh
+ * Fracture using the shape of a chosen static mesh and/or array of dynamic meshes
  * 
  */
 USTRUCT(meta = (DataflowGeometryCollection))
@@ -1060,15 +1061,19 @@ public:
 	UPROPERTY(meta = (DataflowInput, DisplayName = "Transform"))
 	FTransform Transform = FTransform::Identity;
 
-	/** Mesh to cut with */
+	/** Dynamic Meshes to cut with */
+	UPROPERTY(EditAnywhere, Category = "Mesh", meta = (DataflowInput, DataflowIntrinsic))
+	TArray<TObjectPtr<UDynamicMesh>> CuttingDynamicMeshes;
+
+	/** Static Mesh to cut with */
 	UPROPERTY(EditAnywhere, Category = "StaticMesh", meta = (DataflowInput, DataflowIntrinsic))
 	TObjectPtr<UStaticMesh> CuttingStaticMesh;
 
-	/** Output the HiRes representation, if set to true and HiRes doesn't exist it will output empty mesh */
+	/** If using a Static Mesh to cut, attempt to use the Nanite HiRes source mesh, if available and non-empty. */
 	UPROPERTY(EditAnywhere, Category = "StaticMesh", meta = (DisplayName = "Use HiRes"));
 	bool bUseHiRes = false;
 
-	/** Specifies the LOD level to use */
+	/** If using a Static Mesh to cut, and not using the Nanite HiRes source mesh, use this LOD level's mesh */
 	UPROPERTY(EditAnywhere, Category = "StaticMesh", meta = (DisplayName = "LOD Level"));
 	int32 LODLevel = 0;
 
@@ -1076,6 +1081,10 @@ public:
 	UPROPERTY(EditAnywhere, Category = Distribution)
 	EMeshCutterCutDistribution CutDistribution = EMeshCutterCutDistribution::SingleCut;
 
+	/** When there are multiple cutting meshes, how to choose the cut mesh to apply at each location */
+	UPROPERTY(EditAnywhere, Category = Distribution)
+	EMeshCutterPerCutMeshSelection PerCutMeshSelection = EMeshCutterPerCutMeshSelection::All;
+
 	/** Number of meshes to random scatter */
 	UPROPERTY(EditAnywhere, Category = Distribution, meta = (DataflowInput, ClampMin = "1", UIMax = "5000", EditCondition = "CutDistribution == EMeshCutterCutDistribution::UniformRandom", EditConditionHides))
 	int32 NumberToScatter = 10;
@@ -1160,6 +1169,7 @@ public:
 		RegisterInputConnection(&TransformSelection);
 		RegisterInputConnection(&Transform) .SetCanHidePin(true).SetPinIsHidden(true);
 		RegisterInputConnection(&CuttingStaticMesh) .SetCanHidePin(true).SetPinIsHidden(true);
+		RegisterInputConnection(&CuttingDynamicMeshes) .SetCanHidePin(true).SetPinIsHidden(true);
 		RegisterInputConnection(&NumberToScatter) .SetCanHidePin(true).SetPinIsHidden(true);
 		RegisterInputConnection(&GridX) .SetCanHidePin(true).SetPinIsHidden(true);
 		RegisterInputConnection(&GridY) .SetCanHidePin(true).SetPinIsHidden(true);
