also expose 'allow merging leaf hulls' option to the 'generate convex hulls from children' (matching how it's exposed to the generate hulls from leaves node)

diff --git a/Engine/Plugins/Experimental/GeometryCollectionPlugin/Source/GeometryCollectionNodes/Private/Dataflow/GeometryCollectionUtilityNodes.cpp b/Engine/Plugins/Experimental/GeometryCollectionPlugin/Source/GeometryCollectionNodes/Private/Dataflow/GeometryCollectionUtilityNodes.cpp
--- a/Engine/Plugins/Experimental/GeometryCollectionPlugin/Source/GeometryCollectionNodes/Private/Dataflow/GeometryCollectionUtilityNodes.cpp
+++ b/Engine/Plugins/Experimental/GeometryCollectionPlugin/Source/GeometryCollectionNodes/Private/Dataflow/GeometryCollectionUtilityNodes.cpp
@@ -309,6 +309,7 @@ FGenerateClusterConvexHullsFromChildrenHullsDataflowNode::FGenerateClusterConvex
 	RegisterInputConnection(&ConvexCount);
 	RegisterInputConnection(&ErrorTolerance);
 	RegisterInputConnection(&OptionalSelectionFilter);
+	RegisterInputConnection(&bAllowMergingLeafHulls);
 	RegisterInputConnection(&bProtectNegativeSpace);
 	RegisterInputConnection(&TargetNumSamples);
 	RegisterInputConnection(&MinSampleSpacing);
@@ -359,6 +360,7 @@ void FGenerateClusterConvexHullsFromChildrenHullsDataflowNode::Evaluate(UE::Data
 			FGeometryCollectionConvexUtility::FClusterConvexHullSettings HullMergeSettings(InConvexCount, InErrorToleranceInCm, bPreferExternalCollisionShapes);
 			HullMergeSettings.AllowMergesMethod = EAllowConvexMergeMethod::Any; // Note: Only 'Any' is supported for this node currently
 			HullMergeSettings.EmptySpace = bHasNegativeSpace ? &NegativeSpace : nullptr;
+			HullMergeSettings.bAllowMergingLeafHulls = GetValue(Context, &bAllowMergingLeafHulls);
 
 			if (bHasSelectionFilter)
 			{
diff --git a/Engine/Plugins/Experimental/GeometryCollectionPlugin/Source/GeometryCollectionNodes/Public/Dataflow/GeometryCollectionUtilityNodes.h b/Engine/Plugins/Experimental/GeometryCollectionPlugin/Source/GeometryCollectionNodes/Public/Dataflow/GeometryCollectionUtilityNodes.h
--- a/Engine/Plugins/Experimental/GeometryCollectionPlugin/Source/GeometryCollectionNodes/Public/Dataflow/GeometryCollectionUtilityNodes.h
+++ b/Engine/Plugins/Experimental/GeometryCollectionPlugin/Source/GeometryCollectionNodes/Public/Dataflow/GeometryCollectionUtilityNodes.h
@@ -377,6 +377,10 @@ public:
 	UPROPERTY(meta = (DataflowInput))
 	FDataflowTransformSelection OptionalSelectionFilter;
 
+	/** Also allow the same hull merging process to run on leaf hulls (merging hulls on leaves in the selection) */
+	UPROPERTY(EditAnywhere, Category = "Convex", meta = (DataflowInput))
+	bool bAllowMergingLeafHulls = false;
+
 	/** Whether to use a sphere cover to define negative space that should not be covered by convex hulls */
 	UPROPERTY(EditAnywhere, Category = NegativeSpace, meta = (DataflowInput))
 	bool bProtectNegativeSpace = false;
