Add geometry processing dataflow plugin and begin to move relevant existing nodes there (starting with just the mesh boolean)

diff --git a/Engine/Plugins/Experimental/GeometryCollectionPlugin/GeometryCollectionPlugin.uplugin b/Engine/Plugins/Experimental/GeometryCollectionPlugin/GeometryCollectionPlugin.uplugin
--- a/Engine/Plugins/Experimental/GeometryCollectionPlugin/GeometryCollectionPlugin.uplugin
+++ b/Engine/Plugins/Experimental/GeometryCollectionPlugin/GeometryCollectionPlugin.uplugin
@@ -62,6 +62,10 @@
 			"Name": "GeometryProcessing",
 			"Enabled": true
 		},
+		{
+			"Name": "GeometryDataflow",
+			"Enabled": true
+		},
 		{
 			"Name": "Fracture",
 			"Enabled": true
diff --git a/Engine/Plugins/Experimental/GeometryCollectionPlugin/Source/GeometryCollectionNodes/GeometryCollectionNodes.Build.cs b/Engine/Plugins/Experimental/GeometryCollectionPlugin/Source/GeometryCollectionNodes/GeometryCollectionNodes.Build.cs
--- a/Engine/Plugins/Experimental/GeometryCollectionPlugin/Source/GeometryCollectionNodes/GeometryCollectionNodes.Build.cs
+++ b/Engine/Plugins/Experimental/GeometryCollectionPlugin/Source/GeometryCollectionNodes/GeometryCollectionNodes.Build.cs
@@ -19,6 +19,7 @@ namespace UnrealBuildTool.Rules
 					"DataflowNodes",
 					"GeometryCollectionEngine",
 					"GeometryCore",
+					"GeometryDataflowNodes",
 					"Voronoi",
 					"GeometryFramework",
 					"MeshDescription",
diff --git a/Engine/Plugins/Experimental/GeometryCollectionPlugin/Source/GeometryCollectionNodes/Private/Dataflow/GeometryCollectionMeshNodes.cpp b/Engine/Plugins/Experimental/GeometryCollectionPlugin/Source/GeometryCollectionNodes/Private/Dataflow/GeometryCollectionMeshNodes.cpp
--- a/Engine/Plugins/Experimental/GeometryCollectionPlugin/Source/GeometryCollectionNodes/Private/Dataflow/GeometryCollectionMeshNodes.cpp
+++ b/Engine/Plugins/Experimental/GeometryCollectionPlugin/Source/GeometryCollectionNodes/Private/Dataflow/GeometryCollectionMeshNodes.cpp
@@ -20,7 +20,6 @@
 #include "MeshDescriptionToDynamicMesh.h"
 #include "StaticMeshAttributes.h"
 #include "DynamicMeshEditor.h"
-#include "Operations/MeshBoolean.h"
 
 #include "GeometryCollectionToDynamicMesh.h"
 
@@ -51,7 +50,6 @@ namespace UE::Dataflow
 		DATAFLOW_NODE_REGISTER_CREATION_FACTORY(FCollectionToMeshDataflowNode);
 		DATAFLOW_NODE_REGISTER_CREATION_FACTORY(FStaticMeshToMeshDataflowNode);
 		DATAFLOW_NODE_REGISTER_CREATION_FACTORY(FMeshAppendDataflowNode);
-		DATAFLOW_NODE_REGISTER_CREATION_FACTORY(FMeshBooleanDataflowNode);
 		DATAFLOW_NODE_REGISTER_CREATION_FACTORY(FMeshCopyToPointsDataflowNode);
 		DATAFLOW_NODE_REGISTER_CREATION_FACTORY(FGetMeshDataDataflowNode);
 		DATAFLOW_NODE_REGISTER_CREATION_FACTORY(FApplyMeshProcessorToMeshDataflowNode);
@@ -318,59 +316,6 @@ void FMeshAppendDataflowNode::Evaluate(UE::Dataflow::FContext& Context, const FD
 	}
 }
 
-
-void FMeshBooleanDataflowNode::Evaluate(UE::Dataflow::FContext& Context, const FDataflowOutput* Out) const
-{
-	if (Out->IsA<TObjectPtr<UDynamicMesh>>(&Mesh))
-	{
-		if (TObjectPtr<UDynamicMesh> InMesh1 = GetValue<TObjectPtr<UDynamicMesh>>(Context, &Mesh1))
-		{
-			if (TObjectPtr<UDynamicMesh> InMesh2 = GetValue<TObjectPtr<UDynamicMesh>>(Context, &Mesh2))
-			{
-				const UE::Geometry::FDynamicMesh3& DynMesh1 = InMesh1->GetMeshRef();
-				const UE::Geometry::FDynamicMesh3& DynMesh2 = InMesh2->GetMeshRef();
-
-				if (DynMesh1.VertexCount() > 0 && DynMesh2.VertexCount() > 0)
-				{
-					// Get output
-					TObjectPtr<UDynamicMesh> NewMesh = NewObject<UDynamicMesh>();
-					NewMesh->Reset();
-					UE::Geometry::FDynamicMesh3& ResultDynMesh = NewMesh->GetMeshRef();
-
-					UE::Geometry::FMeshBoolean::EBooleanOp BoolOp = UE::Geometry::FMeshBoolean::EBooleanOp::Intersect;
-					if (Operation == EMeshBooleanOperationEnum::Dataflow_MeshBoolean_Intersect)
-					{
-						BoolOp = UE::Geometry::FMeshBoolean::EBooleanOp::Intersect;
-					}
-					else if (Operation == EMeshBooleanOperationEnum::Dataflow_MeshBoolean_Union)
-					{
-						BoolOp = UE::Geometry::FMeshBoolean::EBooleanOp::Union;
-					}
-					else if (Operation == EMeshBooleanOperationEnum::Dataflow_MeshBoolean_Difference)
-					{
-						BoolOp = UE::Geometry::FMeshBoolean::EBooleanOp::Difference;
-					}
-
-					UE::Geometry::FMeshBoolean Boolean(&DynMesh1, &DynMesh2, &ResultDynMesh, BoolOp);
-					Boolean.bSimplifyAlongNewEdges = true;
-					Boolean.PreserveUVsOnlyForMesh = 0; // slight warping of the autogenerated cell UVs generally doesn't matter
-					Boolean.bWeldSharedEdges = false;
-					Boolean.bTrackAllNewEdges = true;
-
-					if (Boolean.Compute())
-					{
-						SetValue(Context, NewMesh, &Mesh);
-
-						return;
-					}
-				}
-			}
-		}
-
-		SetValue(Context, TObjectPtr<UDynamicMesh>(NewObject<UDynamicMesh>()), &Mesh);
-	}
-}
-
 void FMeshCopyToPointsDataflowNode::Evaluate(UE::Dataflow::FContext& Context, const FDataflowOutput* Out) const
 {
 	if (Out->IsA<TObjectPtr<UDynamicMesh>>(&Mesh))
diff --git a/Engine/Plugins/Experimental/GeometryCollectionPlugin/Source/GeometryCollectionNodes/Public/Dataflow/GeometryCollectionMeshNodes.h b/Engine/Plugins/Experimental/GeometryCollectionPlugin/Source/GeometryCollectionNodes/Public/Dataflow/GeometryCollectionMeshNodes.h
--- a/Engine/Plugins/Experimental/GeometryCollectionPlugin/Source/GeometryCollectionNodes/Public/Dataflow/GeometryCollectionMeshNodes.h
+++ b/Engine/Plugins/Experimental/GeometryCollectionPlugin/Source/GeometryCollectionNodes/Public/Dataflow/GeometryCollectionMeshNodes.h
@@ -11,6 +11,11 @@
 
 #include "GeometryCollectionMeshNodes.generated.h"
 
+#if UE_ENABLE_INCLUDE_ORDER_DEPRECATED_IN_5_6
+// Some mesh-specific nodes have been moved to their own plugin
+#include "Dataflow/MeshBooleanNodes.h"
+#endif
+
 #if UE_ENABLE_INCLUDE_ORDER_DEPRECATED_IN_5_5
 namespace Dataflow = UE::Dataflow;
 #else
@@ -477,58 +482,6 @@ public:
 };
 
 
-UENUM(BlueprintType)
-enum class EMeshBooleanOperationEnum : uint8
-{
-	Dataflow_MeshBoolean_Union UMETA(DisplayName = "Union"),
-	Dataflow_MeshBoolean_Intersect UMETA(DisplayName = "Intersect"),
-	Dataflow_MeshBoolean_Difference UMETA(DisplayName = "Difference"),
-	//~~~
-	//256th entry
-	Dataflow_Max                UMETA(Hidden)
-};
-
-/**
- *
- * Mesh boolean (Union, Intersect, Difference) between two meshes
- *
- */
-USTRUCT(meta = (DataflowGeometryCollection))
-struct FMeshBooleanDataflowNode : public FDataflowNode
-{
-	GENERATED_USTRUCT_BODY()
-	DATAFLOW_NODE_DEFINE_INTERNAL(FMeshBooleanDataflowNode, "MeshBoolean", "Mesh|Utilities", "")
-	DATAFLOW_NODE_RENDER_TYPE("SurfaceRender",FName("FDynamicMesh3"), "Mesh")
-
-public:
-	/** Boolean operation */
-	UPROPERTY(EditAnywhere, Category = "Boolean");
-	EMeshBooleanOperationEnum Operation = EMeshBooleanOperationEnum::Dataflow_MeshBoolean_Intersect;
-
-	/** Mesh input */
-	UPROPERTY(meta = (DataflowInput, DataflowIntrinsic))
-	TObjectPtr<UDynamicMesh> Mesh1;
-
-	/** Mesh input */
-	UPROPERTY(meta = (DataflowInput, DataflowIntrinsic))
-	TObjectPtr<UDynamicMesh> Mesh2;
-
-	/** Output mesh */
-	UPROPERTY(meta = (DataflowOutput))
-	TObjectPtr<UDynamicMesh> Mesh;
-
-	FMeshBooleanDataflowNode(const UE::Dataflow::FNodeParameters& InParam, FGuid InGuid = FGuid::NewGuid())
-		: FDataflowNode(InParam, InGuid)
-	{
-		RegisterInputConnection(&Mesh1);
-		RegisterInputConnection(&Mesh2);
-		RegisterOutputConnection(&Mesh);
-	}
-
-	virtual void Evaluate(UE::Dataflow::FContext& Context, const FDataflowOutput* Out) const override;
-
-};
-
 
 /**
  *
diff --git a/Engine/Plugins/Experimental/GeometryDataflow/GeometryDataflow.uplugin b/Engine/Plugins/Experimental/GeometryDataflow/GeometryDataflow.uplugin
new file mode 100644
--- /dev/null
+++ b/Engine/Plugins/Experimental/GeometryDataflow/GeometryDataflow.uplugin
@@ -0,0 +1,34 @@
+{
+    "FileVersion" : 3,
+    "Version" : 1,
+    "VersionName" : "0.1",
+    "FriendlyName" : "Geometry Dataflow Nodes",
+    "Description" : "Geometry Processing in Dataflow.",
+    "Category": "Geometry",
+    "CreatedBy" : "Epic Games, Inc.",
+    "CreatedByURL" : "https://epicgames.com",
+    "DocsURL" : "",
+    "MarketplaceURL" : "",
+    "SupportURL" : "",
+    "EnabledByDefault" : false,
+    "CanContainContent" : true,
+    "IsBetaVersion" : true,
+    "Installed" : false,
+	"Modules": [
+		{
+			"Name": "GeometryDataflowNodes",
+			"Type": "Runtime",
+			"LoadingPhase": "Default"
+		}
+	],
+	"Plugins": [
+		{
+			"Name": "Dataflow",
+			"Enabled": true
+		},
+		{
+			"Name": "GeometryProcessing",
+			"Enabled": true
+		}
+	]
+}
diff --git a/Engine/Plugins/Experimental/GeometryDataflow/Source/GeometryDataflowNodes/GeometryDataflowNodes.Build.cs b/Engine/Plugins/Experimental/GeometryDataflow/Source/GeometryDataflowNodes/GeometryDataflowNodes.Build.cs
new file mode 100644
--- /dev/null
+++ b/Engine/Plugins/Experimental/GeometryDataflow/Source/GeometryDataflowNodes/GeometryDataflowNodes.Build.cs
@@ -0,0 +1,29 @@
+// Copyright Epic Games, Inc. All Rights Reserved.
+
+namespace UnrealBuildTool.Rules
+{
+	public class GeometryDataflowNodes : ModuleRules
+	{
+        public GeometryDataflowNodes(ReadOnlyTargetRules Target) : base(Target)
+		{
+			PublicDependencyModuleNames.AddRange(
+				new string[]
+				{
+					"Core",
+					"CoreUObject",
+					"DataflowCore",
+					"DataflowNodes",
+					"GeometryCore",
+					"GeometryFramework" // for UDynamicMesh
+				}
+			);
+
+			PrivateDependencyModuleNames.AddRange(
+				new string[] 
+				{
+					"DynamicMesh"
+				}
+			);
+		}
+	}
+}
diff --git a/Engine/Plugins/Experimental/GeometryDataflow/Source/GeometryDataflowNodes/Private/Dataflow/MeshBooleanNodes.cpp b/Engine/Plugins/Experimental/GeometryDataflow/Source/GeometryDataflowNodes/Private/Dataflow/MeshBooleanNodes.cpp
new file mode 100644
--- /dev/null
+++ b/Engine/Plugins/Experimental/GeometryDataflow/Source/GeometryDataflowNodes/Private/Dataflow/MeshBooleanNodes.cpp
@@ -0,0 +1,82 @@
+// Copyright Epic Games, Inc. All Rights Reserved.
+
+#include "Dataflow/MeshBooleanNodes.h"
+#include "Dataflow/DataflowCore.h"
+
+#include "Operations/MeshBoolean.h"
+
+#include "DynamicMesh/DynamicMesh3.h"
+#include "UDynamicMesh.h"
+
+#include UE_INLINE_GENERATED_CPP_BY_NAME(MeshBooleanNodes)
+
+namespace UE::Dataflow
+{
+
+	void MeshBooleanNodes()
+	{
+		DATAFLOW_NODE_REGISTER_CREATION_FACTORY(FMeshBooleanDataflowNode);
+	}
+}
+
+
+FMeshBooleanDataflowNode::FMeshBooleanDataflowNode(const UE::Dataflow::FNodeParameters& InParam, FGuid InGuid)
+	: FDataflowNode(InParam, InGuid)
+{
+	RegisterInputConnection(&Mesh1);
+	RegisterInputConnection(&Mesh2);
+	RegisterOutputConnection(&Mesh);
+}
+
+void FMeshBooleanDataflowNode::Evaluate(UE::Dataflow::FContext& Context, const FDataflowOutput* Out) const
+{
+	if (Out->IsA(&Mesh))
+	{
+		if (const TObjectPtr<UDynamicMesh> InMesh1 = GetValue(Context, &Mesh1))
+		{
+			if (const TObjectPtr<UDynamicMesh> InMesh2 = GetValue(Context, &Mesh2))
+			{
+				const UE::Geometry::FDynamicMesh3& DynMesh1 = InMesh1->GetMeshRef();
+				const UE::Geometry::FDynamicMesh3& DynMesh2 = InMesh2->GetMeshRef();
+
+				if (DynMesh1.VertexCount() > 0 && DynMesh2.VertexCount() > 0)
+				{
+					// Get output
+					TObjectPtr<UDynamicMesh> NewMesh = NewObject<UDynamicMesh>();
+					NewMesh->Reset();
+					UE::Geometry::FDynamicMesh3& ResultDynMesh = NewMesh->GetMeshRef();
+
+					UE::Geometry::FMeshBoolean::EBooleanOp BoolOp = UE::Geometry::FMeshBoolean::EBooleanOp::Intersect;
+					if (Operation == EMeshBooleanOperationEnum::Dataflow_MeshBoolean_Intersect)
+					{
+						BoolOp = UE::Geometry::FMeshBoolean::EBooleanOp::Intersect;
+					}
+					else if (Operation == EMeshBooleanOperationEnum::Dataflow_MeshBoolean_Union)
+					{
+						BoolOp = UE::Geometry::FMeshBoolean::EBooleanOp::Union;
+					}
+					else if (Operation == EMeshBooleanOperationEnum::Dataflow_MeshBoolean_Difference)
+					{
+						BoolOp = UE::Geometry::FMeshBoolean::EBooleanOp::Difference;
+					}
+
+					UE::Geometry::FMeshBoolean Boolean(&DynMesh1, &DynMesh2, &ResultDynMesh, BoolOp);
+					Boolean.bSimplifyAlongNewEdges = true;
+					Boolean.PreserveUVsOnlyForMesh = 0; // slight warping of the autogenerated cell UVs generally doesn't matter
+					Boolean.bWeldSharedEdges = false;
+					Boolean.bTrackAllNewEdges = true;
+
+					if (Boolean.Compute())
+					{
+						SetValue(Context, NewMesh, &Mesh);
+
+						return;
+					}
+				}
+			}
+		}
+
+		SetValue(Context, TObjectPtr<UDynamicMesh>(NewObject<UDynamicMesh>()), &Mesh);
+	}
+}
+
diff --git a/Engine/Plugins/Experimental/GeometryDataflow/Source/GeometryDataflowNodes/Private/GeometryDataflowNodesPlugin.cpp b/Engine/Plugins/Experimental/GeometryDataflow/Source/GeometryDataflowNodes/Private/GeometryDataflowNodesPlugin.cpp
new file mode 100644
--- /dev/null
+++ b/Engine/Plugins/Experimental/GeometryDataflow/Source/GeometryDataflowNodes/Private/GeometryDataflowNodesPlugin.cpp
@@ -0,0 +1,25 @@
+// Copyright Epic Games, Inc. All Rights Reserved.
+
+#include "GeometryDataflowNodesPlugin.h"
+
+#include "CoreMinimal.h"
+
+#include "Dataflow/MeshBooleanNodes.h"
+
+#define LOCTEXT_NAMESPACE "DataflowNodes"
+
+
+void IGeometryDataflowNodesPlugin::StartupModule()
+{
+	UE::Dataflow::MeshBooleanNodes();
+}
+
+void IGeometryDataflowNodesPlugin::ShutdownModule()
+{
+}
+
+
+IMPLEMENT_MODULE(IGeometryDataflowNodesPlugin, GeometryDataflowNodes)
+
+
+#undef LOCTEXT_NAMESPACE
diff --git a/Engine/Plugins/Experimental/GeometryDataflow/Source/GeometryDataflowNodes/Public/Dataflow/MeshBooleanNodes.h b/Engine/Plugins/Experimental/GeometryDataflow/Source/GeometryDataflowNodes/Public/Dataflow/MeshBooleanNodes.h
new file mode 100644
--- /dev/null
+++ b/Engine/Plugins/Experimental/GeometryDataflow/Source/GeometryDataflowNodes/Public/Dataflow/MeshBooleanNodes.h
@@ -0,0 +1,71 @@
+// Copyright Epic Games, Inc. All Rights Reserved.
+
+#pragma once 
+
+#include "CoreMinimal.h"
+#include "Dataflow/DataflowEngine.h"
+#include "Dataflow/DataflowSelection.h"
+#include "DynamicMesh/DynamicMesh3.h"
+#include "DynamicMeshProcessor.h"
+
+#include "MeshBooleanNodes.generated.h"
+
+class UDynamicMesh;
+class UMaterialInterface;
+
+
+UENUM(BlueprintType)
+enum class EMeshBooleanOperationEnum : uint8
+{
+	Dataflow_MeshBoolean_Union UMETA(DisplayName = "Union"),
+	Dataflow_MeshBoolean_Intersect UMETA(DisplayName = "Intersect"),
+	Dataflow_MeshBoolean_Difference UMETA(DisplayName = "Difference"),
+	//~~~
+	//256th entry
+	Dataflow_Max                UMETA(Hidden)
+};
+
+/**
+ *
+ * Mesh boolean (Union, Intersect, Difference) between two meshes
+ *
+ */
+USTRUCT()
+struct FMeshBooleanDataflowNode : public FDataflowNode
+{
+	GENERATED_USTRUCT_BODY()
+	DATAFLOW_NODE_DEFINE_INTERNAL(FMeshBooleanDataflowNode, "MeshBoolean", "Mesh|Utilities", "")
+	DATAFLOW_NODE_RENDER_TYPE("SurfaceRender",FName("FDynamicMesh3"), "Mesh")
+
+public:
+
+	FMeshBooleanDataflowNode(const UE::Dataflow::FNodeParameters& InParam, FGuid InGuid = FGuid::NewGuid());
+
+private:
+	/** Boolean operation */
+	UPROPERTY(EditAnywhere, Category = "Boolean");
+	EMeshBooleanOperationEnum Operation = EMeshBooleanOperationEnum::Dataflow_MeshBoolean_Intersect;
+
+	/** Mesh input */
+	UPROPERTY(meta = (DataflowInput, DataflowIntrinsic))
+	TObjectPtr<UDynamicMesh> Mesh1;
+
+	/** Mesh input */
+	UPROPERTY(meta = (DataflowInput, DataflowIntrinsic))
+	TObjectPtr<UDynamicMesh> Mesh2;
+
+	/** Output mesh */
+	UPROPERTY(meta = (DataflowOutput))
+	TObjectPtr<UDynamicMesh> Mesh;
+
+	virtual void Evaluate(UE::Dataflow::FContext& Context, const FDataflowOutput* Out) const override;
+
+};
+
+
+
+namespace UE::Dataflow
+{
+	void MeshBooleanNodes();
+}
+
diff --git a/Engine/Plugins/Experimental/GeometryDataflow/Source/GeometryDataflowNodes/Public/GeometryDataflowNodesPlugin.h b/Engine/Plugins/Experimental/GeometryDataflow/Source/GeometryDataflowNodes/Public/GeometryDataflowNodesPlugin.h
new file mode 100644
--- /dev/null
+++ b/Engine/Plugins/Experimental/GeometryDataflow/Source/GeometryDataflowNodes/Public/GeometryDataflowNodesPlugin.h
@@ -0,0 +1,18 @@
+// Copyright Epic Games, Inc. All Rights Reserved.
+
+#pragma once
+
+#include "CoreMinimal.h"
+#include "Modules/ModuleInterface.h"
+#include "Modules/ModuleManager.h"
+
+/**
+ * The public interface to this module
+ */
+class IGeometryDataflowNodesPlugin : public IModuleInterface
+{
+public:
+	virtual void StartupModule();
+	virtual void ShutdownModule();
+};
+
