Dataflow:
- Deprecated math nodes

diff --git a/Engine/Plugins/Experimental/GeometryCollectionPlugin/Source/GeometryCollectionNodes/Public/Dataflow/GeometryCollectionMathNodes.h b/Engine/Plugins/Experimental/GeometryCollectionPlugin/Source/GeometryCollectionNodes/Public/Dataflow/GeometryCollectionMathNodes.h
--- a/Engine/Plugins/Experimental/GeometryCollectionPlugin/Source/GeometryCollectionNodes/Public/Dataflow/GeometryCollectionMathNodes.h
+++ b/Engine/Plugins/Experimental/GeometryCollectionPlugin/Source/GeometryCollectionNodes/Public/Dataflow/GeometryCollectionMathNodes.h
@@ -18,7 +18,7 @@ namespace UE_DEPRECATED(5.5, "Use UE::Dataflow instead.") Dataflow {}
 * 
 *
 */
-USTRUCT(meta = (DataflowGeometryCollection))
+USTRUCT(meta = (DataflowGeometryCollection, Deprecated = "5.6"))
 struct FAddDataflowNode : public FDataflowNode
 {
 	GENERATED_USTRUCT_BODY()
@@ -50,7 +50,7 @@ public:
 *
 *
 */
-USTRUCT(meta = (DataflowGeometryCollection))
+USTRUCT(meta = (DataflowGeometryCollection, Deprecated = "5.6"))
 struct FSubtractDataflowNode : public FDataflowNode
 {
 	GENERATED_USTRUCT_BODY()
@@ -82,7 +82,7 @@ public:
 *
 *
 */
-USTRUCT(meta = (DataflowGeometryCollection))
+USTRUCT(meta = (DataflowGeometryCollection, Deprecated = "5.6"))
 struct FMultiplyDataflowNode : public FDataflowNode
 {
 	GENERATED_USTRUCT_BODY()
@@ -114,7 +114,7 @@ public:
 *
 *
 */
-USTRUCT(meta = (DataflowGeometryCollection))
+USTRUCT(meta = (DataflowGeometryCollection, Deprecated = "5.6"))
 struct FSafeDivideDataflowNode : public FDataflowNode
 {
 	GENERATED_USTRUCT_BODY()
@@ -141,7 +141,7 @@ public:
 	virtual void Evaluate(UE::Dataflow::FContext& Context, const FDataflowOutput* Out) const override;
 };
 
-USTRUCT(meta = (DataflowGeometryCollection))
+USTRUCT(meta = (DataflowGeometryCollection, Deprecated = "5.6"))
 struct FDivideDataflowNode : public FSafeDivideDataflowNode
 {
 	GENERATED_USTRUCT_BODY()
@@ -194,7 +194,7 @@ public:
 *
 *
 */
-USTRUCT(meta = (DataflowGeometryCollection))
+USTRUCT(meta = (DataflowGeometryCollection, Deprecated = "5.6"))
 struct FSafeReciprocalDataflowNode : public FDataflowNode
 {
 	GENERATED_USTRUCT_BODY()
@@ -222,7 +222,7 @@ public:
 *
 *
 */
-USTRUCT(meta = (DataflowGeometryCollection))
+USTRUCT(meta = (DataflowGeometryCollection, Deprecated = "5.6"))
 struct FSquareDataflowNode : public FDataflowNode
 {
 	GENERATED_USTRUCT_BODY()
@@ -250,7 +250,7 @@ public:
 *
 *
 */
-USTRUCT(meta = (DataflowGeometryCollection))
+USTRUCT(meta = (DataflowGeometryCollection, Deprecated = "5.6"))
 struct FSquareRootDataflowNode : public FDataflowNode
 {
 	GENERATED_USTRUCT_BODY()
@@ -278,7 +278,7 @@ public:
 *
 *
 */
-USTRUCT(meta = (DataflowGeometryCollection))
+USTRUCT(meta = (DataflowGeometryCollection, Deprecated = "5.6"))
 struct FInverseSqrtDataflowNode : public FDataflowNode
 {
 	GENERATED_USTRUCT_BODY()
@@ -306,7 +306,7 @@ public:
 *
 *
 */
-USTRUCT(meta = (DataflowGeometryCollection))
+USTRUCT(meta = (DataflowGeometryCollection, Deprecated = "5.6"))
 struct FCubeDataflowNode : public FDataflowNode
 {
 	GENERATED_USTRUCT_BODY()
@@ -334,7 +334,7 @@ public:
 *
 *
 */
-USTRUCT(meta = (DataflowGeometryCollection))
+USTRUCT(meta = (DataflowGeometryCollection, Deprecated = "5.6"))
 struct FNegateDataflowNode : public FDataflowNode
 {
 	GENERATED_USTRUCT_BODY()
@@ -362,7 +362,7 @@ public:
 *
 *
 */
-USTRUCT(meta = (DataflowGeometryCollection))
+USTRUCT(meta = (DataflowGeometryCollection, Deprecated = "5.6"))
 struct FAbsDataflowNode : public FDataflowNode
 {
 	GENERATED_USTRUCT_BODY()
@@ -390,7 +390,7 @@ public:
 * Converts a float to a nearest less or equal integer
 *
 */
-USTRUCT(meta = (DataflowGeometryCollection))
+USTRUCT(meta = (DataflowGeometryCollection, Deprecated = "5.6"))
 struct FFloorDataflowNode : public FDataflowNode
 {
 	GENERATED_USTRUCT_BODY()
@@ -418,7 +418,7 @@ public:
 * Converts a float to the nearest greater or equal integer
 *
 */
-USTRUCT(meta = (DataflowGeometryCollection))
+USTRUCT(meta = (DataflowGeometryCollection, Deprecated = "5.6"))
 struct FCeilDataflowNode : public FDataflowNode
 {
 	GENERATED_USTRUCT_BODY()
@@ -446,7 +446,7 @@ public:
 * Converts a float to the nearest integer. Rounds up when the fraction is .5
 *
 */
-USTRUCT(meta = (DataflowGeometryCollection))
+USTRUCT(meta = (DataflowGeometryCollection, Deprecated = "5.6"))
 struct FRoundDataflowNode : public FDataflowNode
 {
 	GENERATED_USTRUCT_BODY()
@@ -474,7 +474,7 @@ public:
 * Converts a float to the nearest integer. Rounds up when the fraction is .5
 *
 */
-USTRUCT(meta = (DataflowGeometryCollection))
+USTRUCT(meta = (DataflowGeometryCollection, Deprecated = "5.6"))
 struct FTruncDataflowNode : public FDataflowNode
 {
 	GENERATED_USTRUCT_BODY()
@@ -502,7 +502,7 @@ public:
 * Returns the fractional part of a float
 *
 */
-USTRUCT(meta = (DataflowGeometryCollection))
+USTRUCT(meta = (DataflowGeometryCollection, Deprecated = "5.6"))
 struct FFracDataflowNode : public FDataflowNode
 {
 	GENERATED_USTRUCT_BODY()
@@ -530,7 +530,7 @@ public:
 *
 *
 */
-USTRUCT(meta = (DataflowGeometryCollection))
+USTRUCT(meta = (DataflowGeometryCollection, Deprecated = "5.6"))
 struct FMinDataflowNode : public FDataflowNode
 {
 	GENERATED_USTRUCT_BODY()
@@ -562,7 +562,7 @@ public:
 *
 *
 */
-USTRUCT(meta = (DataflowGeometryCollection))
+USTRUCT(meta = (DataflowGeometryCollection, Deprecated = "5.6"))
 struct FMaxDataflowNode : public FDataflowNode
 {
 	GENERATED_USTRUCT_BODY()
@@ -666,7 +666,7 @@ public:
 * Returns 1, 0, or -1 depending on relation of Float to 0
 *
 */
-USTRUCT(meta = (DataflowGeometryCollection))
+USTRUCT(meta = (DataflowGeometryCollection, Deprecated = "5.6"))
 struct FSignDataflowNode : public FDataflowNode
 {
 	GENERATED_USTRUCT_BODY()
@@ -820,7 +820,7 @@ public:
 * 
 * 
 */
-USTRUCT(meta = (DataflowGeometryCollection))
+USTRUCT(meta = (DataflowGeometryCollection, Deprecated = "5.6"))
 struct FPowDataflowNode : public FDataflowNode
 {
 	GENERATED_USTRUCT_BODY()
@@ -852,7 +852,7 @@ public:
 * Returns log of A base Base (if Base ^ R == A, where R is ReturnValue)
 *
 */
-USTRUCT(meta = (DataflowGeometryCollection))
+USTRUCT(meta = (DataflowGeometryCollection, Deprecated = "5.6"))
 struct FLogDataflowNode : public FDataflowNode
 {
 	GENERATED_USTRUCT_BODY()
@@ -884,7 +884,7 @@ public:
 * Returns log of A base Base (if Base ^ R == A, where R is ReturnValue)
 *
 */
-USTRUCT(meta = (DataflowGeometryCollection))
+USTRUCT(meta = (DataflowGeometryCollection, Deprecated = "5.6"))
 struct FLogeDataflowNode : public FDataflowNode
 {
 	GENERATED_USTRUCT_BODY()
@@ -984,7 +984,7 @@ public:
 *
 *
 */
-USTRUCT(meta = (DataflowGeometryCollection))
+USTRUCT(meta = (DataflowGeometryCollection, Deprecated = "5.6"))
 struct FExpDataflowNode : public FDataflowNode
 {
 	GENERATED_USTRUCT_BODY()
@@ -1012,7 +1012,7 @@ public:
 * Returns the sine of A. (Expects Radians)
 *
 */
-USTRUCT(meta = (DataflowGeometryCollection))
+USTRUCT(meta = (DataflowGeometryCollection, Deprecated = "5.6"))
 struct FSinDataflowNode : public FDataflowNode
 {
 	GENERATED_USTRUCT_BODY()
@@ -1040,7 +1040,7 @@ public:
 * Returns the arcsine of A. (ReturnValue is in Radians)
 *
 */
-USTRUCT(meta = (DataflowGeometryCollection))
+USTRUCT(meta = (DataflowGeometryCollection, Deprecated = "5.6"))
 struct FArcSinDataflowNode : public FDataflowNode
 {
 	GENERATED_USTRUCT_BODY()
@@ -1068,7 +1068,7 @@ public:
 * Returns the cosine of A. (Expects Radians)
 *
 */
-USTRUCT(meta = (DataflowGeometryCollection))
+USTRUCT(meta = (DataflowGeometryCollection, Deprecated = "5.6"))
 struct FCosDataflowNode : public FDataflowNode
 {
 	GENERATED_USTRUCT_BODY()
@@ -1096,7 +1096,7 @@ public:
 * Returns the arccosine of A. (ReturnValue is in Radians)
 *
 */
-USTRUCT(meta = (DataflowGeometryCollection))
+USTRUCT(meta = (DataflowGeometryCollection, Deprecated = "5.6"))
 struct FArcCosDataflowNode : public FDataflowNode
 {
 	GENERATED_USTRUCT_BODY()
@@ -1124,7 +1124,7 @@ public:
 * Returns the cosine of A. (Expects Radians)
 *
 */
-USTRUCT(meta = (DataflowGeometryCollection))
+USTRUCT(meta = (DataflowGeometryCollection, Deprecated = "5.6"))
 struct FTanDataflowNode : public FDataflowNode
 {
 	GENERATED_USTRUCT_BODY()
@@ -1152,7 +1152,7 @@ public:
 * Returns the arccosine of A. (ReturnValue is in Radians)
 *
 */
-USTRUCT(meta = (DataflowGeometryCollection))
+USTRUCT(meta = (DataflowGeometryCollection, Deprecated = "5.6"))
 struct FArcTanDataflowNode : public FDataflowNode
 {
 	GENERATED_USTRUCT_BODY()
@@ -1180,7 +1180,7 @@ public:
 * Returns the arccosine of A. (ReturnValue is in Radians)
 *
 */
-USTRUCT(meta = (DataflowGeometryCollection))
+USTRUCT(meta = (DataflowGeometryCollection, Deprecated = "5.6"))
 struct FArcTan2DataflowNode : public FDataflowNode
 {
 	GENERATED_USTRUCT_BODY()
@@ -1617,7 +1617,7 @@ public:
  * Converts radians to degrees
  *
  */
-USTRUCT(meta = (DataflowGeometryCollection))
+USTRUCT(meta = (DataflowGeometryCollection, Deprecated = "5.6"))
 struct FRadiansToDegreesDataflowNode : public FDataflowNode
 {
 	GENERATED_USTRUCT_BODY()
@@ -1647,7 +1647,7 @@ public:
  * Converts degrees to radians
  *
  */
-USTRUCT(meta = (DataflowGeometryCollection))
+USTRUCT(meta = (DataflowGeometryCollection, Deprecated = "5.6"))
 struct FDegreesToRadiansDataflowNode : public FDataflowNode
 {
 	GENERATED_USTRUCT_BODY()
@@ -1698,7 +1698,7 @@ enum class EMathConstantsEnum : uint8
  * Offers a selection of Math constants
  *
  */
-USTRUCT(meta = (DataflowGeometryCollection))
+USTRUCT(meta = (DataflowGeometryCollection, Deprecated = "5.6"))
 struct FMathConstantsDataflowNode : public FDataflowNode
 {
 	GENERATED_USTRUCT_BODY()
@@ -1728,7 +1728,7 @@ public:
 * Returns 1.0 - A
 *
 */
-USTRUCT(meta = (DataflowGeometryCollection))
+USTRUCT(meta = (DataflowGeometryCollection, Deprecated = "5.6"))
 struct FOneMinusDataflowNode : public FDataflowNode
 {
 	GENERATED_USTRUCT_BODY()
