Dataflow:
- Updated MakeTransform node to have FVector, FRotator, FQuad inputs for rotation
- Added MakeRotator, BreakTransform nodes

diff --git a/Engine/Plugins/Experimental/GeometryCollectionPlugin/Source/GeometryCollectionNodes/Private/Dataflow/GeometryCollectionMakeNodes.cpp b/Engine/Plugins/Experimental/GeometryCollectionPlugin/Source/GeometryCollectionNodes/Private/Dataflow/GeometryCollectionMakeNodes.cpp
--- a/Engine/Plugins/Experimental/GeometryCollectionPlugin/Source/GeometryCollectionNodes/Private/Dataflow/GeometryCollectionMakeNodes.cpp
+++ b/Engine/Plugins/Experimental/GeometryCollectionPlugin/Source/GeometryCollectionNodes/Private/Dataflow/GeometryCollectionMakeNodes.cpp
@@ -52,10 +52,14 @@ namespace UE::Dataflow
 		DATAFLOW_NODE_REGISTER_CREATION_FACTORY(FMakeLiteralIntDataflowNode);
 		DATAFLOW_NODE_REGISTER_CREATION_FACTORY(FMakeLiteralBoolDataflowNode);
 		DATAFLOW_NODE_REGISTER_CREATION_FACTORY(FMakeLiteralVectorDataflowNode);
-		DATAFLOW_NODE_REGISTER_CREATION_FACTORY(FMakeTransformDataflowNode);
+		DATAFLOW_NODE_REGISTER_CREATION_FACTORY(FMakeTransformDataflowNode_v2);
 		DATAFLOW_NODE_REGISTER_CREATION_FACTORY(FMakeQuaternionDataflowNode);
 		DATAFLOW_NODE_REGISTER_CREATION_FACTORY(FMakeFloatArrayDataflowNode);
 		DATAFLOW_NODE_REGISTER_CREATION_FACTORY(FMakeCollectionDataflowNode);
+		DATAFLOW_NODE_REGISTER_CREATION_FACTORY(FMakeRotatorDataflowNode);
+		DATAFLOW_NODE_REGISTER_CREATION_FACTORY(FBreakTransformDataflowNode);
+		// Deprecated
+		DATAFLOW_NODE_REGISTER_CREATION_FACTORY(FMakeTransformDataflowNode);
 	}
 }
 
@@ -154,6 +158,49 @@ void FMakeTransformDataflowNode::Evaluate(UE::Dataflow::FContext& Context, const
 	}
 }
 
+/* -------------------------------------------------------------------------------- */
+
+FMakeTransformDataflowNode_v2::FMakeTransformDataflowNode_v2(const UE::Dataflow::FNodeParameters& InParam, FGuid InGuid)
+	: FDataflowNode(InParam, InGuid)
+{
+	RegisterInputConnection(&Translation);
+	RegisterInputConnection(&Rotation);
+	RegisterInputConnection(&Rotator) .SetCanHidePin(true) .SetPinIsHidden(true);
+	RegisterInputConnection(&Quat).SetCanHidePin(true).SetPinIsHidden(true);
+	RegisterInputConnection(&Scale);
+	RegisterOutputConnection(&Transform);
+}
+
+void FMakeTransformDataflowNode_v2::Evaluate(UE::Dataflow::FContext& Context, const FDataflowOutput* Out) const
+{
+	if (Out->IsA<FTransform>(&Transform))
+	{
+		const FVector InTranslation = GetValue(Context, &Translation);
+		const FVector InScale = GetValue(Context, &Scale);
+		FQuat OutQuat;
+		if (IsConnected(&Rotation))
+		{
+			const FVector InRotation = GetValue(Context, &Rotation);
+			OutQuat = FQuat::MakeFromEuler(InRotation);
+		}
+		else if (IsConnected(&Rotator))
+		{
+			const FRotator InRotator = GetValue(Context, &Rotator);
+			OutQuat = FQuat::MakeFromRotator(InRotator);
+		}
+		else if (IsConnected(&Quat))
+		{
+			const FQuat InQuat = GetValue(Context, &Quat);
+			OutQuat = InQuat;
+		}
+
+		FTransform OutTransform = FTransform(OutQuat, InTranslation, InScale);
+		SetValue(Context, OutTransform, &Transform);
+	}
+}
+
+/* -------------------------------------------------------------------------------- */
+
 void FMakeQuaternionDataflowNode::Evaluate(UE::Dataflow::FContext& Context, const FDataflowOutput* Out) const
 {
 	if (Out->IsA<FQuat>(&Quaternion))
@@ -184,3 +231,64 @@ void FMakeCollectionDataflowNode::Evaluate(UE::Dataflow::FContext& Context, cons
 		SetValue(Context, FManagedArrayCollection(), &Collection);
 	}
 }
+
+/* -------------------------------------------------------------------------------- */
+
+FMakeRotatorDataflowNode::FMakeRotatorDataflowNode(const UE::Dataflow::FNodeParameters& InParam, FGuid InGuid)
+: FDataflowNode(InParam, InGuid)
+{
+	RegisterInputConnection(&Pitch);
+	RegisterInputConnection(&Yaw);
+	RegisterInputConnection(&Roll);
+	RegisterOutputConnection(&Rotator);
+}
+
+void FMakeRotatorDataflowNode::Evaluate(UE::Dataflow::FContext& Context, const FDataflowOutput* Out) const
+{
+	if (Out->IsA(&Rotator))
+	{
+		const float InPitch = GetValue(Context, &Pitch);
+		const float InYaw = GetValue(Context, &Yaw);
+		const float InRoll = GetValue(Context, &Roll);
+		SetValue(Context, FRotator(InPitch, InYaw, InRoll), &Rotator);
+	}
+}
+
+/* -------------------------------------------------------------------------------- */
+
+FBreakTransformDataflowNode::FBreakTransformDataflowNode(const UE::Dataflow::FNodeParameters& InParam, FGuid InGuid)
+	: FDataflowNode(InParam, InGuid)
+{
+	RegisterInputConnection(&Transform);
+	RegisterOutputConnection(&Translation);
+	RegisterOutputConnection(&Rotation);
+	RegisterOutputConnection(&Rotator);
+	RegisterOutputConnection(&Quat);
+	RegisterOutputConnection(&Scale);
+}
+
+void FBreakTransformDataflowNode::Evaluate(UE::Dataflow::FContext& Context, const FDataflowOutput* Out) const
+{
+	if (Out->IsA(&Translation) ||
+		Out->IsA(&Rotation) ||
+		Out->IsA(&Rotator) ||
+		Out->IsA(&Quat) ||
+		Out->IsA(&Scale))
+	{
+		const FTransform InTransform = GetValue(Context, &Transform);
+
+		const FVector OutTranslation = InTransform.GetTranslation();
+		const FVector OutRotationAsEuler = InTransform.GetRotation().Euler();
+		const FRotator OutRotator = InTransform.GetRotation().Rotator();
+		const FQuat OutQuat = InTransform.GetRotation();
+		const FVector OutScale = InTransform.GetScale3D();
+
+		SetValue(Context, OutTranslation, &Translation);
+		SetValue(Context, OutRotationAsEuler, &Rotation);
+		SetValue(Context, OutRotator, &Rotator);
+		SetValue(Context, OutQuat, &Quat);
+		SetValue(Context, OutScale, &Scale);
+	}
+}
+
+/* -------------------------------------------------------------------------------- */
diff --git a/Engine/Plugins/Experimental/GeometryCollectionPlugin/Source/GeometryCollectionNodes/Public/Dataflow/GeometryCollectionMakeNodes.h b/Engine/Plugins/Experimental/GeometryCollectionPlugin/Source/GeometryCollectionNodes/Public/Dataflow/GeometryCollectionMakeNodes.h
--- a/Engine/Plugins/Experimental/GeometryCollectionPlugin/Source/GeometryCollectionNodes/Public/Dataflow/GeometryCollectionMakeNodes.h
+++ b/Engine/Plugins/Experimental/GeometryCollectionPlugin/Source/GeometryCollectionNodes/Public/Dataflow/GeometryCollectionMakeNodes.h
@@ -44,7 +44,6 @@ public:
 	}
 
 	virtual void Evaluate(UE::Dataflow::FContext& Context, const FDataflowOutput* Out) const override;
-
 };
 
 /**
@@ -73,7 +72,6 @@ public:
 	}
 
 	virtual void Evaluate(UE::Dataflow::FContext& Context, const FDataflowOutput* Out) const override;
-
 };
 
 
@@ -129,7 +127,6 @@ public:
 	}
 
 	virtual void Evaluate(UE::Dataflow::FContext& Context, const FDataflowOutput* Out) const override;
-
 };
 
 
@@ -164,7 +161,6 @@ public:
 	}
 
 	virtual void Evaluate(UE::Dataflow::FContext& Context, const FDataflowOutput* Out) const override;
-
 };
 
 
@@ -193,7 +189,6 @@ public:
 	}
 
 	virtual void Evaluate(UE::Dataflow::FContext& Context, const FDataflowOutput* Out) const override;
-
 };
 
 
@@ -222,7 +217,6 @@ public:
 	}
 
 	virtual void Evaluate(UE::Dataflow::FContext& Context, const FDataflowOutput* Out) const override;
-
 };
 
 /**
@@ -250,7 +244,6 @@ public:
 	}
 
 	virtual void Evaluate(UE::Dataflow::FContext& Context, const FDataflowOutput* Out) const override;
-
 };
 
 
@@ -288,15 +281,14 @@ public:
 	}
 
 	virtual void Evaluate(UE::Dataflow::FContext& Context, const FDataflowOutput* Out) const override;
-
 };
 
 /**
- *
- * Description for this node
- *
+ * Make an FTransform
+ * Deprecated (5.6)
+ * Use FMakeTransformDataflowNode_v2 instead
  */
-USTRUCT()
+USTRUCT(meta = (Deprecated = "5.6"))
 struct FMakeTransformDataflowNode : public FDataflowNode
 {
 	GENERATED_USTRUCT_BODY()
@@ -326,7 +318,50 @@ public:
 	}
 
 	virtual void Evaluate(UE::Dataflow::FContext& Context, const FDataflowOutput* Out) const override;
+};
+
+/**
+*
+* Make a FTransform
+*
+*/
+USTRUCT()
+struct FMakeTransformDataflowNode_v2 : public FDataflowNode
+{
+	GENERATED_USTRUCT_BODY()
+	DATAFLOW_NODE_DEFINE_INTERNAL(FMakeTransformDataflowNode_v2, "MakeTransform", "Generators|Transform", "")
+
+private:
+
+	/** Translation */
+	UPROPERTY(EditAnywhere, Category = "Transform", meta = (DataflowInput));
+	FVector Translation = FVector(0.f, 0.f, 0.f);
+
+	/** Rotation as Euler */
+	UPROPERTY(EditAnywhere, Category = "Transform", meta = (DataflowInput));
+	FVector Rotation = FVector(0.f, 0.f, 0.f);
+
+	/** Rotation a Rotator */
+	UPROPERTY(EditAnywhere, Category = "Transform", meta = (DataflowInput));
+	FRotator Rotator = FRotator(0.f, 0.f, 0.f);
+
+	/** Rotation as a quaternion */
+	UPROPERTY(EditAnywhere, Category = "Transform", meta = (DataflowInput));
+	FQuat Quat;
+
+	/** Scale */
+	UPROPERTY(EditAnywhere, Category = "Transform", meta = (DataflowInput));
+	FVector Scale = FVector(1.f, 1.f, 1.f);
+
+	/** Result transform */
+	UPROPERTY(meta = (DataflowOutput));
+	FTransform Transform = FTransform::Identity;
+
+public:
+	FMakeTransformDataflowNode_v2(const UE::Dataflow::FNodeParameters& InParam, FGuid InGuid = FGuid::NewGuid());
 
+private:
+	virtual void Evaluate(UE::Dataflow::FContext& Context, const FDataflowOutput* Out) const override;
 };
 
 /**
@@ -367,7 +402,6 @@ public:
 	}
 
 	virtual void Evaluate(UE::Dataflow::FContext& Context, const FDataflowOutput* Out) const override;
-
 };
 
 /**
@@ -403,7 +437,6 @@ public:
 	}
 
 	virtual void Evaluate(UE::Dataflow::FContext& Context, const FDataflowOutput* Out) const override;
-
 };
 
 /**
@@ -431,6 +464,81 @@ public:
 	virtual void Evaluate(UE::Dataflow::FContext& Context, const FDataflowOutput* Out) const override;
 };
 
+/**
+ *
+ * Make a Rotator
+ *
+ */
+USTRUCT()
+struct FMakeRotatorDataflowNode : public FDataflowNode
+{
+	GENERATED_USTRUCT_BODY()
+	DATAFLOW_NODE_DEFINE_INTERNAL(FMakeRotatorDataflowNode, "MakeRotator", "Generators|Transform", "")
+
+private:
+	/** Rotation around the right axis (around Y axis), Looking up and down (0=Straight Ahead, +Up, -Down) */
+	UPROPERTY(EditAnywhere, Category = "Rotator", meta = (DataflowInput));
+	float Pitch = 0.f;
+
+	/** Rotation around the up axis (around Z axis), Turning around (0=Forward, +Right, -Left)*/
+	UPROPERTY(EditAnywhere, Category = "Rotator", meta = (DataflowInput));
+	float Yaw = 0.f;
+
+	/** Rotation around the forward axis (around X axis), Tilting your head, (0=Straight, +Clockwise, -CCW) */
+	UPROPERTY(EditAnywhere, Category = "Rotator", meta = (DataflowInput));
+	float Roll = 0.f;
+
+	/** Rotator output */
+	UPROPERTY(meta = (DataflowOutput));
+	FRotator Rotator = FRotator(0.f, 0.f, 0.f);
+
+public:
+	FMakeRotatorDataflowNode(const UE::Dataflow::FNodeParameters& InParam, FGuid InGuid = FGuid::NewGuid());
+
+	virtual void Evaluate(UE::Dataflow::FContext& Context, const FDataflowOutput* Out) const override;
+};
+
+/**
+* Break a Transform into Translation, Rotation (Euler, Rotator, Quaternion), Scale
+*/
+USTRUCT()
+struct FBreakTransformDataflowNode : public FDataflowNode
+{
+	GENERATED_USTRUCT_BODY()
+	DATAFLOW_NODE_DEFINE_INTERNAL(FBreakTransformDataflowNode, "BreakTransform", "Math|Transform", "")
+
+private:
+	/** Transform to break into components */
+	UPROPERTY(EditAnywhere, Category = "Transform", meta = (DataflowInput));
+	FTransform Transform;
+
+	/** Translation */
+	UPROPERTY(meta = (DataflowOutput));
+	FDataflowVectorTypes Translation;
+
+	/** Rotation as Euler */
+	UPROPERTY(meta = (DataflowOutput));
+	FVector Rotation = FVector(0.f, 0.f, 0.f);
+
+	/** Rotation as a rotator */
+	UPROPERTY(meta = (DataflowOutput));
+	FRotator Rotator = FRotator(0.f, 0.f, 0.f);
+
+	/** Rotation as a quaternion */
+	UPROPERTY(meta = (DataflowOutput));
+	FQuat Quat;
+
+	/** Scale */
+	UPROPERTY(meta = (DataflowOutput));
+	FDataflowVectorTypes Scale;
+
+public:
+	FBreakTransformDataflowNode(const UE::Dataflow::FNodeParameters& InParam, FGuid InGuid = FGuid::NewGuid());
+
+private:
+	virtual void Evaluate(UE::Dataflow::FContext& Context, const FDataflowOutput* Out) const override;
+};
+
 namespace UE::Dataflow
 {
 	void GeometryCollectionMakeNodes();
