Dataflow : add Uv support to Obj string node

diff --git a/Engine/Plugins/Experimental/GeometryCollectionPlugin/Source/GeometryCollectionNodes/Private/Dataflow/GeometryCollectionDebugNodes.cpp b/Engine/Plugins/Experimental/GeometryCollectionPlugin/Source/GeometryCollectionNodes/Private/Dataflow/GeometryCollectionDebugNodes.cpp
--- a/Engine/Plugins/Experimental/GeometryCollectionPlugin/Source/GeometryCollectionNodes/Private/Dataflow/GeometryCollectionDebugNodes.cpp
+++ b/Engine/Plugins/Experimental/GeometryCollectionPlugin/Source/GeometryCollectionNodes/Private/Dataflow/GeometryCollectionDebugNodes.cpp
@@ -135,6 +135,32 @@ void FMeshToOBJStringDebugDataflowNode::Evaluate(UE::Dataflow::FContext& Context
 			}
 			Build += FString::Printf(TEXT("v %f %f %f\n"), V.X, V.Y, V.Z);
 		}
+
+		// Uvs
+		const UE::Geometry::FDynamicMeshUVOverlay* UvOverlay = nullptr;
+
+		if (MeshRef.HasVertexUVs())
+		{
+			for (int32 VID = 0; VID < MeshRef.MaxVertexID(); ++VID)
+			{
+				const FVector2f Uv = MeshRef.GetVertexUV(VID);
+				Build += FString::Printf(TEXT("vt %f %f\n"), Uv.X, Uv.Y);
+			}
+		}
+		else if (MeshRef.Attributes()->NumUVLayers() > 0)
+		{
+			UvOverlay = MeshRef.Attributes()->GetUVLayer(0);
+			if (UvOverlay)
+			{
+				for (int32 UvIndex = 0; UvIndex < UvOverlay->ElementCount(); ++UvIndex)
+				{
+					const FVector2f Uv = UvOverlay->GetElement(UvIndex);
+					Build += FString::Printf(TEXT("vt %f %f\n"), Uv.X, Uv.Y);
+				}
+			}
+		}
+
+		// faces
 		bool bInInvertFaces = GetValue(Context, &bInvertFaces);
 		for (int32 TID = 0; TID < MeshRef.MaxTriangleID(); ++TID)
 		{
@@ -146,7 +172,19 @@ void FMeshToOBJStringDebugDataflowNode::Evaluate(UE::Dataflow::FContext& Context
 				{
 					Swap(T.B, T.C);
 				}
-				Build += FString::Printf(TEXT("f %d %d %d\n"), T.A + 1, T.B + 1, T.C + 1);
+				if (UvOverlay)
+				{
+					UE::Geometry::FIndex3i UvIndices = UvOverlay->GetTriangle(TID);
+					if (!bInInvertFaces)
+					{
+						Swap(UvIndices.B, UvIndices.C);
+					}
+					Build += FString::Printf(TEXT("f %d/%d %d/%d %d/%d\n"), T.A + 1, UvIndices.A + 1, T.B + 1, UvIndices.B + 1, T.C + 1, UvIndices.C + 1);
+				}
+				else
+				{
+					Build += FString::Printf(TEXT("f %d %d %d\n"), T.A + 1, T.B + 1, T.C + 1);
+				}
 			}
 		}
 		SetValue(Context, MoveTemp(Build), &StringOBJ);
