Change default behavior for new CollectionSelectionToMeshes nodes to filter the selection to leaves, as this is more commonly useful (especially for the pin-disconnected-so-select-all case)

Fix the selection input comment to remove copy-pasted reference to geometry script, and to clarify that the selection can optionally be filtered to leaves.

diff --git a/Engine/Plugins/Experimental/GeometryCollectionPlugin/Source/GeometryCollectionNodes/Public/Dataflow/GeometryCollectionMeshNodes.h b/Engine/Plugins/Experimental/GeometryCollectionPlugin/Source/GeometryCollectionNodes/Public/Dataflow/GeometryCollectionMeshNodes.h
--- a/Engine/Plugins/Experimental/GeometryCollectionPlugin/Source/GeometryCollectionNodes/Public/Dataflow/GeometryCollectionMeshNodes.h
+++ b/Engine/Plugins/Experimental/GeometryCollectionPlugin/Source/GeometryCollectionNodes/Public/Dataflow/GeometryCollectionMeshNodes.h
@@ -198,13 +198,13 @@ public:
 	UPROPERTY(meta = (DataflowInput, DataflowIntrinsic))
 	FManagedArrayCollection Collection;
 
-	/** Selected bones will have geometry script processing applied (if they have geometry). If not connected, all geometry will be processed. */
+	/** Geometry on or under selected bones will be converted to meshes, optionally after filtering the selection to leaves. If not connected, all geometry will be processed. */
 	UPROPERTY(meta = (DataflowInput))
 	FDataflowTransformSelection TransformSelection;
 
 	/** Whether to convert the input selection to only leaves, which may directly store geometry. Otherwise, meshes for selected cluster nodes will be generated by appending together geometry from leaf nodes. */
 	UPROPERTY(EditAnywhere, Category = Options)
-	bool bConvertSelectionToLeaves = false;
+	bool bConvertSelectionToLeaves = true;
 
 	/** Whether the processed mesh will have edges at normal/UV/color seams welded so they are treated as one edge during processing. */
 	UPROPERTY(EditAnywhere, Category = Options)
@@ -214,7 +214,7 @@ public:
 	UPROPERTY(EditAnywhere, Category = Options)
 	bool bPreserveIsolatedVertices = true;
 
-	/** Output DynamicMesh */
+	/** Output Array of DynamicMesh */
 	UPROPERTY(meta = (DataflowOutput))
 	TArray<TObjectPtr<UDynamicMesh>> Meshes;
 
