add render tags to more dataflow geometrycollection nodes
+ remove incorrect dataflowintrinsic tag from some inputs

diff --git a/Engine/Plugins/Experimental/GeometryCollectionPlugin/Source/GeometryCollectionNodes/Public/Dataflow/GeometryCollectionDebugNodes.h b/Engine/Plugins/Experimental/GeometryCollectionPlugin/Source/GeometryCollectionNodes/Public/Dataflow/GeometryCollectionDebugNodes.h
--- a/Engine/Plugins/Experimental/GeometryCollectionPlugin/Source/GeometryCollectionNodes/Public/Dataflow/GeometryCollectionDebugNodes.h
+++ b/Engine/Plugins/Experimental/GeometryCollectionPlugin/Source/GeometryCollectionNodes/Public/Dataflow/GeometryCollectionDebugNodes.h
@@ -32,7 +32,7 @@ private:
 	FManagedArrayCollection Collection;
 
 	/** Optional transform selection to convert hulls from -- if not provided, all convex hulls will be converted. */
-	UPROPERTY(meta = (DataflowInput, DataflowIntrinsic))
+	UPROPERTY(meta = (DataflowInput))
 	FDataflowTransformSelection OptionalSelectionFilter;
 
 	/** Single mesh aggregating all the convex hulls together */
@@ -55,6 +55,7 @@ struct FSphereCoveringToMeshDataflowNode : public FDataflowNode
 {
 	GENERATED_USTRUCT_BODY()
 	DATAFLOW_NODE_DEFINE_INTERNAL(FSphereCoveringToMeshDataflowNode, "Sphere Covering to Mesh", "GeometryCollection|Mesh", "")
+	DATAFLOW_NODE_RENDER_TYPE("SurfaceRender", FName("FDynamicMesh3"), "Mesh")
 
 public:
 	// The sphere covering to convert to a mesh
diff --git a/Engine/Plugins/Experimental/GeometryCollectionPlugin/Source/GeometryCollectionNodes/Public/Dataflow/GeometryCollectionFracturingNodes.h b/Engine/Plugins/Experimental/GeometryCollectionPlugin/Source/GeometryCollectionNodes/Public/Dataflow/GeometryCollectionFracturingNodes.h
--- a/Engine/Plugins/Experimental/GeometryCollectionPlugin/Source/GeometryCollectionNodes/Public/Dataflow/GeometryCollectionFracturingNodes.h
+++ b/Engine/Plugins/Experimental/GeometryCollectionPlugin/Source/GeometryCollectionNodes/Public/Dataflow/GeometryCollectionFracturingNodes.h
@@ -400,6 +400,7 @@ struct FVoronoiFractureDataflowNode_v2 : public FDataflowNode
 {
 	GENERATED_USTRUCT_BODY()
 	DATAFLOW_NODE_DEFINE_INTERNAL(FVoronoiFractureDataflowNode_v2, "VoronoiFracture", "GeometryCollection|Fracture", "")
+	DATAFLOW_NODE_RENDER_TYPE("SurfaceRender", FGeometryCollection::StaticType(), "Collection")
 
 public:
 	/** Collection to fracture */
@@ -596,6 +597,7 @@ struct FPlaneCutterDataflowNode_v2 : public FDataflowNode
 {
 	GENERATED_USTRUCT_BODY()
 	DATAFLOW_NODE_DEFINE_INTERNAL(FPlaneCutterDataflowNode_v2, "PlaneCutter", "GeometryCollection|Fracture", "")
+	DATAFLOW_NODE_RENDER_TYPE("SurfaceRender", FGeometryCollection::StaticType(), "Collection")
 
 public:
 	/** Collection to fracture */
@@ -770,6 +772,7 @@ struct FSliceCutterDataflowNode : public FDataflowNode
 {
 	GENERATED_USTRUCT_BODY()
 	DATAFLOW_NODE_DEFINE_INTERNAL(FSliceCutterDataflowNode, "SliceCutter", "GeometryCollection|Fracture", "")
+	DATAFLOW_NODE_RENDER_TYPE("SurfaceRender", FGeometryCollection::StaticType(), "Collection")
 
 public:
 	/** Collection to fracture */
@@ -907,6 +910,7 @@ struct FBrickCutterDataflowNode : public FDataflowNode
 {
 	GENERATED_USTRUCT_BODY()
 	DATAFLOW_NODE_DEFINE_INTERNAL(FBrickCutterDataflowNode, "BrickCutter", "GeometryCollection|Fracture", "")
+	DATAFLOW_NODE_RENDER_TYPE("SurfaceRender", FGeometryCollection::StaticType(), "Collection")
 
 public:
 	/** Collection to cut */
@@ -1043,6 +1047,7 @@ struct FMeshCutterDataflowNode : public FDataflowNode
 {
 	GENERATED_USTRUCT_BODY()
 	DATAFLOW_NODE_DEFINE_INTERNAL(FMeshCutterDataflowNode, "MeshCutter", "GeometryCollection|Fracture", "")
+	DATAFLOW_NODE_RENDER_TYPE("SurfaceRender", FGeometryCollection::StaticType(), "Collection")
 
 public:
 	/** Collection to cut */
@@ -1062,11 +1067,11 @@ public:
 	FTransform Transform = FTransform::Identity;
 
 	/** Dynamic Meshes to cut with */
-	UPROPERTY(EditAnywhere, Category = "Mesh", meta = (DataflowInput, DataflowIntrinsic))
+	UPROPERTY(EditAnywhere, Category = "Mesh", meta = (DataflowInput))
 	TArray<TObjectPtr<UDynamicMesh>> CuttingDynamicMeshes;
 
 	/** Static Mesh to cut with */
-	UPROPERTY(EditAnywhere, Category = "StaticMesh", meta = (DataflowInput, DataflowIntrinsic))
+	UPROPERTY(EditAnywhere, Category = "StaticMesh", meta = (DataflowInput))
 	TObjectPtr<UStaticMesh> CuttingStaticMesh;
 
 	/** If using a Static Mesh to cut, attempt to use the Nanite HiRes source mesh, if available and non-empty. */
