Initialized some quaternions

#rnx

diff --git a/Engine/Plugins/Experimental/GeometryCollectionPlugin/Source/GeometryCollectionNodes/Public/Dataflow/GeometryCollectionMakeNodes.h b/Engine/Plugins/Experimental/GeometryCollectionPlugin/Source/GeometryCollectionNodes/Public/Dataflow/GeometryCollectionMakeNodes.h
--- a/Engine/Plugins/Experimental/GeometryCollectionPlugin/Source/GeometryCollectionNodes/Public/Dataflow/GeometryCollectionMakeNodes.h
+++ b/Engine/Plugins/Experimental/GeometryCollectionPlugin/Source/GeometryCollectionNodes/Public/Dataflow/GeometryCollectionMakeNodes.h
@@ -449,7 +449,7 @@ private:
 
 	/** Rotation as a quaternion */
 	UPROPERTY(EditAnywhere, Category = "Transform", meta = (DataflowInput));
-	FQuat Quat;
+	FQuat Quat = FQuat(ForceInit);
 
 	/** Scale */
 	UPROPERTY(EditAnywhere, Category = "Transform", meta = (DataflowInput));
@@ -628,7 +628,7 @@ private:
 
 	/** Rotation as a quaternion */
 	UPROPERTY(meta = (DataflowOutput));
-	FQuat Quat;
+	FQuat Quat = FQuat(ForceInit);
 
 	/** Scale */
 	UPROPERTY(meta = (DataflowOutput));
