Dataflow : Extra selection check for baking texture node

diff --git a/Engine/Plugins/Experimental/GeometryCollectionPlugin/Source/GeometryCollectionNodes/Private/Dataflow/GeometryCollectionTextureNodes.cpp b/Engine/Plugins/Experimental/GeometryCollectionPlugin/Source/GeometryCollectionNodes/Private/Dataflow/GeometryCollectionTextureNodes.cpp
--- a/Engine/Plugins/Experimental/GeometryCollectionPlugin/Source/GeometryCollectionNodes/Private/Dataflow/GeometryCollectionTextureNodes.cpp
+++ b/Engine/Plugins/Experimental/GeometryCollectionPlugin/Source/GeometryCollectionNodes/Private/Dataflow/GeometryCollectionTextureNodes.cpp
@@ -59,6 +59,9 @@ void FBakeTextureFromCollectionDataflowNode::Evaluate(UE::Dataflow::FContext& Co
 			// log until we have proper error reporting in context 
 			UE_LOG(LogChaosDataflow, Error, TEXT("Dataflow Node [%s] : Invalid UV channel [%d] - The collection has %d UV channels"), *GetName().ToString(), InUVChannel, GeomCollection->NumUVLayers());
 
+			const FString Message = FString::Printf(TEXT("Invalid UV channel [%d] - The collection has %d UV channels"), InUVChannel, GeomCollection->NumUVLayers());
+			Context.Warning(Message, this, Out);
+
 			// Write out a red 16x16 image as a sign of error
 			FDataflowImage DefaultImage;
 			DefaultImage.CreateFromColor(EDataflowImageResolution::Resolution16, FLinearColor::Red);
@@ -87,7 +90,7 @@ void FBakeTextureFromCollectionDataflowNode::Evaluate(UE::Dataflow::FContext& Co
 		if (IsConnected(&FaceSelection))
 		{
 			const FDataflowFaceSelection& InFaceSelection = GetValue(Context, &FaceSelection);
-			if (InFaceSelection.Num() > 0)
+			if (InFaceSelection.Num() == GeomCollection->NumElements(FGeometryCollection::FacesGroup))
 			{
 				TArray<bool> FacesToBake;
 				FacesToBake.SetNumUninitialized(InFaceSelection.Num());
@@ -101,6 +104,10 @@ void FBakeTextureFromCollectionDataflowNode::Evaluate(UE::Dataflow::FContext& Co
 					BakeAttributes, BakeAttributeSettings, ImageBuilder,
 					FacesToBake, nullptr /* Progress */);
 			}
+			else
+			{
+				Context.Warning(TEXT("Selection does not match the collection, the collection may have changed since the selection was generated from it"), this, Out);
+			}
 		}
 		else
 		{
