Dataflow:
- Addded DebugDraw to MakePoints, PlaneCutter nodes
- Added more functionality to DebugDraw
- Added MakeSphereMesh, MakeCylinderMesh, MakeCapsuleMesh, MakeBoxMesh nodes
- Bug fixes

diff --git a/Engine/Plugins/Experimental/Fracture/Source/FractureEngine/Private/FractureEngineFracturing.cpp b/Engine/Plugins/Experimental/Fracture/Source/FractureEngine/Private/FractureEngineFracturing.cpp
--- a/Engine/Plugins/Experimental/Fracture/Source/FractureEngine/Private/FractureEngineFracturing.cpp
+++ b/Engine/Plugins/Experimental/Fracture/Source/FractureEngine/Private/FractureEngineFracturing.cpp
@@ -471,6 +471,24 @@ int32 FFractureEngineFracturing::VoronoiFracture(FManagedArrayCollection& InOutC
 }
 
 
+void FFractureEngineFracturing::GenerateSliceTransforms(const FBox& InBoundingBox,
+	const int32 InRandomSeed,
+	const int32 InNumPlanes,
+	TArray<FTransform>& OutCuttingPlaneTransforms)
+{
+	FRandomStream RandStream(InRandomSeed);
+
+	FBox Bounds = InBoundingBox;
+	const FVector Extent(Bounds.Max - Bounds.Min);
+
+	OutCuttingPlaneTransforms.Reserve(OutCuttingPlaneTransforms.Num() + InNumPlanes);
+	for (int32 Idx = 0; Idx < InNumPlanes; ++Idx)
+	{
+		FVector Position(Bounds.Min + FVector(RandStream.FRand(), RandStream.FRand(), RandStream.FRand()) * Extent);
+		OutCuttingPlaneTransforms.Emplace(FTransform(FRotator(RandStream.FRand() * 360.0f, RandStream.FRand() * 360.0f, 0.0f), Position));
+	}
+}
+
 int32 FFractureEngineFracturing::PlaneCutter(FManagedArrayCollection& InOutCollection,
 	FDataflowTransformSelection InTransformSelection,
 	const FBox& InBoundingBox,
@@ -495,17 +513,7 @@ int32 FFractureEngineFracturing::PlaneCutter(FManagedArrayCollection& InOutColle
 		TArray<FPlane> CuttingPlanes;
 		TArray<FTransform> CuttingPlaneTransforms;
 
-		FRandomStream RandStream(InRandomSeed);
-
-		FBox Bounds = InBoundingBox;
-		const FVector Extent(Bounds.Max - Bounds.Min);
-
-		CuttingPlaneTransforms.Reserve(CuttingPlaneTransforms.Num() + InNumPlanes);
-		for (int32 ii = 0; ii < InNumPlanes; ++ii)
-		{
-			FVector Position(Bounds.Min + FVector(RandStream.FRand(), RandStream.FRand(), RandStream.FRand()) * Extent);
-			CuttingPlaneTransforms.Emplace(FTransform(FRotator(RandStream.FRand() * 360.0f, RandStream.FRand() * 360.0f, 0.0f), Position));
-		}
+		FFractureEngineFracturing::GenerateSliceTransforms(InBoundingBox, InRandomSeed, InNumPlanes, CuttingPlaneTransforms);
 
 		for (const FTransform& Transform : CuttingPlaneTransforms)
 		{
diff --git a/Engine/Plugins/Experimental/Fracture/Source/FractureEngine/Public/FractureEngineFracturing.h b/Engine/Plugins/Experimental/Fracture/Source/FractureEngine/Public/FractureEngineFracturing.h
--- a/Engine/Plugins/Experimental/Fracture/Source/FractureEngine/Public/FractureEngineFracturing.h
+++ b/Engine/Plugins/Experimental/Fracture/Source/FractureEngine/Public/FractureEngineFracturing.h
@@ -103,6 +103,12 @@ public:
 		bool InAddSamplesForCollision,
 		float InCollisionSampleSpacing);
 
+	static void GenerateSliceTransforms(
+		const FBox& InBoundingBox,
+		const int32 InRandomSeed,
+		const int32 InNumPlanes,
+		TArray<FTransform>& OutCuttingPlaneTransforms);
+	
 	static int32 PlaneCutter(FManagedArrayCollection& InOutCollection,
 		FDataflowTransformSelection InTransformSelection,
 		const FBox& InBoundingBox,
diff --git a/Engine/Plugins/Experimental/GeometryCollectionPlugin/Source/GeometryCollectionNodes/Private/Dataflow/GeometryCollectionFracturingNodes.cpp b/Engine/Plugins/Experimental/GeometryCollectionPlugin/Source/GeometryCollectionNodes/Private/Dataflow/GeometryCollectionFracturingNodes.cpp
--- a/Engine/Plugins/Experimental/GeometryCollectionPlugin/Source/GeometryCollectionNodes/Private/Dataflow/GeometryCollectionFracturingNodes.cpp
+++ b/Engine/Plugins/Experimental/GeometryCollectionPlugin/Source/GeometryCollectionNodes/Private/Dataflow/GeometryCollectionFracturingNodes.cpp
@@ -36,6 +36,7 @@
 #include "DynamicMeshToMeshDescription.h"
 #include "MeshDescriptionToDynamicMesh.h"
 #include "StaticMeshAttributes.h"
+#include "Dataflow/DataflowSimpleDebugDrawMesh.h"
 
 #include UE_INLINE_GENERATED_CPP_BY_NAME(GeometryCollectionFracturingNodes)
 
@@ -492,12 +493,39 @@ void FPlaneCutterDataflowNode::Evaluate(UE::Dataflow::FContext& Context, const F
 	}
 }
 
+/*--------------------------------------------------------------------------------------------------------*/
+
+static FLinearColor GetRandomColor(const int32 RandomSeed, int32 Idx)
+{
+	FRandomStream RandomStream(RandomSeed * 23 + Idx * 4078);
+	
+	const uint8 R = static_cast<uint8>(RandomStream.FRandRange(128, 255));
+	const uint8 G = static_cast<uint8>(RandomStream.FRandRange(128, 255));
+	const uint8 B = static_cast<uint8>(RandomStream.FRandRange(128, 255));
+
+	return FLinearColor(FColor(R, G, B, 255));
+}
+
 void FPlaneCutterDataflowNode_v2::Evaluate(UE::Dataflow::FContext& Context, const FDataflowOutput* Out) const
 {
 	if (Out->IsA<FManagedArrayCollection>(&Collection) ||
 		Out->IsA<FDataflowTransformSelection>(&TransformSelection) ||
 		Out->IsA<FDataflowTransformSelection>(&NewGeometryTransformSelection))
 	{
+		FManagedArrayCollection InCollection = GetValue<FManagedArrayCollection>(Context, &Collection);
+		FBox InBoundingBox = GetValue(Context, &BoundingBox);
+
+		//
+		// If not connected get bounding box of incoming collection
+		//
+		if (!IsConnected(&BoundingBox))
+		{
+			GeometryCollection::Facades::FBoundsFacade BoundsFacade(InCollection);
+			const FBox& BoundingBoxInCollectionSpace = BoundsFacade.GetBoundingBoxInCollectionSpace();
+
+			InBoundingBox = BoundingBoxInCollectionSpace;
+		}
+
 		FDataflowTransformSelection InTransformSelection = GetValue<FDataflowTransformSelection>(Context, &TransformSelection);
 
 		//
@@ -505,8 +533,6 @@ void FPlaneCutterDataflowNode_v2::Evaluate(UE::Dataflow::FContext& Context, cons
 		//
 		if (!IsConnected<FDataflowTransformSelection>(&TransformSelection))
 		{
-			const FManagedArrayCollection& InCollection = GetValue<FManagedArrayCollection>(Context, &Collection);
-
 			GeometryCollection::Facades::FCollectionTransformSelectionFacade TransformSelectionFacade(InCollection);
 			const TArray<int32>& SelectionArr = TransformSelectionFacade.SelectAll();
 
@@ -519,25 +545,23 @@ void FPlaneCutterDataflowNode_v2::Evaluate(UE::Dataflow::FContext& Context, cons
 
 		if (InTransformSelection.AnySelected())
 		{
-			FManagedArrayCollection InCollection = GetValue<FManagedArrayCollection>(Context, &Collection);
-
 			int32 ResultGeometryIndex = FFractureEngineFracturing::PlaneCutter(InCollection,
 				InTransformSelection,
-				GetValue<FBox>(Context, &BoundingBox),
-				GetValue<FTransform>(Context, &Transform),
+				InBoundingBox,
+				GetValue(Context, &Transform),
 				NumPlanes,
-				GetValue<int32>(Context, &RandomSeed),
-				GetValue<float>(Context, &ChanceToFracture),
+				GetValue(Context, &RandomSeed),
+				GetValue(Context, &ChanceToFracture),
 				SplitIslands,
-				GetValue<float>(Context, &Grout),
-				GetValue<float>(Context, &Amplitude),
-				GetValue<float>(Context, &Frequency),
-				GetValue<float>(Context, &Persistence),
-				GetValue<float>(Context, &Lacunarity),
-				GetValue<int32>(Context, &OctaveNumber),
-				GetValue<float>(Context, &PointSpacing),
+				GetValue(Context, &Grout),
+				GetValue(Context, &Amplitude),
+				GetValue(Context, &Frequency),
+				GetValue(Context, &Persistence),
+				GetValue(Context, &Lacunarity),
+				GetValue(Context, &OctaveNumber),
+				GetValue(Context, &PointSpacing),
 				AddSamplesForCollision,
-				GetValue<float>(Context, &CollisionSampleSpacing));
+				GetValue(Context, &CollisionSampleSpacing));
 
 			FDataflowTransformSelection NewSelection;
 			FDataflowTransformSelection OriginalSelection;
@@ -581,6 +605,106 @@ void FPlaneCutterDataflowNode_v2::Evaluate(UE::Dataflow::FContext& Context, cons
 	}
 }
 
+void FPlaneCutterDataflowNode_v2::DebugDraw(UE::Dataflow::FContext& Context, IDataflowDebugDrawInterface& DataflowRenderingInterface, const FDebugDrawParameters& DebugDrawParameters) const
+{
+	if ((DebugDrawParameters.bNodeIsSelected || DebugDrawParameters.bNodeIsPinned))
+	{
+		FManagedArrayCollection InCollection = GetValue<FManagedArrayCollection>(Context, &Collection);
+		FBox InBoundingBox = GetValue(Context, &BoundingBox);
+
+		if (!IsConnected(&BoundingBox))
+		{
+			GeometryCollection::Facades::FBoundsFacade BoundsFacade(InCollection);
+			const FBox& BoundingBoxInCollectionSpace = BoundsFacade.GetBoundingBoxInCollectionSpace();
+
+			InBoundingBox = BoundingBoxInCollectionSpace;
+		}
+
+		float PrincipalAxisLength = 2.f * FMath::Max3(InBoundingBox.GetExtent().X, InBoundingBox.GetExtent().Y, InBoundingBox.GetExtent().Z);
+		float PlaneSize = PrincipalAxisLength * PlaneSizeMultiplier;
+
+		const int32 InNumPlanes = GetValue(Context, &NumPlanes);
+
+		FNoiseSettings NoiseSettings;
+		NoiseSettings.Amplitude = GetValue(Context, &Amplitude);
+		NoiseSettings.Frequency = GetValue(Context, &Frequency);
+		NoiseSettings.Lacunarity = GetValue(Context, &Lacunarity);
+		NoiseSettings.Persistence = GetValue(Context, &Persistence);
+		NoiseSettings.Octaves = GetValue(Context, &OctaveNumber);
+		NoiseSettings.PointSpacing = GetValue(Context, &PointSpacing);
+
+		const int32 InRandomSeed = GetValue<int32>(Context, &RandomSeed);
+
+		TArray<FTransform> PlaneTransforms;
+
+		FFractureEngineFracturing::GenerateSliceTransforms(InBoundingBox, InRandomSeed, NumPlanes, PlaneTransforms);
+
+		if (PlaneTransforms.IsEmpty())
+		{
+			return;
+		}
+
+		FVector Center(0, 0, 0);
+
+		DataflowRenderingInterface.SetLineWidth(LineWidthMultiplier);
+		if (RenderType == EDataflowDebugDrawRenderType::Shaded)
+		{
+			DataflowRenderingInterface.SetShaded(true);
+			DataflowRenderingInterface.SetTranslucent(bTranslucent);
+			DataflowRenderingInterface.SetWireframe(true);
+		}
+		else
+		{
+			DataflowRenderingInterface.SetShaded(false);
+			DataflowRenderingInterface.SetWireframe(true);
+		}
+		DataflowRenderingInterface.SetWorldPriority();
+		DataflowRenderingInterface.SetColor(FLinearColor::Gray);
+
+		FRandomStream RandomStream(InRandomSeed);
+		FNoiseOffsets NoiseOffset(RandomStream);
+
+		TArray<FSimpleDebugDrawMesh> DebugMeshes;
+		DebugMeshes.SetNum(PlaneTransforms.Num());
+
+		for (int32 PlaneIdx = 0; PlaneIdx < PlaneTransforms.Num(); ++PlaneIdx)
+		{
+			if (bRandomizeColors)
+			{
+				DataflowRenderingInterface.SetColor(GetRandomColor(ColorRandomSeed, PlaneIdx));
+			}
+
+			const float Width = PlaneSize;
+			const float Height = PlaneSize;
+
+			constexpr int32 MaxCountPerDim = 2000;
+			const int32 WidthVertexCount = FMath::Min(MaxCountPerDim, Width / NoiseSettings.PointSpacing);
+			const int32 HeightVertexCount = FMath::Min(MaxCountPerDim, Height / NoiseSettings.PointSpacing);
+
+			const FTransform& PlaneTransform = PlaneTransforms[PlaneIdx];
+			const FVector NoisePivot = GetValue(Context, &Transform).GetLocation();
+			FVector Normal = PlaneTransform.GetUnitAxis(EAxis::Z);
+
+			FSimpleDebugDrawMesh Mesh;
+			DebugMeshes[PlaneIdx] = Mesh;
+
+			DebugMeshes[PlaneIdx].MakeRectangleMesh(Center, Width, Height, WidthVertexCount, HeightVertexCount);
+
+			for (int32 VertexIdx = 0; VertexIdx < DebugMeshes[PlaneIdx].GetMaxVertexIndex(); ++VertexIdx)
+			{
+				FVector WorldPos = PlaneTransforms[PlaneIdx].TransformPosition(DebugMeshes[PlaneIdx].Vertices[VertexIdx]);
+				FVector NewWorldPos = WorldPos + NoiseSettings.NoiseVector(WorldPos - NoisePivot, NoiseOffset).Dot(Normal) * Normal;
+				
+				DebugMeshes[PlaneIdx].Vertices[VertexIdx] = NewWorldPos;
+			}
+
+			DataflowRenderingInterface.DrawMesh(DebugMeshes[PlaneIdx]);
+		}
+	}
+}
+
+/*--------------------------------------------------------------------------------------------------------*/
+
 void FExplodedViewDataflowNode::Evaluate(UE::Dataflow::FContext& Context, const FDataflowOutput* Out) const
 {
 	if (Out->IsA<FManagedArrayCollection>(&Collection))
diff --git a/Engine/Plugins/Experimental/GeometryCollectionPlugin/Source/GeometryCollectionNodes/Private/Dataflow/GeometryCollectionMakeNodes.cpp b/Engine/Plugins/Experimental/GeometryCollectionPlugin/Source/GeometryCollectionNodes/Private/Dataflow/GeometryCollectionMakeNodes.cpp
--- a/Engine/Plugins/Experimental/GeometryCollectionPlugin/Source/GeometryCollectionNodes/Private/Dataflow/GeometryCollectionMakeNodes.cpp
+++ b/Engine/Plugins/Experimental/GeometryCollectionPlugin/Source/GeometryCollectionNodes/Private/Dataflow/GeometryCollectionMakeNodes.cpp
@@ -37,6 +37,14 @@
 #include "GeometryCollection/Facades/CollectionTransformFacade.h"
 #include "DynamicMesh/MeshTransforms.h"
 #include "DynamicMesh/DynamicMesh3.h"
+#include "Dataflow/DataflowDebugDraw.h"
+#include "Dataflow/DataflowSimpleDebugDrawMesh.h"
+#include "DynamicMesh/DynamicMesh3.h"
+#include "Generators/SphereGenerator.h"
+#include "Generators/BoxSphereGenerator.h"
+#include "Generators/CapsuleGenerator.h"
+#include "Generators/SweepGenerator.h"
+#include "Generators/GridBoxMeshGenerator.h"
 
 #include UE_INLINE_GENERATED_CPP_BY_NAME(GeometryCollectionMakeNodes)
 
@@ -65,6 +73,10 @@ namespace UE::Dataflow
 		DATAFLOW_NODE_REGISTER_CREATION_FACTORY(FMakeLiteralFloatDataflowNode);
 		DATAFLOW_NODE_REGISTER_CREATION_FACTORY(FMakeLiteralIntDataflowNode);
 		DATAFLOW_NODE_REGISTER_CREATION_FACTORY(FMakeLiteralBoolDataflowNode);
+		DATAFLOW_NODE_REGISTER_CREATION_FACTORY(FMakeSphereMeshDataflowNode);
+		DATAFLOW_NODE_REGISTER_CREATION_FACTORY(FMakeCapsuleMeshDataflowNode);
+		DATAFLOW_NODE_REGISTER_CREATION_FACTORY(FMakeCylinderMeshDataflowNode);
+		DATAFLOW_NODE_REGISTER_CREATION_FACTORY(FMakeBoxMeshDataflowNode);
 	}
 }
 
@@ -84,6 +96,31 @@ void FMakeLiteralStringDataflowNode_v2::Evaluate(UE::Dataflow::FContext& Context
 	}
 }
 
+/*--------------------------------------------------------------------------------------------------------*/
+
+void FMakePointsDataflowNode::DebugDraw(UE::Dataflow::FContext& Context, IDataflowDebugDrawInterface& DataflowRenderingInterface, const FDebugDrawParameters& DebugDrawParameters) const
+{
+	if ((DebugDrawParameters.bNodeIsSelected || DebugDrawParameters.bNodeIsPinned))
+	{
+		FVector Min(UE_BIG_NUMBER);
+		FVector Max(-UE_BIG_NUMBER);
+
+		// Compute (Min, Max) of BoundingBox
+		for (const FVector& Pt : Point)
+		{
+			Min = FVector::Min(Pt, Min);
+			Max = FVector::Max(Pt, Max);
+		}
+
+		DataflowRenderingInterface.SetLineWidth(1.0);
+		DataflowRenderingInterface.SetWireframe(true);
+		DataflowRenderingInterface.SetWorldPriority();
+		DataflowRenderingInterface.SetColor(FLinearColor::Gray);
+		
+		DataflowRenderingInterface.DrawBox(0.5 * FVector(Max - Min), FQuat::Identity, 0.5 * FVector(Min + Max), 1.0);
+	}
+}
+
 void FMakePointsDataflowNode::Evaluate(UE::Dataflow::FContext& Context, const FDataflowOutput* Out) const
 {
 	if (Out->IsA<TArray<FVector>>(&Points))
@@ -92,6 +129,8 @@ void FMakePointsDataflowNode::Evaluate(UE::Dataflow::FContext& Context, const FD
 	}
 }
 
+/*--------------------------------------------------------------------------------------------------------*/
+
 void FMakeBoxDataflowNode::Evaluate(UE::Dataflow::FContext& Context, const FDataflowOutput* Out) const
 {
 	if (Out->IsA<FBox>(&Box))
@@ -113,6 +152,7 @@ void FMakeBoxDataflowNode::Evaluate(UE::Dataflow::FContext& Context, const FData
 	}
 }
 
+/*--------------------------------------------------------------------------------------------------------*/
 
 void FMakeSphereDataflowNode::Evaluate(UE::Dataflow::FContext& Context, const FDataflowOutput* Out) const
 {
@@ -125,6 +165,24 @@ void FMakeSphereDataflowNode::Evaluate(UE::Dataflow::FContext& Context, const FD
 	}
 }
 
+void FMakeSphereDataflowNode::DebugDraw(UE::Dataflow::FContext& Context, IDataflowDebugDrawInterface& DataflowRenderingInterface, const FDebugDrawParameters& DebugDrawParameters) const
+{
+	if ((DebugDrawParameters.bNodeIsSelected || DebugDrawParameters.bNodeIsPinned))
+	{
+		FVector Min = Center - FVector(Radius);
+		FVector Max = Center + FVector(Radius);
+
+		DataflowRenderingInterface.SetLineWidth(1.0);
+		DataflowRenderingInterface.SetWireframe(true);
+		DataflowRenderingInterface.SetWorldPriority();
+		DataflowRenderingInterface.SetColor(FLinearColor::Gray);
+
+		DataflowRenderingInterface.DrawBox(0.5 * FVector(Max - Min), FQuat::Identity, 0.5 * FVector(Min + Max), 1.0);
+	}
+}
+
+/*--------------------------------------------------------------------------------------------------------*/
+
 void FMakeLiteralFloatDataflowNode::Evaluate(UE::Dataflow::FContext& Context, const FDataflowOutput* Out) const
 {
 	if (Out->IsA<float>(&Float))
@@ -346,3 +404,133 @@ void FBreakTransformDataflowNode::Evaluate(UE::Dataflow::FContext& Context, cons
 }
 
 /* -------------------------------------------------------------------------------- */
+
+FMakeSphereMeshDataflowNode::FMakeSphereMeshDataflowNode(const UE::Dataflow::FNodeParameters& InParam, FGuid InGuid)
+	: FDataflowNode(InParam, InGuid)
+{
+	RegisterOutputConnection(&Mesh);
+}
+
+void FMakeSphereMeshDataflowNode::Evaluate(UE::Dataflow::FContext& Context, const FDataflowOutput* Out) const
+{
+	using namespace UE::Geometry;
+
+	if (Out->IsA(&Mesh))
+	{
+		TObjectPtr<UDynamicMesh> NewMesh = NewObject<UDynamicMesh>();
+		NewMesh->Reset();
+
+		FDynamicMesh3& DynMesh = NewMesh->GetMeshRef();
+
+		FSphereGenerator SphereGenerator;
+		SphereGenerator.Radius = FMath::Max(FMathf::ZeroTolerance, Radius);
+		SphereGenerator.NumPhi = FMath::Max(3, NumPhi);
+		SphereGenerator.NumTheta = FMath::Max(3, NumTheta);
+		SphereGenerator.bPolygroupPerQuad = false;
+		SphereGenerator.Generate();
+
+		DynMesh.Copy(&SphereGenerator);
+		SetValue(Context, NewMesh, &Mesh);
+	}
+}
+
+/* -------------------------------------------------------------------------------- */
+
+FMakeCapsuleMeshDataflowNode::FMakeCapsuleMeshDataflowNode(const UE::Dataflow::FNodeParameters& InParam, FGuid InGuid)
+	: FDataflowNode(InParam, InGuid)
+{
+	RegisterOutputConnection(&Mesh);
+}
+
+void FMakeCapsuleMeshDataflowNode::Evaluate(UE::Dataflow::FContext& Context, const FDataflowOutput* Out) const
+{
+	using namespace UE::Geometry;
+
+	if (Out->IsA(&Mesh))
+	{
+		TObjectPtr<UDynamicMesh> NewMesh = NewObject<UDynamicMesh>();
+		NewMesh->Reset();
+
+		FDynamicMesh3& DynMesh = NewMesh->GetMeshRef();
+
+		FCapsuleGenerator CapsuleGenerator;
+		CapsuleGenerator.Radius = FMath::Max(FMathf::ZeroTolerance, Radius);
+		CapsuleGenerator.SegmentLength = FMath::Max(FMathf::ZeroTolerance, SegmentLength);
+		CapsuleGenerator.NumHemisphereArcSteps = FMath::Max(5, NumHemisphereArcSteps);
+		CapsuleGenerator.NumCircleSteps = FMath::Max(3, NumCircleSteps);
+		CapsuleGenerator.NumSegmentSteps = FMath::Max(0, NumSegmentSteps);
+		CapsuleGenerator.bPolygroupPerQuad = false;
+		CapsuleGenerator.Generate();
+
+		DynMesh.Copy(&CapsuleGenerator);
+		SetValue(Context, NewMesh, &Mesh);
+	}
+}
+
+/* -------------------------------------------------------------------------------- */
+
+FMakeCylinderMeshDataflowNode::FMakeCylinderMeshDataflowNode(const UE::Dataflow::FNodeParameters& InParam, FGuid InGuid)
+	: FDataflowNode(InParam, InGuid)
+{
+	RegisterOutputConnection(&Mesh);
+}
+
+void FMakeCylinderMeshDataflowNode::Evaluate(UE::Dataflow::FContext& Context, const FDataflowOutput* Out) const
+{
+	using namespace UE::Geometry;
+
+	if (Out->IsA(&Mesh))
+	{
+		TObjectPtr<UDynamicMesh> NewMesh = NewObject<UDynamicMesh>();
+		NewMesh->Reset();
+
+		FDynamicMesh3& DynMesh = NewMesh->GetMeshRef();
+
+		FCylinderGenerator CylinderGenerator;
+		CylinderGenerator.Radius[0] = FMath::Max(FMathf::ZeroTolerance, Radius1);
+		CylinderGenerator.Radius[1] = FMath::Max(FMathf::ZeroTolerance, Radius2);
+		CylinderGenerator.Height = FMath::Max(FMathf::ZeroTolerance, Height);
+		CylinderGenerator.LengthSamples = LengthSamples;
+		CylinderGenerator.AngleSamples = AngleSamples;
+		CylinderGenerator.bCapped = true;
+		CylinderGenerator.bPolygroupPerQuad = false;
+		CylinderGenerator.Generate();
+
+		DynMesh.Copy(&CylinderGenerator);
+		SetValue(Context, NewMesh, &Mesh);
+	}
+}
+
+/* -------------------------------------------------------------------------------- */
+
+FMakeBoxMeshDataflowNode::FMakeBoxMeshDataflowNode(const UE::Dataflow::FNodeParameters& InParam, FGuid InGuid)
+	: FDataflowNode(InParam, InGuid)
+{
+	RegisterOutputConnection(&Mesh);
+}
+
+void FMakeBoxMeshDataflowNode::Evaluate(UE::Dataflow::FContext& Context, const FDataflowOutput* Out) const
+{
+	using namespace UE::Geometry;
+
+	if (Out->IsA(&Mesh))
+	{
+		TObjectPtr<UDynamicMesh> NewMesh = NewObject<UDynamicMesh>();
+		NewMesh->Reset();
+
+		FDynamicMesh3& DynMesh = NewMesh->GetMeshRef();
+
+		FGridBoxMeshGenerator GridBoxMeshGenerator;
+		GridBoxMeshGenerator.Box = UE::Geometry::FOrientedBox3d(Center, 0.5 * Size);
+		GridBoxMeshGenerator.EdgeVertices = FIndex3i(SubdivisionsX + 1, SubdivisionsY + 1, SubdivisionsZ + 1);
+		GridBoxMeshGenerator.bPolygroupPerQuad = false;
+		GridBoxMeshGenerator.Generate();
+
+		DynMesh.Copy(&GridBoxMeshGenerator);
+		SetValue(Context, NewMesh, &Mesh);
+	}
+}
+
+/* -------------------------------------------------------------------------------- */
+
+
diff --git a/Engine/Plugins/Experimental/GeometryCollectionPlugin/Source/GeometryCollectionNodes/Public/Dataflow/GeometryCollectionFracturingNodes.h b/Engine/Plugins/Experimental/GeometryCollectionPlugin/Source/GeometryCollectionNodes/Public/Dataflow/GeometryCollectionFracturingNodes.h
--- a/Engine/Plugins/Experimental/GeometryCollectionPlugin/Source/GeometryCollectionNodes/Public/Dataflow/GeometryCollectionFracturingNodes.h
+++ b/Engine/Plugins/Experimental/GeometryCollectionPlugin/Source/GeometryCollectionNodes/Public/Dataflow/GeometryCollectionFracturingNodes.h
@@ -7,6 +7,7 @@
 #include "GeometryCollection/ManagedArrayCollection.h"
 #include "Dataflow/DataflowSelection.h"
 #include "FractureEngineFracturing.h"
+#include "Dataflow/DataflowDebugDraw.h"
 
 #include "GeometryCollectionFracturingNodes.generated.h"
 
@@ -605,7 +606,7 @@ public:
 	FManagedArrayCollection Collection;
 
 	/** Bound for plane centers */
-	UPROPERTY(meta = (DataflowInput, DataflowIntrinsic))
+	UPROPERTY(meta = (DataflowInput))
 	FBox BoundingBox = FBox(ForceInit);
 
 	/** Pieces to cut */
@@ -661,7 +662,7 @@ public:
 
 	/** Distance (in cm) between vertices on cut surfaces where noise is added.  Larger spacing between vertices will create more efficient meshes with fewer triangles, but less resolution to see the shape of the added noise  */
 	UPROPERTY(EditAnywhere, Category = "Noise", meta = (DataflowInput, UIMin = 0.f));
-	float PointSpacing = 10.f;
+	float PointSpacing = 20.f;
 
 	/**
 	 * If enabled, add extra vertices (without triangles) to the geometry in regions where vertices are spaced too far apart (e.g. across large triangles)
@@ -679,8 +680,26 @@ public:
 	UPROPERTY(EditAnywhere, Category = "Collision", meta = (DataflowInput, UIMin = 0.f, EditCondition = "AddSamplesForCollision"));
 	float CollisionSampleSpacing = 50.f;
 
+	UPROPERTY(EditAnywhere, Category = "Debug Draw")
+	EDataflowDebugDrawRenderType RenderType = EDataflowDebugDrawRenderType::Wireframe;
+
+	UPROPERTY(EditAnywhere, Category = "Debug Draw", meta = (UIMin = "1.0", UIMax = "10.0", ClampMin = "1.0", ClampMax = "10.0"));
+	float PlaneSizeMultiplier = 2.2f;
+
+	UPROPERTY(EditAnywhere, Category = "Debug Draw", meta = (EditCondition = "RenderType==EDataflowDebugDrawRenderType::Shaded"));
+	bool bTranslucent = true;
+
+	UPROPERTY(EditAnywhere, Category = "Debug Draw");
+	bool bRandomizeColors = true;
+
+	UPROPERTY(EditAnywhere, Category = "Debug Draw", meta = (UIMin = "0"));
+	int32 ColorRandomSeed = 0;
+
+	UPROPERTY(EditAnywhere, Category = "Debug Draw", meta = (UIMin = "0.1", UIMax = "10.0", ClampMin = "0.1", ClampMax = "10.0"));
+	float LineWidthMultiplier = 1.f;
+
 	/** Fractured Pieces */
-	UPROPERTY(meta = (DataflowOutput, DisplayName = "FracturedTransformSelection "))
+	UPROPERTY(meta = (DataflowOutput, DisplayName = "FracturedTransformSelection "));
 	FDataflowTransformSelection NewGeometryTransformSelection;
 
 	FPlaneCutterDataflowNode_v2(const UE::Dataflow::FNodeParameters& InParam, FGuid InGuid = FGuid::NewGuid())
@@ -707,6 +726,8 @@ public:
 	}
 
 	virtual void Evaluate(UE::Dataflow::FContext& Context, const FDataflowOutput* Out) const override;
+
+	virtual void DebugDraw(UE::Dataflow::FContext& Context, IDataflowDebugDrawInterface& DataflowRenderingInterface, const FDebugDrawParameters& DebugDrawParameters) const override;
 };
 
 /**
diff --git a/Engine/Plugins/Experimental/GeometryCollectionPlugin/Source/GeometryCollectionNodes/Public/Dataflow/GeometryCollectionMakeNodes.h b/Engine/Plugins/Experimental/GeometryCollectionPlugin/Source/GeometryCollectionNodes/Public/Dataflow/GeometryCollectionMakeNodes.h
--- a/Engine/Plugins/Experimental/GeometryCollectionPlugin/Source/GeometryCollectionNodes/Public/Dataflow/GeometryCollectionMakeNodes.h
+++ b/Engine/Plugins/Experimental/GeometryCollectionPlugin/Source/GeometryCollectionNodes/Public/Dataflow/GeometryCollectionMakeNodes.h
@@ -93,6 +93,8 @@ public:
 	}
 
 	virtual void Evaluate(UE::Dataflow::FContext& Context, const FDataflowOutput* Out) const override;
+
+	virtual void DebugDraw(UE::Dataflow::FContext& Context, IDataflowDebugDrawInterface& DataflowRenderingInterface, const FDebugDrawParameters& DebugDrawParameters) const override;
 };
 
 
@@ -180,6 +182,8 @@ public:
 	}
 
 	virtual void Evaluate(UE::Dataflow::FContext& Context, const FDataflowOutput* Out) const override;
+
+	virtual void DebugDraw(UE::Dataflow::FContext& Context, IDataflowDebugDrawInterface& DataflowRenderingInterface, const FDebugDrawParameters& DebugDrawParameters) const override;
 };
 
 /**
@@ -350,11 +354,11 @@ public:
 
 
 /**
- *
- *
- *
+ * Make a vector
+ * Deprecated(5.6)
+ * Use MakeVector3 instead
  */
-USTRUCT()
+USTRUCT(meta = (DataflowGeometryCollection, Deprecated = "5.6"))
 struct FMakeLiteralVectorDataflowNode : public FDataflowNode
 {
 	GENERATED_USTRUCT_BODY()
@@ -641,6 +645,174 @@ private:
 	virtual void Evaluate(UE::Dataflow::FContext& Context, const FDataflowOutput* Out) const override;
 };
 
+UENUM(BlueprintType)
+enum class EMakeMeshTypeEnum : uint8
+{
+	Sphere UMETA(DisplayName = "Sphere"),
+	Capsule UMETA(DisplayName = "Capsule"),
+	Cylinder UMETA(DisplayName = "Cylinder"),
+};
+
+/**
+ * Make a sphere mesh
+ */
+USTRUCT()
+struct FMakeSphereMeshDataflowNode : public FDataflowNode
+{
+	GENERATED_USTRUCT_BODY()
+	DATAFLOW_NODE_DEFINE_INTERNAL(FMakeSphereMeshDataflowNode, "MakeSphereMesh", "Generators|Mesh", "")
+	DATAFLOW_NODE_RENDER_TYPE("SurfaceRender", FName("FDynamicMesh3"), "Mesh")
+
+private:
+	/** Sphere Radius */
+	UPROPERTY(EditAnywhere, Category = "Sphere", meta = (UIMin = "0.01"));
+	float Radius = 1.f;
+
+	/** Sphere numphi */
+	UPROPERTY(EditAnywhere, Category = "Sphere", meta = (DisplayName = "Steps Phi", UIMin = "3"));
+	int32 NumPhi = 12;
+
+	/** Sphere numtheta */
+	UPROPERTY(EditAnywhere, Category = "Sphere", meta = (DisplayName = "Steps Theta", UIMin = "3"));
+	int32 NumTheta = 12;
+
+	/** Output mesh */
+	UPROPERTY(meta = (DataflowOutput))
+	TObjectPtr<UDynamicMesh> Mesh;
+
+public:
+	FMakeSphereMeshDataflowNode(const UE::Dataflow::FNodeParameters& InParam, FGuid InGuid = FGuid::NewGuid());
+
+private:
+	virtual void Evaluate(UE::Dataflow::FContext& Context, const FDataflowOutput* Out) const override;
+};
+
+/**
+ * Make a capsule mesh
+ */
+USTRUCT()
+struct FMakeCapsuleMeshDataflowNode : public FDataflowNode
+{
+	GENERATED_USTRUCT_BODY()
+	DATAFLOW_NODE_DEFINE_INTERNAL(FMakeCapsuleMeshDataflowNode, "MakeCapsuleMesh", "Generators|Mesh", "")
+	DATAFLOW_NODE_RENDER_TYPE("SurfaceRender", FName("FDynamicMesh3"), "Mesh")
+
+private:
+	/** Radius of capsule */
+	UPROPERTY(EditAnywhere, Category = "Capsule", meta = (UIMin = "0.01"));
+	float Radius = 1.f;
+
+	/** Length of capsule line segment, so total height is SegmentLength + 2*Radius */
+	UPROPERTY(EditAnywhere, Category = "Capsule", meta = (UIMin = "0.1"));
+	float SegmentLength = 1.f;
+
+	/** Number of vertices along the 90-degree arc from the pole to edge of spherical cap. */
+	UPROPERTY(EditAnywhere, Category = "Capsule", meta = (UIMin = "5"));
+	int32 NumHemisphereArcSteps = 5;
+
+	/** Number of vertices along each circle */
+	UPROPERTY(EditAnywhere, Category = "Capsule", meta = (UIMin = "3"));
+	int32 NumCircleSteps = 6;
+
+	/** Number of subdivisions lengthwise along the cylindrical section */
+	UPROPERTY(EditAnywhere, Category = "Capsule", meta = (UIMin = "0"));
+	int32 NumSegmentSteps = 0;
+
+	/** Output mesh */
+	UPROPERTY(meta = (DataflowOutput))
+	TObjectPtr<UDynamicMesh> Mesh;
+
+public:
+	FMakeCapsuleMeshDataflowNode(const UE::Dataflow::FNodeParameters& InParam, FGuid InGuid = FGuid::NewGuid());
+
+private:
+	virtual void Evaluate(UE::Dataflow::FContext& Context, const FDataflowOutput* Out) const override;
+};
+
+/**
+ * Make a cylinder mesh
+ */
+USTRUCT()
+struct FMakeCylinderMeshDataflowNode : public FDataflowNode
+{
+	GENERATED_USTRUCT_BODY()
+	DATAFLOW_NODE_DEFINE_INTERNAL(FMakeCylinderMeshDataflowNode, "MakeCylinderMesh", "Generators|Mesh", "")
+	DATAFLOW_NODE_RENDER_TYPE("SurfaceRender", FName("FDynamicMesh3"), "Mesh")
+
+private:
+	/** Radius1 of cylinder */
+	UPROPERTY(EditAnywhere, Category = "Cylinder", meta = (UIMin = "0.01"));
+	float Radius1 = 1.f;
+
+	/** Radius2 of cylinder */
+	UPROPERTY(EditAnywhere, Category = "Cylinder", meta = (UIMin = "0.01"));
+	float Radius2 = 1.f;
+
+	/** Height of cylinder */
+	UPROPERTY(EditAnywhere, Category = "Cylinder", meta = (UIMin = "0.01"));
+	float Height = 5.f;
+
+	/** LengthSamples of cylinder */
+	UPROPERTY(EditAnywhere, Category = "Cylinder", meta = (UIMin = "0"));
+	int32 LengthSamples = 0;
+
+	/** LengthSamples of cylinder */
+	UPROPERTY(EditAnywhere, Category = "Cylinder", meta = (UIMin = "0"));
+	int32 AngleSamples = 12;
+
+	/** Output mesh */
+	UPROPERTY(meta = (DataflowOutput))
+	TObjectPtr<UDynamicMesh> Mesh;
+
+public:
+	FMakeCylinderMeshDataflowNode(const UE::Dataflow::FNodeParameters& InParam, FGuid InGuid = FGuid::NewGuid());
+
+private:
+	virtual void Evaluate(UE::Dataflow::FContext& Context, const FDataflowOutput* Out) const override;
+};
+
+/**
+ * Make a box mesh
+ */
+USTRUCT()
+struct FMakeBoxMeshDataflowNode : public FDataflowNode
+{
+	GENERATED_USTRUCT_BODY()
+	DATAFLOW_NODE_DEFINE_INTERNAL(FMakeBoxMeshDataflowNode, "MakeBoxMesh", "Generators|Mesh", "")
+	DATAFLOW_NODE_RENDER_TYPE("SurfaceRender", FName("FDynamicMesh3"), "Mesh")
+
+private:
+	/**  */
+	UPROPERTY(EditAnywhere, Category = "Box");
+	FVector Center = FVector(0.0);
+
+	/**  */
+	UPROPERTY(EditAnywhere, Category = "Box", meta = (UIMin = "0.1"));
+	FVector Size = FVector(5.0);
+
+	/**  */
+	UPROPERTY(EditAnywhere, Category = "Box", meta = (UIMin = "1"));
+	int32 SubdivisionsX = 3;
+
+	/**  */
+	UPROPERTY(EditAnywhere, Category = "Box", meta = (UIMin = "1"));
+	int32 SubdivisionsY = 3;
+
+	/**  */
+	UPROPERTY(EditAnywhere, Category = "Box", meta = (UIMin = "1"));
+	int32 SubdivisionsZ = 3;
+
+	/** Output mesh */
+	UPROPERTY(meta = (DataflowOutput))
+	TObjectPtr<UDynamicMesh> Mesh;
+
+public:
+	FMakeBoxMeshDataflowNode(const UE::Dataflow::FNodeParameters& InParam, FGuid InGuid = FGuid::NewGuid());
+
+private:
+	virtual void Evaluate(UE::Dataflow::FContext& Context, const FDataflowOutput* Out) const override;
+};
+
 namespace UE::Dataflow
 {
 	void GeometryCollectionMakeNodes();
diff --git a/Engine/Source/Runtime/Experimental/Dataflow/Core/Private/Dataflow/DataflowSimpleDebugDrawMesh.cpp b/Engine/Source/Runtime/Experimental/Dataflow/Core/Private/Dataflow/DataflowSimpleDebugDrawMesh.cpp
--- a/Engine/Source/Runtime/Experimental/Dataflow/Core/Private/Dataflow/DataflowSimpleDebugDrawMesh.cpp
+++ b/Engine/Source/Runtime/Experimental/Dataflow/Core/Private/Dataflow/DataflowSimpleDebugDrawMesh.cpp
@@ -1,6 +1,7 @@
 // Copyright Epic Games, Inc. All Rights Reserved.
 
 #include "Dataflow/DataflowSimpleDebugDrawMesh.h"
+#include "VectorTypes.h"
 
 int32 FSimpleDebugDrawMesh::GetMaxVertexIndex() const
 {
@@ -9,7 +10,7 @@ int32 FSimpleDebugDrawMesh::GetMaxVertexIndex() const
 
 bool FSimpleDebugDrawMesh::IsValidVertex(int32 VertexIndex) const
 {
-	return VertexIndex > 0 && VertexIndex < Vertices.Num();
+	return VertexIndex >= 0 && VertexIndex < Vertices.Num();
 }
 
 FVector FSimpleDebugDrawMesh::GetVertexPosition(int32 VertexIndex) const
@@ -17,6 +18,19 @@ FVector FSimpleDebugDrawMesh::GetVertexPosition(int32 VertexIndex) const
 	return Vertices[VertexIndex];
 }
 
+FVector FSimpleDebugDrawMesh::GetVertexNormal(int32 VertexIndex) const
+{
+	return VertexNormals[VertexIndex];
+}
+
+void FSimpleDebugDrawMesh::SetVertex(const int32 VertexIndex, const FVector& VertexPosition)
+{
+	if (IsValidVertex(VertexIndex))
+	{
+		Vertices[VertexIndex] = VertexPosition;
+	}
+}
+
 int32 FSimpleDebugDrawMesh::GetMaxTriangleIndex() const
 {
 	return Triangles.Num();
@@ -24,7 +38,7 @@ int32 FSimpleDebugDrawMesh::GetMaxTriangleIndex() const
 
 bool FSimpleDebugDrawMesh::IsValidTriangle(int32 TriangleIndex) const
 {
-	return TriangleIndex > 0 && TriangleIndex < Triangles.Num();
+	return TriangleIndex >= 0 && TriangleIndex < Triangles.Num();
 }
 
 FIntVector3 FSimpleDebugDrawMesh::GetTriangle(int32 TriangleIndex) const
@@ -32,3 +46,109 @@ FIntVector3 FSimpleDebugDrawMesh::GetTriangle(int32 TriangleIndex) const
 	return Triangles[TriangleIndex];
 }
 
+void FSimpleDebugDrawMesh::SetTriangle(const int32 TriangleIndex, const int32 VertexIndexA, const int32 VertexIndexB, const int32 VertexIndexC)
+{
+	if (IsValidTriangle(TriangleIndex))
+	{
+		Triangles[TriangleIndex] = FIntVector(VertexIndexA, VertexIndexB, VertexIndexC);
+	}
+}
+
+void FSimpleDebugDrawMesh::TransformVertices(const FTransform& Transform) 
+{
+	for (FVector& Vertex : Vertices)
+	{
+		Vertex = Transform.TransformPosition(Vertex);
+	}
+
+	for (FVector& Normal : VertexNormals)
+	{
+		Normal = Transform.TransformVector(Normal);
+	}
+}
+
+static FVector BilinearInterp(const FVector& V00, const FVector& V10, const FVector& V11, const FVector& V01, double Tx, double Ty)
+{
+	FVector A = UE::Geometry::Lerp(V00, V01, Ty);
+	FVector B = UE::Geometry::Lerp(V10, V11, Ty);
+	return UE::Geometry::Lerp(A, B, Tx);
+}
+
+void FSimpleDebugDrawMesh::MakeRectangleMesh(const FVector& Origin, const float Width, const float Height, const int32 WidthVertexCount, const int32 HeightVertexCount)
+{
+	int32 WidthNV = (WidthVertexCount > 1) ? WidthVertexCount : 2;
+	int32 HeightNV = (HeightVertexCount > 1) ? HeightVertexCount : 2;
+
+	int32 TotalNumVertices = WidthNV * HeightNV;
+	int32 TotalNumTriangles = 2 * (WidthNV - 1) * (HeightNV - 1);
+	
+	Vertices.Empty();
+	Triangles.Empty();
+
+	if (TotalNumVertices > 0)
+	{
+		Vertices.AddUninitialized(TotalNumVertices);
+		VertexNormals.AddUninitialized(TotalNumVertices);
+	}
+	if (TotalNumTriangles > 0)
+	{
+		Triangles.AddUninitialized(TotalNumTriangles);
+	}
+
+	// corner vertices
+	FVector V00 = Origin + FVector(-Width / 2.0f, -Height / 2.0f, 0.f); 
+	FVector V01 = Origin + FVector( Width / 2.0f, -Height / 2.0f, 0.f);
+	FVector V11 = Origin + FVector( Width / 2.0f,  Height / 2.0f, 0.f);
+	FVector V10 = Origin + FVector(-Width / 2.0f,  Height / 2.0f, 0.f);
+
+	// Compute normal vector for the plane
+	const FVector Edge1 = V01 - V00;
+	const FVector Edge2 = V10 - V00;
+	FVector Normal = FVector::CrossProduct(Edge1, Edge2);
+	Normal.Normalize();
+
+	int32 Vi = 0;
+	int32 Ti = 0;
+
+	// add vertex rows
+	int32 Start_vi = Vi;
+	for (int32 Yi = 0; Yi < HeightNV; ++Yi)
+	{
+		double Ty = (double)Yi / (double)(HeightNV - 1);
+		for (int32 Xi = 0; Xi < WidthNV; ++Xi)
+		{
+			double Tx = (double)Xi / (double)(WidthNV - 1);
+			SetVertex(Vi, BilinearInterp(V00, V01, V11, V10, Tx, Ty));
+
+			Vi++;
+		}
+	}
+
+	// add triangulated quads
+	for (int32 Y0 = 0; Y0 < HeightNV - 1; ++Y0)
+	{
+		for (int32 X0 = 0; X0 < WidthNV - 1; ++X0)
+		{
+			int32 I00 = Start_vi + Y0 * WidthNV + X0;
+			int32 I10 = Start_vi + (Y0 + 1) * WidthNV + X0;
+			int32 I01 = I00 + 1;
+			int32 I11 = I10 + 1;
+
+			SetTriangle(Ti, I00, I11, I01);
+
+			Ti++;
+
+			SetTriangle(Ti, I00, I10, I11);
+
+			Ti++;
+		}
+	}
+
+	// Add normal vectors
+	for (int32 Idx = 0; Idx < TotalNumVertices; ++Idx)
+	{
+		VertexNormals[Idx] = Normal;
+	}
+}
+
+
diff --git a/Engine/Source/Runtime/Experimental/Dataflow/Core/Public/Dataflow/DataflowDebugDrawInterface.h b/Engine/Source/Runtime/Experimental/Dataflow/Core/Public/Dataflow/DataflowDebugDrawInterface.h
--- a/Engine/Source/Runtime/Experimental/Dataflow/Core/Public/Dataflow/DataflowDebugDrawInterface.h
+++ b/Engine/Source/Runtime/Experimental/Dataflow/Core/Public/Dataflow/DataflowDebugDrawInterface.h
@@ -31,6 +31,7 @@ public:
 		virtual int32 GetMaxVertexIndex() const = 0;
 		virtual bool IsValidVertex(int32 VertexIndex) const = 0;
 		virtual FVector GetVertexPosition(int32 VertexIndex) const = 0;
+		virtual FVector GetVertexNormal(int32 VertexIndex) const = 0;
 
 		virtual int32 GetMaxTriangleIndex() const = 0;
 		virtual bool IsValidTriangle(int32 TriangleIndex) const = 0;
@@ -45,3 +46,5 @@ public:
 	// TODO: More draw functions
 
 };
+
+
diff --git a/Engine/Source/Runtime/Experimental/Dataflow/Core/Public/Dataflow/DataflowSimpleDebugDrawMesh.h b/Engine/Source/Runtime/Experimental/Dataflow/Core/Public/Dataflow/DataflowSimpleDebugDrawMesh.h
--- a/Engine/Source/Runtime/Experimental/Dataflow/Core/Public/Dataflow/DataflowSimpleDebugDrawMesh.h
+++ b/Engine/Source/Runtime/Experimental/Dataflow/Core/Public/Dataflow/DataflowSimpleDebugDrawMesh.h
@@ -8,13 +8,20 @@ struct FSimpleDebugDrawMesh : public IDataflowDebugDrawInterface::IDebugDrawMesh
 {
 	TArray<FVector> Vertices;
 	TArray<FIntVector3> Triangles;
+	TArray<FVector> VertexNormals;
 
 	DATAFLOWCORE_API virtual ~FSimpleDebugDrawMesh() = default;
 	DATAFLOWCORE_API virtual int32 GetMaxVertexIndex() const override;
 	DATAFLOWCORE_API virtual bool IsValidVertex(int32 VertexIndex) const override;
 	DATAFLOWCORE_API virtual FVector GetVertexPosition(int32 VertexIndex) const override;
+	DATAFLOWCORE_API virtual FVector GetVertexNormal(int32 VertexIndex) const override;
+	DATAFLOWCORE_API void SetVertex(const int32 VertexIndex, const FVector& VertexPosition);
 
 	DATAFLOWCORE_API virtual int32 GetMaxTriangleIndex() const override;
 	DATAFLOWCORE_API virtual bool IsValidTriangle(int32 VertexIndex) const override;
 	DATAFLOWCORE_API virtual FIntVector3 GetTriangle(int32 VertexIndex) const override;
+	DATAFLOWCORE_API void SetTriangle(const int32 TriangleIndex, const int32 VertexIndexA, const int32 VertexIndexB, const int32 VertexIndexC);
+
+	DATAFLOWCORE_API void MakeRectangleMesh(const FVector& Origin, const float Width, const float Height, const int32 WidthVertexCount, const int32 HeightVertexCount);
+	DATAFLOWCORE_API void TransformVertices(const FTransform& Transform);
 };
diff --git a/Engine/Source/Runtime/Experimental/Dataflow/Engine/Private/Dataflow/DataflowDebugDraw.cpp b/Engine/Source/Runtime/Experimental/Dataflow/Engine/Private/Dataflow/DataflowDebugDraw.cpp
--- a/Engine/Source/Runtime/Experimental/Dataflow/Engine/Private/Dataflow/DataflowDebugDraw.cpp
+++ b/Engine/Source/Runtime/Experimental/Dataflow/Engine/Private/Dataflow/DataflowDebugDraw.cpp
@@ -115,6 +115,8 @@ void FDataflowDebugDraw::DrawMesh(const IDebugDrawMesh& Mesh) const
 
 	if (bWireframe)		// FDebugRenderSceneProxy only renders solid meshes
 	{
+		DebugRenderSceneProxy->Lines.Reserve(DebugRenderSceneProxy->Lines.Num() + Mesh.GetMaxTriangleIndex() * 3);
+
 		for (int32 TriangleIndex = 0; TriangleIndex < Mesh.GetMaxTriangleIndex(); ++TriangleIndex)
 		{
 			if (Mesh.IsValidTriangle(TriangleIndex))
@@ -140,7 +142,7 @@ void FDataflowDebugDraw::DrawMesh(const IDebugDrawMesh& Mesh) const
 
 		for (int32 VertexIndex = 0; VertexIndex < Mesh.GetMaxVertexIndex(); ++VertexIndex)
 		{
-			if (Mesh.IsValidVertex(VertexIndex))
+			if (!Mesh.IsValidVertex(VertexIndex))
 			{
 				SceneProxyMesh.Vertices.Add(FDynamicMeshVertex(FVector3f(0.0f), FVector2f(0.f), FColor(0)));
 			}
diff --git a/Engine/Source/Runtime/Experimental/Dataflow/Engine/Public/Dataflow/DataflowDebugDraw.h b/Engine/Source/Runtime/Experimental/Dataflow/Engine/Public/Dataflow/DataflowDebugDraw.h
--- a/Engine/Source/Runtime/Experimental/Dataflow/Engine/Public/Dataflow/DataflowDebugDraw.h
+++ b/Engine/Source/Runtime/Experimental/Dataflow/Engine/Public/Dataflow/DataflowDebugDraw.h
@@ -7,6 +7,8 @@
 #include "Engine/EngineTypes.h"
 #include "Math/Color.h"
 
+#include "DataflowDebugDraw.generated.h"
+
 class FDataflowDebugDraw : public IDataflowDebugDrawInterface
 {
 public:
@@ -45,3 +47,10 @@ private:
 	FDebugRenderSceneProxy::EDrawType DrawType = FDebugRenderSceneProxy::EDrawType::Invalid;
 };
 
+UENUM(BlueprintType)
+enum class EDataflowDebugDrawRenderType : uint8
+{
+	Wireframe UMETA(DisplayName = "Wireframe"),
+	Shaded UMETA(DisplayName = "Shaded"),
+};
+
