add split options to the mesh->collection dataflow node
+ add option to add a root cluster bone above the mesh, instead of always adding it as a root. (In the split mesh case, always add a cluster bone.)

diff --git a/Engine/Plugins/Experimental/GeometryCollectionPlugin/Source/GeometryCollectionNodes/Private/Dataflow/GeometryCollectionMeshNodes.cpp b/Engine/Plugins/Experimental/GeometryCollectionPlugin/Source/GeometryCollectionNodes/Private/Dataflow/GeometryCollectionMeshNodes.cpp
--- a/Engine/Plugins/Experimental/GeometryCollectionPlugin/Source/GeometryCollectionNodes/Private/Dataflow/GeometryCollectionMeshNodes.cpp
+++ b/Engine/Plugins/Experimental/GeometryCollectionPlugin/Source/GeometryCollectionNodes/Private/Dataflow/GeometryCollectionMeshNodes.cpp
@@ -20,6 +20,7 @@
 #include "MeshDescriptionToDynamicMesh.h"
 #include "StaticMeshAttributes.h"
 #include "DynamicMeshEditor.h"
+#include "VertexConnectedComponents.h"
 
 #include "GeometryCollectionToDynamicMesh.h"
 
@@ -134,25 +135,78 @@ void FMeshInfoDataflowNode::Evaluate(UE::Dataflow::FContext& Context, const FDat
 }
 
 
+FMeshToCollectionDataflowNode::FMeshToCollectionDataflowNode(const UE::Dataflow::FNodeParameters& InParam, FGuid InGuid)
+	: FDataflowNode(InParam, InGuid)
+{
+	RegisterInputConnection(&Mesh);
+	RegisterInputConnection(&bSplitIslands).SetCanHidePin(true).SetPinIsHidden(true);
+	RegisterInputConnection(&bAddClusterRootForSingleMesh).SetCanHidePin(true).SetPinIsHidden(true);
+
+	RegisterOutputConnection(&Collection);
+}
+
+
 void FMeshToCollectionDataflowNode::Evaluate(UE::Dataflow::FContext& Context, const FDataflowOutput* Out) const
 {
+	using namespace UE::Geometry;
 	if (Out->IsA<FManagedArrayCollection>(&Collection))
 	{
 		if (const TObjectPtr<const UDynamicMesh> InMesh = GetValue(Context, &Mesh))
 		{
-			const UE::Geometry::FDynamicMesh3& DynMesh = InMesh->GetMeshRef();
+			const FDynamicMesh3& DynMesh = InMesh->GetMeshRef();
+
+			bool bInSplitIslands = GetValue(Context, &bSplitIslands);
+			bool bInAlwaysAddRoot = GetValue(Context, &bAddClusterRootForSingleMesh);
 
 			if (DynMesh.VertexCount() > 0)
 			{
-				FMeshDescription MeshDescription;
-				FStaticMeshAttributes Attributes(MeshDescription);
-				Attributes.Register();
+				FGeometryCollection NewGeometryCollection = FGeometryCollection();
+				FGeometryCollectionToDynamicMeshes Convert;
+				FGeometryCollectionToDynamicMeshes::FToCollectionOptions Options;
+
+				TArray<FDynamicMesh3> SplitMeshes;
+				if (bInSplitIslands)
+				{
+					FVertexConnectedComponents Components(DynMesh.MaxVertexID());
+					Components.ConnectTriangles(DynMesh);
+					if (bConnectIslandsByVertexOverlap)
+					{
+						Components.ConnectCloseVertices(DynMesh, ConnectVerticesThreshold, 2);
+					}
+					FDynamicMeshEditor::SplitMesh(&DynMesh, SplitMeshes, [&Components, &DynMesh](int32 TID)
+						{
+							return Components.GetComponent(DynMesh.GetTriangle(TID).A);
+						});
+				}
 
-				FDynamicMeshToMeshDescription Converter;
-				Converter.Convert(&DynMesh, MeshDescription, true);
+				auto AddRoot = [](FGeometryCollection& Collection) -> int32
+					{
+						int32 Idx = Collection.AddElements(1, FGeometryCollection::TransformGroup);
+						Collection.Parent[Idx] = INDEX_NONE;
+						Collection.BoneColor[Idx] = FLinearColor::White;
+						return Idx;
+					};
 
-				FGeometryCollection NewGeometryCollection = FGeometryCollection();
-				FGeometryCollectionEngineConversion::AppendMeshDescription(&MeshDescription, FString(TEXT("TEST")), 0, FTransform().Identity, &NewGeometryCollection);
+				if (SplitMeshes.Num() > 1)
+				{
+					Options.NewMeshParentIndex = AddRoot(NewGeometryCollection);
+					for (const FDynamicMesh3& SplitMesh : SplitMeshes)
+					{
+						Convert.AppendMeshToCollection(NewGeometryCollection, SplitMesh, FTransform::Identity, Options);
+					}
+				}
+				else
+				{
+					if (bInAlwaysAddRoot)
+					{
+						Options.NewMeshParentIndex = AddRoot(NewGeometryCollection);
+					}
+					else
+					{
+						Options.bAllowAppendAsRoot = true;	
+					}
+					Convert.AppendMeshToCollection(NewGeometryCollection, DynMesh, FTransform::Identity, Options);
+				}
 
 				FManagedArrayCollection NewCollection = FManagedArrayCollection();
 				NewGeometryCollection.CopyTo(&NewCollection);
diff --git a/Engine/Plugins/Experimental/GeometryCollectionPlugin/Source/GeometryCollectionNodes/Public/Dataflow/GeometryCollectionMeshNodes.h b/Engine/Plugins/Experimental/GeometryCollectionPlugin/Source/GeometryCollectionNodes/Public/Dataflow/GeometryCollectionMeshNodes.h
--- a/Engine/Plugins/Experimental/GeometryCollectionPlugin/Source/GeometryCollectionNodes/Public/Dataflow/GeometryCollectionMeshNodes.h
+++ b/Engine/Plugins/Experimental/GeometryCollectionPlugin/Source/GeometryCollectionNodes/Public/Dataflow/GeometryCollectionMeshNodes.h
@@ -351,12 +351,23 @@ public:
 	UPROPERTY(meta = (DataflowOutput))
 	FManagedArrayCollection Collection;
 
-	FMeshToCollectionDataflowNode(const UE::Dataflow::FNodeParameters& InParam, FGuid InGuid = FGuid::NewGuid())
-		: FDataflowNode(InParam, InGuid)
-	{
-		RegisterInputConnection(&Mesh);
-		RegisterOutputConnection(&Collection);
-	}
+	/** Whether to split the mesh into multiple bones based on the mesh connectivity */
+	UPROPERTY(EditAnywhere, Category = "General", meta = (DataflowInput));
+	bool bSplitIslands = false;
+
+	/** Whether to consider coincident vertices as connected even if the topology does not connect them */
+	UPROPERTY(EditAnywhere, Category = "General");
+	bool bConnectIslandsByVertexOverlap = false;
+
+	/** Vertices closer than this distance are considered to be overlapping */
+	UPROPERTY(EditAnywhere, Category = "General", meta = (ClampMin = 0.0));
+	float ConnectVerticesThreshold = 0.001;
+
+	/** Whether to add a root cluster for the single mesh case. Note if the mesh is split, the root cluster will always be added. */
+	UPROPERTY(EditAnywhere, Category = "General", meta = (DataflowInput));
+	bool bAddClusterRootForSingleMesh = false;
+
+	FMeshToCollectionDataflowNode(const UE::Dataflow::FNodeParameters& InParam, FGuid InGuid = FGuid::NewGuid());
 
 	virtual void Evaluate(UE::Dataflow::FContext& Context, const FDataflowOutput* Out) const override;
 
diff --git a/Engine/Source/Runtime/MeshConversionEngineTypes/Private/GeometryCollectionToDynamicMesh.cpp b/Engine/Source/Runtime/MeshConversionEngineTypes/Private/GeometryCollectionToDynamicMesh.cpp
--- a/Engine/Source/Runtime/MeshConversionEngineTypes/Private/GeometryCollectionToDynamicMesh.cpp
+++ b/Engine/Source/Runtime/MeshConversionEngineTypes/Private/GeometryCollectionToDynamicMesh.cpp
@@ -499,33 +499,46 @@ int32 FGeometryCollectionToDynamicMeshes::AppendMeshToCollection(FGeometryCollec
 	Collection.TransformIndex[GeometryIdx] = TransformIdx;
 	Collection.TransformToGeometryIndex[TransformIdx] = GeometryIdx;
 	int32 TransformParent = Options.NewMeshParentIndex;
+	// Don't allow the requested parent to be the new transform itself, or a transform that doesn't exist yet
 	if (TransformParent >= TransformIdx)
 	{
 		TransformParent = INDEX_NONE;
 	}
-	// If parent is INDEX_NONE, search for a cluster root (other than self) to use as parent
-	if (TransformParent == INDEX_NONE)
+	// If parent is INDEX_NONE, and we're not allowing geometry to be appended directly as root, search for a cluster root (other than self) to use as parent
+	if (!Options.bAllowAppendAsRoot)
 	{
-		for (int32 Idx = 0; Idx < TransformIdx; ++Idx)
+		if (TransformParent == INDEX_NONE)
 		{
-			if (Collection.Parent[Idx] == -1 && Collection.SimulationType[Idx] == FGeometryCollection::ESimulationTypes::FST_Clustered)
+			for (int32 Idx = 0; Idx < TransformIdx; ++Idx)
 			{
-				TransformParent = Idx;
+				if (Collection.Parent[Idx] == INDEX_NONE && Collection.SimulationType[Idx] == FGeometryCollection::ESimulationTypes::FST_Clustered)
+				{
+					TransformParent = Idx;
+				}
 			}
 		}
+		// If still haven't found a valid parent, add a new root as parent (can occur if collection was empty)
+		if (!Collection.Parent.IsValidIndex(TransformParent))
+		{
+			TransformParent = Collection.AddElements(1, FGeometryCollection::TransformGroup);
+			Collection.Parent[TransformParent] = INDEX_NONE;
+			Collection.BoneColor[TransformParent] = FLinearColor::White;
+		}
+	}
+
+	if (TransformParent > INDEX_NONE)
+	{
+		Collection.BoneName[TransformIdx] = Collection.BoneName[TransformParent] + "_" + FString::FromInt(Collection.Children[TransformParent].Num());
+		Collection.BoneColor[TransformIdx] = Collection.BoneColor[TransformParent];
+		Collection.Children[TransformParent].Add(TransformIdx);
+		Collection.SimulationType[TransformParent] = FGeometryCollection::ESimulationTypes::FST_Clustered;
 	}
-	// If still haven't found a valid parent, add a new root as parent (can occur if collection was empty)
-	if (!Collection.Parent.IsValidIndex(TransformParent))
+	else
 	{
-		TransformParent = Collection.AddElements(1, FGeometryCollection::TransformGroup);
-		Collection.Parent[TransformParent] = INDEX_NONE;
-		Collection.BoneColor[TransformParent] = FLinearColor::White;
+		Collection.BoneName[TransformIdx] = FString::FromInt(Collection.BoneName.Num());
+		Collection.BoneColor[TransformIdx] = FLinearColor::White;
 	}
-	Collection.BoneName[TransformIdx] = Collection.BoneName[TransformParent] + "_" + FString::FromInt(Collection.Children[TransformParent].Num());
-	Collection.BoneColor[TransformIdx] = Collection.BoneColor[TransformParent];
 	Collection.Parent[TransformIdx] = TransformParent;
-	Collection.Children[TransformParent].Add(TransformIdx);
-	Collection.SimulationType[TransformParent] = FGeometryCollection::ESimulationTypes::FST_Clustered;
 
 	Collection.Transform[TransformIdx] = FTransform3f::Identity;
 	Collection.SimulationType[TransformIdx] = FGeometryCollection::ESimulationTypes::FST_Rigid;
diff --git a/Engine/Source/Runtime/MeshConversionEngineTypes/Public/GeometryCollectionToDynamicMesh.h b/Engine/Source/Runtime/MeshConversionEngineTypes/Public/GeometryCollectionToDynamicMesh.h
--- a/Engine/Source/Runtime/MeshConversionEngineTypes/Public/GeometryCollectionToDynamicMesh.h
+++ b/Engine/Source/Runtime/MeshConversionEngineTypes/Public/GeometryCollectionToDynamicMesh.h
@@ -49,6 +49,8 @@ struct FGeometryCollectionToDynamicMeshes
 		bool bDefaultFaceVisible = true;
 		// Whether to set faces with no 'internal' tag as internal (if true) or external (if false)
 		bool bDefaultFaceInternal = false;
+		// Whether the appended geometry is allowed to be added as a root transform (if NewMeshParentIndex == -1)
+		bool bAllowAppendAsRoot = false;
 
 		// Parent index to use if adding a new mesh (with no existing transform) -- if invalid, will add to root
 		int32 NewMeshParentIndex = -1;
