CIS fix for shadowed variable in GeometryCollectionMeshNodes.cpp's AddRoot lambda method

diff --git a/Engine/Plugins/Experimental/GeometryCollectionPlugin/Source/GeometryCollectionNodes/Private/Dataflow/GeometryCollectionMeshNodes.cpp b/Engine/Plugins/Experimental/GeometryCollectionPlugin/Source/GeometryCollectionNodes/Private/Dataflow/GeometryCollectionMeshNodes.cpp
--- a/Engine/Plugins/Experimental/GeometryCollectionPlugin/Source/GeometryCollectionNodes/Private/Dataflow/GeometryCollectionMeshNodes.cpp
+++ b/Engine/Plugins/Experimental/GeometryCollectionPlugin/Source/GeometryCollectionNodes/Private/Dataflow/GeometryCollectionMeshNodes.cpp
@@ -179,11 +179,11 @@ void FMeshToCollectionDataflowNode::Evaluate(UE::Dataflow::FContext& Context, co
 						});
 				}
 
-				auto AddRoot = [](FGeometryCollection& Collection) -> int32
+				auto AddRoot = [](FGeometryCollection& ToCollection) -> int32
 					{
-						int32 Idx = Collection.AddElements(1, FGeometryCollection::TransformGroup);
-						Collection.Parent[Idx] = INDEX_NONE;
-						Collection.BoneColor[Idx] = FLinearColor::White;
+						int32 Idx = ToCollection.AddElements(1, FGeometryCollection::TransformGroup);
+						ToCollection.Parent[Idx] = INDEX_NONE;
+						ToCollection.BoneColor[Idx] = FLinearColor::White;
 						return Idx;
 					};
 
