Dataflow:
- Deprecated FStringAppendDataflowNode
- Added FStringAppendDataflowNode_v2 to handle unlimited number of inputs

diff --git a/Engine/Plugins/Experimental/GeometryCollectionPlugin/Source/GeometryCollectionNodes/Private/Dataflow/GeometryCollectionNodes.cpp b/Engine/Plugins/Experimental/GeometryCollectionPlugin/Source/GeometryCollectionNodes/Private/Dataflow/GeometryCollectionNodes.cpp
--- a/Engine/Plugins/Experimental/GeometryCollectionPlugin/Source/GeometryCollectionNodes/Private/Dataflow/GeometryCollectionNodes.cpp
+++ b/Engine/Plugins/Experimental/GeometryCollectionPlugin/Source/GeometryCollectionNodes/Private/Dataflow/GeometryCollectionNodes.cpp
@@ -53,7 +53,7 @@ namespace UE::Dataflow
 		DATAFLOW_NODE_REGISTER_CREATION_FACTORY(FExpandBoundingBoxDataflowNode);
 		DATAFLOW_NODE_REGISTER_CREATION_FACTORY(FGetBoxLengthsDataflowNode);
 		DATAFLOW_NODE_REGISTER_CREATION_FACTORY(FExpandVectorDataflowNode);
-		DATAFLOW_NODE_REGISTER_CREATION_FACTORY(FStringAppendDataflowNode);
+		DATAFLOW_NODE_REGISTER_CREATION_FACTORY(FStringAppendDataflowNode_v2);
 		DATAFLOW_NODE_REGISTER_CREATION_FACTORY(FHashStringDataflowNode);
 		DATAFLOW_NODE_REGISTER_CREATION_FACTORY(FHashVectorDataflowNode);
 		DATAFLOW_NODE_REGISTER_CREATION_FACTORY(FGetBoundingBoxesFromCollectionDataflowNode);
@@ -88,6 +88,9 @@ namespace UE::Dataflow
 		DATAFLOW_NODE_REGISTER_CREATION_FACTORY(FPointsToCollectionDataflowNode);
 		DATAFLOW_NODE_REGISTER_CREATION_FACTORY(FCollectionToPointsDataflowNode);
 		DATAFLOW_NODE_REGISTER_CREATION_FACTORY(FSpheresToPointsDataflowNode);
+
+		// Deprecated
+		DATAFLOW_NODE_REGISTER_CREATION_FACTORY(FStringAppendDataflowNode);
 	}
 }
 
@@ -248,6 +251,120 @@ void FStringAppendDataflowNode::Evaluate(UE::Dataflow::FContext& Context, const
 	}
 }
 
+//-----------------------------------------------------------------------------------------------
+
+FStringAppendDataflowNode_v2::FStringAppendDataflowNode_v2(const UE::Dataflow::FNodeParameters& InParam, FGuid InGuid)
+	: FDataflowNode(InParam, InGuid)
+{
+	RegisterOutputConnection(&String);
+
+	// Add initial variable inputs
+	for (int32 Index = 0; Index < NumInitialVariableInputs; ++Index)
+	{
+		AddPins();
+	}
+}
+
+void FStringAppendDataflowNode_v2::Evaluate(UE::Dataflow::FContext& Context, const FDataflowOutput* Out) const
+{
+	if (Out->IsA(&String))
+	{
+		FString ResultStr;
+
+		for (int32 Idx = 0; Idx < Inputs.Num(); ++Idx)
+		{
+			const FString InputValue = GetValue(Context, GetConnectionReference(Idx));
+
+			ResultStr = ResultStr + InputValue;
+		}
+
+		SetValue(Context, ResultStr, &String);
+	}
+}
+
+bool FStringAppendDataflowNode_v2::CanAddPin() const
+{
+	return true;
+}
+
+bool FStringAppendDataflowNode_v2::CanRemovePin() const
+{
+	return Inputs.Num() > 0;
+}
+
+UE::Dataflow::TConnectionReference<FDataflowStringConvertibleTypes> FStringAppendDataflowNode_v2::GetConnectionReference(int32 Index) const
+{
+	return { &Inputs[Index], Index, &Inputs };
+}
+
+TArray<UE::Dataflow::FPin> FStringAppendDataflowNode_v2::AddPins()
+{
+	const int32 Index = Inputs.AddDefaulted();
+	const FDataflowInput& Input = RegisterInputArrayConnection(GetConnectionReference(Index));
+	return { { UE::Dataflow::FPin::EDirection::INPUT, Input.GetType(), Input.GetName() } };
+}
+
+TArray<UE::Dataflow::FPin> FStringAppendDataflowNode_v2::GetPinsToRemove() const
+{
+	const int32 Index = (Inputs.Num() - 1);
+	check(Inputs.IsValidIndex(Index));
+	if (const FDataflowInput* const Input = FindInput(GetConnectionReference(Index)))
+	{
+		return { { UE::Dataflow::FPin::EDirection::INPUT, Input->GetType(), Input->GetName() } };
+	}
+	return Super::GetPinsToRemove();
+}
+
+void FStringAppendDataflowNode_v2::OnPinRemoved(const UE::Dataflow::FPin& Pin)
+{
+	const int32 Index = Inputs.Num() - 1;
+	check(Inputs.IsValidIndex(Index));
+#if DO_CHECK
+	const FDataflowInput* const Input = FindInput(GetConnectionReference(Index));
+	check(Input);
+	check(Input->GetName() == Pin.Name);
+	check(Input->GetType() == Pin.Type);
+#endif
+	Inputs.SetNum(Index);
+
+	return Super::OnPinRemoved(Pin);
+}
+
+void FStringAppendDataflowNode_v2::PostSerialize(const FArchive& Ar)
+{
+	if (Ar.IsLoading())
+	{
+		check(Inputs.Num() >= 0);
+		// register new elements from the array as inputs
+		for (int32 Index = 0; Index < Inputs.Num(); ++Index)
+		{
+			FindOrRegisterInputArrayConnection(GetConnectionReference(Index));
+		}
+		if (Ar.IsTransacting())
+		{
+			// if we have more inputs than materials then we need to unregister the inputs 
+			const int32 NumVariableInputs = (GetNumInputs() - NumOtherInputs);
+			const int32 NumInputs = Inputs.Num();
+			if (NumVariableInputs > NumInputs)
+			{
+				// Inputs have been removed.
+				// Temporarily expand Collections so we can get connection references.
+				Inputs.SetNum(NumVariableInputs);
+				for (int32 Index = NumInputs; Index < Inputs.Num(); ++Index)
+				{
+					UnregisterInputConnection(GetConnectionReference(Index));
+				}
+				Inputs.SetNum(NumInputs);
+			}
+		}
+		else
+		{
+			ensureAlways(Inputs.Num() + NumOtherInputs == GetNumInputs());
+		}
+	}
+}
+
+//-----------------------------------------------------------------------------------------------
 
 void FHashStringDataflowNode::Evaluate(UE::Dataflow::FContext& Context, const FDataflowOutput* Out) const
 {
diff --git a/Engine/Plugins/Experimental/GeometryCollectionPlugin/Source/GeometryCollectionNodes/Public/Dataflow/GeometryCollectionNodes.h b/Engine/Plugins/Experimental/GeometryCollectionPlugin/Source/GeometryCollectionNodes/Public/Dataflow/GeometryCollectionNodes.h
--- a/Engine/Plugins/Experimental/GeometryCollectionPlugin/Source/GeometryCollectionNodes/Public/Dataflow/GeometryCollectionNodes.h
+++ b/Engine/Plugins/Experimental/GeometryCollectionPlugin/Source/GeometryCollectionNodes/Public/Dataflow/GeometryCollectionNodes.h
@@ -419,9 +419,10 @@ public:
 /**
  *
  * Concatenates two strings together to make a new string
+ * DEPRECATED - use new version of the same node
  *
  */
-USTRUCT()
+USTRUCT(meta = (Deprecated = "5.6"))
 struct FStringAppendDataflowNode : public FDataflowNode
 {
 	GENERATED_USTRUCT_BODY()
@@ -449,6 +450,43 @@ public:
 
 };
 
+/**
+ *
+ * Concatenates strings together to make a new string
+ *
+ */
+USTRUCT()
+struct FStringAppendDataflowNode_v2 : public FDataflowNode
+{
+	GENERATED_USTRUCT_BODY()
+	DATAFLOW_NODE_DEFINE_INTERNAL(FStringAppendDataflowNode_v2, "StringAppend", "Utilities|String", "")
+
+private:
+	UPROPERTY()
+	TArray<FDataflowStringConvertibleTypes> Inputs;
+
+	UPROPERTY(meta = (DataflowOutput))
+	FString String = FString("");
+
+	//~ Begin FDataflowNode interface
+	virtual TArray<UE::Dataflow::FPin> AddPins() override;
+	virtual bool CanAddPin() const override;
+	virtual bool CanRemovePin() const override;
+	virtual TArray<UE::Dataflow::FPin> GetPinsToRemove() const override;
+	virtual void OnPinRemoved(const UE::Dataflow::FPin& Pin) override;
+	virtual void PostSerialize(const FArchive& Ar) override;
+	virtual void Evaluate(UE::Dataflow::FContext& Context, const FDataflowOutput* Out) const override;
+	//~ End FDataflowNode interface
+
+	UE::Dataflow::TConnectionReference<FDataflowStringConvertibleTypes> GetConnectionReference(int32 Index) const;
+
+	static constexpr int32 NumOtherInputs = 0; // No other inputs
+	static constexpr int32 NumInitialVariableInputs = 2;
+
+public:
+	FStringAppendDataflowNode_v2(const UE::Dataflow::FNodeParameters& InParam, FGuid InGuid = FGuid::NewGuid());
+};
+
 /**
  *
  * Generates a hash value from a string
