Dataflow : Flesh TriangleBoundaryIndices node : : change category and make it experimental to avoid confusion by artists searching for nodes

diff --git a/Engine/Plugins/Experimental/GeometryCollectionPlugin/Source/GeometryCollectionNodes/Private/Dataflow/GeometryCollectionTriangleBoundaryIndicesNode.cpp b/Engine/Plugins/Experimental/GeometryCollectionPlugin/Source/GeometryCollectionNodes/Private/Dataflow/GeometryCollectionTriangleBoundaryIndicesNode.cpp
--- a/Engine/Plugins/Experimental/GeometryCollectionPlugin/Source/GeometryCollectionNodes/Private/Dataflow/GeometryCollectionTriangleBoundaryIndicesNode.cpp
+++ b/Engine/Plugins/Experimental/GeometryCollectionPlugin/Source/GeometryCollectionNodes/Private/Dataflow/GeometryCollectionTriangleBoundaryIndicesNode.cpp
@@ -9,11 +9,18 @@
 
 #include UE_INLINE_GENERATED_CPP_BY_NAME(GeometryCollectionTriangleBoundaryIndicesNode)
 
+FTriangleBoundaryIndicesNode::FTriangleBoundaryIndicesNode(const UE::Dataflow::FNodeParameters& InParam, FGuid InGuid)
+	: FDataflowNode(InParam, InGuid)
+{
+	RegisterInputConnection(&Collection);
+	RegisterOutputConnection(&BoundaryIndicesOut);
+}
+
 void FTriangleBoundaryIndicesNode::Evaluate(UE::Dataflow::FContext& Context, const FDataflowOutput* Out) const
 {
-	if (Out->IsA<TArray<int32>>(&BoundaryIndicesOut))
+	if (Out->IsA(&BoundaryIndicesOut))
 	{
-		const FManagedArrayCollection& InCollection = GetValue<FManagedArrayCollection>(Context, &Collection);
+		const FManagedArrayCollection& InCollection = GetValue(Context, &Collection);
 
 		TArray<Chaos::TVec3<int32>> TriangleMeshArray;
 
@@ -44,6 +51,6 @@ void FTriangleBoundaryIndicesNode::Evaluate(UE::Dataflow::FContext& Context, con
 
 		TArray<int32> IndicesOut = Chaos::Utilities::ComputeBoundaryNodes(TriangleMeshArray);
 
-		SetValue<TArray<int32>>(Context, MoveTemp(IndicesOut), &BoundaryIndicesOut);
+		SetValue(Context, MoveTemp(IndicesOut), &BoundaryIndicesOut);
 	}
 }
diff --git a/Engine/Plugins/Experimental/GeometryCollectionPlugin/Source/GeometryCollectionNodes/Public/Dataflow/GeometryCollectionTriangleBoundaryIndicesNode.h b/Engine/Plugins/Experimental/GeometryCollectionPlugin/Source/GeometryCollectionNodes/Public/Dataflow/GeometryCollectionTriangleBoundaryIndicesNode.h
--- a/Engine/Plugins/Experimental/GeometryCollectionPlugin/Source/GeometryCollectionNodes/Public/Dataflow/GeometryCollectionTriangleBoundaryIndicesNode.h
+++ b/Engine/Plugins/Experimental/GeometryCollectionPlugin/Source/GeometryCollectionNodes/Public/Dataflow/GeometryCollectionTriangleBoundaryIndicesNode.h
@@ -15,31 +15,24 @@ namespace UE_DEPRECATED(5.5, "Use UE::Dataflow instead.") Dataflow {}
 #endif
 
 //Outputs boundary nodes of a triangle mesh
-USTRUCT(meta = (DataflowFlesh))
+USTRUCT(meta = (DataflowFlesh, Experimental))
 struct FTriangleBoundaryIndicesNode : public FDataflowNode
 {
 	GENERATED_USTRUCT_BODY()
-	DATAFLOW_NODE_DEFINE_INTERNAL(FTriangleBoundaryIndicesNode, "TriangleBoundaryIndices", "Geometry Collection|Utiltites", "")
-		DATAFLOW_NODE_RENDER_TYPE("SurfaceRender",FGeometryCollection::StaticType(),  "Collection")
-		
+	DATAFLOW_NODE_DEFINE_INTERNAL(FTriangleBoundaryIndicesNode, "TriangleBoundaryIndices", "Flesh|Experimental", "")
+	DATAFLOW_NODE_RENDER_TYPE("SurfaceRender",FGeometryCollection::StaticType(),  "Collection")
 
-public:
+private:
 	typedef FManagedArrayCollection DataType;
 
-
 	UPROPERTY(meta = (DataflowInput, DisplayName = "Collection", DataflowPassthrough = "Collection"))
-		FManagedArrayCollection Collection;
+	FManagedArrayCollection Collection;
 
 	UPROPERTY(meta = (DataflowOutput, DisplayName = "BoundaryIndices"))
-		TArray<int32> BoundaryIndicesOut;
-
-
-	FTriangleBoundaryIndicesNode(const UE::Dataflow::FNodeParameters& InParam, FGuid InGuid = FGuid::NewGuid())
-		: FDataflowNode(InParam, InGuid)
-	{
-		RegisterInputConnection(&Collection);
-		RegisterOutputConnection(&BoundaryIndicesOut);
-	}
+	TArray<int32> BoundaryIndicesOut;
 
 	virtual void Evaluate(UE::Dataflow::FContext& Context, const FDataflowOutput* Out) const override;
+
+public:
+	FTriangleBoundaryIndicesNode(const UE::Dataflow::FNodeParameters& InParam, FGuid InGuid = FGuid::NewGuid());
 };
\ No newline at end of file
