Dataflow:
- Added FMakePlaneDataflowNode, FMakeDiscMeshDataflowNode, FMakeStairMeshDataflowNode, FMakeRectangleMeshDataflowNode, FMakeTorusMeshDataflowNode nodes
- Added debugdraw to FMakePlaneDataflowNode

diff --git a/Engine/Plugins/Experimental/GeometryCollectionPlugin/Source/GeometryCollectionNodes/Private/Dataflow/GeometryCollectionMakeNodes.cpp b/Engine/Plugins/Experimental/GeometryCollectionPlugin/Source/GeometryCollectionNodes/Private/Dataflow/GeometryCollectionMakeNodes.cpp
--- a/Engine/Plugins/Experimental/GeometryCollectionPlugin/Source/GeometryCollectionNodes/Private/Dataflow/GeometryCollectionMakeNodes.cpp
+++ b/Engine/Plugins/Experimental/GeometryCollectionPlugin/Source/GeometryCollectionNodes/Private/Dataflow/GeometryCollectionMakeNodes.cpp
@@ -45,6 +45,9 @@
 #include "Generators/CapsuleGenerator.h"
 #include "Generators/SweepGenerator.h"
 #include "Generators/GridBoxMeshGenerator.h"
+#include "Generators/DiscMeshGenerator.h"
+#include "Generators/StairGenerator.h"
+#include "Generators/RectangleMeshGenerator.h"
 
 #include UE_INLINE_GENERATED_CPP_BY_NAME(GeometryCollectionMakeNodes)
 
@@ -77,6 +80,11 @@ namespace UE::Dataflow
 		DATAFLOW_NODE_REGISTER_CREATION_FACTORY(FMakeCapsuleMeshDataflowNode);
 		DATAFLOW_NODE_REGISTER_CREATION_FACTORY(FMakeCylinderMeshDataflowNode);
 		DATAFLOW_NODE_REGISTER_CREATION_FACTORY(FMakeBoxMeshDataflowNode);
+		DATAFLOW_NODE_REGISTER_CREATION_FACTORY(FMakePlaneDataflowNode);
+		DATAFLOW_NODE_REGISTER_CREATION_FACTORY(FMakeDiscMeshDataflowNode);
+		DATAFLOW_NODE_REGISTER_CREATION_FACTORY(FMakeStairMeshDataflowNode);
+		DATAFLOW_NODE_REGISTER_CREATION_FACTORY(FMakeRectangleMeshDataflowNode);
+		DATAFLOW_NODE_REGISTER_CREATION_FACTORY(FMakeTorusMeshDataflowNode);
 	}
 }
 
@@ -533,4 +541,267 @@ void FMakeBoxMeshDataflowNode::Evaluate(UE::Dataflow::FContext& Context, const F
 
 /* -------------------------------------------------------------------------------- */
 
+FMakePlaneDataflowNode::FMakePlaneDataflowNode(const UE::Dataflow::FNodeParameters& InParam, FGuid InGuid)
+	: FDataflowNode(InParam, InGuid)
+{
+	RegisterInputConnection(&BasePoint) .SetCanHidePin(true).SetPinIsHidden(true);
+	RegisterInputConnection(&Normal) .SetCanHidePin(true).SetPinIsHidden(true);
+	RegisterOutputConnection(&Plane);
+}
+
+void FMakePlaneDataflowNode::Evaluate(UE::Dataflow::FContext& Context, const FDataflowOutput* Out) const
+{
+	if (Out->IsA(&Plane))
+	{
+		const FVector InBasePoint = GetValue(Context, &BasePoint);
+		const FVector InNormal = GetValue(Context, &Normal);
+
+		const FPlane OutPlane = FPlane(InBasePoint, InNormal);
+
+		SetValue(Context, OutPlane, &Plane);
+	}
+}
+
+void FMakePlaneDataflowNode::DebugDraw(UE::Dataflow::FContext& Context, IDataflowDebugDrawInterface& DataflowRenderingInterface, const FDebugDrawParameters& DebugDrawParameters) const
+{
+	if ((DebugDrawParameters.bNodeIsSelected || DebugDrawParameters.bNodeIsPinned))
+	{
+		DebugDrawRenderSettings.SetDebugDrawSettings(DataflowRenderingInterface);
+
+		FSimpleDebugDrawMesh Mesh;
+		Mesh.MakeRectangleMesh(FVector(0.0), PlaneSizeMultiplier * 10.f, PlaneSizeMultiplier * 10.f, 11, 11);
+
+		const FVector Up = FVector::UpVector;
+		FQuat Quat = FQuat::FindBetweenVectors(Up, Normal);
+
+		FTransform PlaneTransform = FTransform::Identity;
+		PlaneTransform.SetRotation(Quat);
+		PlaneTransform.SetTranslation(BasePoint);
 
+		for (int32 VertexIdx = 0; VertexIdx < Mesh.GetMaxVertexIndex(); ++VertexIdx)
+		{
+			Mesh.Vertices[VertexIdx] = PlaneTransform.TransformPosition(Mesh.Vertices[VertexIdx]);
+		}
+
+		DataflowRenderingInterface.DrawMesh(Mesh);
+	}
+}
+
+/* -------------------------------------------------------------------------------- */
+
+FMakeDiscMeshDataflowNode::FMakeDiscMeshDataflowNode(const UE::Dataflow::FNodeParameters& InParam, FGuid InGuid)
+	: FDataflowNode(InParam, InGuid)
+{
+	RegisterOutputConnection(&Mesh);
+}
+
+void FMakeDiscMeshDataflowNode::Evaluate(UE::Dataflow::FContext& Context, const FDataflowOutput* Out) const
+{
+	using namespace UE::Geometry;
+
+	if (Out->IsA(&Mesh))
+	{
+		TObjectPtr<UDynamicMesh> NewMesh = NewObject<UDynamicMesh>();
+		NewMesh->Reset();
+
+		FDynamicMesh3& DynMesh = NewMesh->GetMeshRef();
+
+		FDiscMeshGenerator DiscGenerator;
+		DiscGenerator.Radius = FMath::Max(FMathf::ZeroTolerance, Radius);
+		DiscGenerator.Normal = FVector3f(Normal);
+		DiscGenerator.AngleSamples = AngleSamples;
+		DiscGenerator.RadialSamples = RadialSamples;
+		DiscGenerator.StartAngle = StartAngle;
+		DiscGenerator.EndAngle = EndAngle;
+		DiscGenerator.bSinglePolygroup = true;
+		DiscGenerator.Generate();
+
+		DynMesh.Copy(&DiscGenerator);
+		SetValue(Context, NewMesh, &Mesh);
+	}
+}
+
+/* -------------------------------------------------------------------------------- */
+
+FMakeStairMeshDataflowNode::FMakeStairMeshDataflowNode(const UE::Dataflow::FNodeParameters& InParam, FGuid InGuid)
+	: FDataflowNode(InParam, InGuid)
+{
+	RegisterOutputConnection(&Mesh);
+}
+
+void FMakeStairMeshDataflowNode::Evaluate(UE::Dataflow::FContext& Context, const FDataflowOutput* Out) const
+{
+	using namespace UE::Geometry;
+
+	if (Out->IsA(&Mesh))
+	{
+		TObjectPtr<UDynamicMesh> NewMesh = NewObject<UDynamicMesh>();
+		NewMesh->Reset();
+
+		FDynamicMesh3& DynMesh = NewMesh->GetMeshRef();
+
+		if (StairType == EDataflowStairTypeEnum::Linear)
+		{
+			FLinearStairGenerator StairGenerator;
+			StairGenerator.NumSteps = NumSteps;
+			StairGenerator.StepWidth = StepWidth;
+			StairGenerator.StepHeight = StepHeight;
+			StairGenerator.StepDepth = StepDepth;
+			StairGenerator.bScaleUVByAspectRatio = true;
+			StairGenerator.bPolygroupPerQuad = true;
+			StairGenerator.Generate();
+
+			DynMesh.Copy(&StairGenerator);
+		}
+		else if (StairType == EDataflowStairTypeEnum::Floating)
+		{
+			FFloatingStairGenerator StairGenerator;
+			StairGenerator.NumSteps = NumSteps;
+			StairGenerator.StepWidth = StepWidth;
+			StairGenerator.StepHeight = StepHeight;
+			StairGenerator.StepDepth = StepDepth;
+			StairGenerator.bScaleUVByAspectRatio = true;
+			StairGenerator.bPolygroupPerQuad = true;
+			StairGenerator.Generate();
+
+			DynMesh.Copy(&StairGenerator);
+		}
+		else if (StairType == EDataflowStairTypeEnum::Curved)
+		{
+			FCurvedStairGenerator StairGenerator;
+			StairGenerator.NumSteps = NumSteps;
+			StairGenerator.StepWidth = StepWidth;
+			StairGenerator.StepHeight = StepHeight;
+			StairGenerator.CurveAngle = CurveAngle;
+			StairGenerator.InnerRadius = InnerRadius;
+			StairGenerator.bScaleUVByAspectRatio = true;
+			StairGenerator.bPolygroupPerQuad = true;
+			StairGenerator.Generate();
+
+			DynMesh.Copy(&StairGenerator);
+		}
+		else if (StairType == EDataflowStairTypeEnum::Spiral)
+		{
+			FSpiralStairGenerator StairGenerator;
+			StairGenerator.NumSteps = NumSteps;
+			StairGenerator.StepWidth = StepWidth;
+			StairGenerator.StepHeight = StepHeight;
+			StairGenerator.CurveAngle = CurveAngle;
+			StairGenerator.InnerRadius = InnerRadius;
+			StairGenerator.bScaleUVByAspectRatio = true;
+			StairGenerator.bPolygroupPerQuad = true;
+			StairGenerator.Generate();
+
+			DynMesh.Copy(&StairGenerator);
+		}
+
+		SetValue(Context, NewMesh, &Mesh);
+	}
+}
+
+/* -------------------------------------------------------------------------------- */
+
+FMakeRectangleMeshDataflowNode::FMakeRectangleMeshDataflowNode(const UE::Dataflow::FNodeParameters& InParam, FGuid InGuid)
+	: FDataflowNode(InParam, InGuid)
+{
+	RegisterInputConnection(&Origin) .SetCanHidePin(true).SetPinIsHidden(true);
+	RegisterInputConnection(&Normal) .SetCanHidePin(true).SetPinIsHidden(true);
+	RegisterOutputConnection(&Mesh);
+}
+
+void FMakeRectangleMeshDataflowNode::Evaluate(UE::Dataflow::FContext& Context, const FDataflowOutput* Out) const
+{
+	using namespace UE::Geometry;
+
+	if (Out->IsA(&Mesh))
+	{
+		const FVector InOrigin = GetValue(Context, &Origin);
+		const FVector InNormal = GetValue(Context, &Normal);
+
+		TObjectPtr<UDynamicMesh> NewMesh = NewObject<UDynamicMesh>();
+		NewMesh->Reset();
+
+		FDynamicMesh3& DynMesh = NewMesh->GetMeshRef();
+
+		FRectangleMeshGenerator RectangleGenerator;
+		RectangleGenerator.Origin = InOrigin;
+		RectangleGenerator.Normal = FVector3f(InNormal);
+		RectangleGenerator.Width = Width;
+		RectangleGenerator.Height = Height;
+		RectangleGenerator.WidthVertexCount = WidthVertexCount;
+		RectangleGenerator.HeightVertexCount = HeightVertexCount;
+		RectangleGenerator.bSinglePolyGroup = true;
+		RectangleGenerator.Generate();
+
+		DynMesh.Copy(&RectangleGenerator);
+		SetValue(Context, NewMesh, &Mesh);
+	}
+}
+
+/* -------------------------------------------------------------------------------- */
+
+FMakeTorusMeshDataflowNode::FMakeTorusMeshDataflowNode(const UE::Dataflow::FNodeParameters& InParam, FGuid InGuid)
+	: FDataflowNode(InParam, InGuid)
+{
+	RegisterInputConnection(&Origin).SetCanHidePin(true).SetPinIsHidden(true);
+	RegisterOutputConnection(&Mesh);
+}
+
+void FMakeTorusMeshDataflowNode::Evaluate(UE::Dataflow::FContext& Context, const FDataflowOutput* Out) const
+{
+	using namespace UE::Geometry;
+
+	if (Out->IsA(&Mesh))
+	{
+		const FVector InOrigin = GetValue(Context, &Origin);
+
+		TObjectPtr<UDynamicMesh> NewMesh = NewObject<UDynamicMesh>();
+		NewMesh->Reset();
+
+		FDynamicMesh3& DynMesh = NewMesh->GetMeshRef();
+
+		TArray<FVector3d> ProfileCurve; ProfileCurve.SetNumUninitialized(ProfileVertexCount);
+		TArray<FFrame3d> SweepCurve; SweepCurve.SetNumUninitialized(SweepVertexCount);
+
+		FVector Vec1(0.0, -Radius1, 0.0);
+		const float RotateAngle1 = 360.f / float(ProfileVertexCount);
+
+		// Create profile curve
+		for (int32 Idx = 0; Idx < ProfileVertexCount; ++Idx)
+		{
+			ProfileCurve[Idx] = Vec1;
+
+			Vec1 = Vec1.RotateAngleAxis(RotateAngle1, FVector::XAxisVector);
+		}
+
+		FVector Vec2(0.0, -Radius2, 0.0);
+		const float RotateAngle2 = 360.f / float(SweepVertexCount);
+		TQuaternion<double> Quat(FVector::ZAxisVector, RotateAngle2, true);
+
+		FFrame3d Frame; // Construct a frame positioned at(0, 0, 0) aligned to the unit axes
+
+		// Create sweep curve
+		for (int32 Idx = 0; Idx < SweepVertexCount; ++Idx)
+		{
+			FFrame3d Frame1 = Frame;
+			Frame1.Origin = Vec2 + InOrigin;
+
+			SweepCurve[Idx] = Frame1;
+
+			Frame.Rotate(Quat);
+			Vec2 = Vec2.RotateAngleAxis(RotateAngle2, FVector::ZAxisVector);
+		}
+
+		FProfileSweepGenerator SweepGenerator;
+		SweepGenerator.ProfileCurve = ProfileCurve;
+		SweepGenerator.SweepCurve = SweepCurve;
+		SweepGenerator.bSweepCurveIsClosed = true;
+		SweepGenerator.bProfileCurveIsClosed = true;
+		SweepGenerator.Generate();
+
+		DynMesh.Copy(&SweepGenerator);
+		SetValue(Context, NewMesh, &Mesh);
+	}
+}
+
+/* -------------------------------------------------------------------------------- */
diff --git a/Engine/Plugins/Experimental/GeometryCollectionPlugin/Source/GeometryCollectionNodes/Public/Dataflow/GeometryCollectionMakeNodes.h b/Engine/Plugins/Experimental/GeometryCollectionPlugin/Source/GeometryCollectionNodes/Public/Dataflow/GeometryCollectionMakeNodes.h
--- a/Engine/Plugins/Experimental/GeometryCollectionPlugin/Source/GeometryCollectionNodes/Public/Dataflow/GeometryCollectionMakeNodes.h
+++ b/Engine/Plugins/Experimental/GeometryCollectionPlugin/Source/GeometryCollectionNodes/Public/Dataflow/GeometryCollectionMakeNodes.h
@@ -10,6 +10,7 @@
 #include "Dataflow/DataflowSelection.h"
 #include "Math/MathFwd.h"
 #include "DynamicMesh/DynamicMesh3.h"
+#include "Dataflow/DataflowDebugDraw.h"
 
 #include "GeometryCollectionMakeNodes.generated.h"
 
@@ -665,15 +666,15 @@ struct FMakeSphereMeshDataflowNode : public FDataflowNode
 
 private:
 	/** Sphere Radius */
-	UPROPERTY(EditAnywhere, Category = "Sphere", meta = (UIMin = "0.01"));
+	UPROPERTY(EditAnywhere, Category = "Sphere", meta = (UIMin = "0.1", ClampMin = "0.1"));
 	float Radius = 1.f;
 
 	/** Sphere numphi */
-	UPROPERTY(EditAnywhere, Category = "Sphere", meta = (DisplayName = "Steps Phi", UIMin = "3"));
+	UPROPERTY(EditAnywhere, Category = "Sphere", meta = (DisplayName = "Steps Phi", UIMin = "3", ClampMin = "3"));
 	int32 NumPhi = 12;
 
 	/** Sphere numtheta */
-	UPROPERTY(EditAnywhere, Category = "Sphere", meta = (DisplayName = "Steps Theta", UIMin = "3"));
+	UPROPERTY(EditAnywhere, Category = "Sphere", meta = (DisplayName = "Steps Theta", UIMin = "3", ClampMin = "3"));
 	int32 NumTheta = 12;
 
 	/** Output mesh */
@@ -699,23 +700,23 @@ struct FMakeCapsuleMeshDataflowNode : public FDataflowNode
 
 private:
 	/** Radius of capsule */
-	UPROPERTY(EditAnywhere, Category = "Capsule", meta = (UIMin = "0.01"));
+	UPROPERTY(EditAnywhere, Category = "Capsule", meta = (UIMin = "0.1", ClampMin = "0.1"));
 	float Radius = 1.f;
 
 	/** Length of capsule line segment, so total height is SegmentLength + 2*Radius */
-	UPROPERTY(EditAnywhere, Category = "Capsule", meta = (UIMin = "0.1"));
+	UPROPERTY(EditAnywhere, Category = "Capsule", meta = (UIMin = "0.1", ClampMin = "0.1"));
 	float SegmentLength = 1.f;
 
 	/** Number of vertices along the 90-degree arc from the pole to edge of spherical cap. */
-	UPROPERTY(EditAnywhere, Category = "Capsule", meta = (UIMin = "5"));
+	UPROPERTY(EditAnywhere, Category = "Capsule", meta = (UIMin = "5", ClampMin = "5"));
 	int32 NumHemisphereArcSteps = 5;
 
 	/** Number of vertices along each circle */
-	UPROPERTY(EditAnywhere, Category = "Capsule", meta = (UIMin = "3"));
+	UPROPERTY(EditAnywhere, Category = "Capsule", meta = (UIMin = "3", ClampMin = "3"));
 	int32 NumCircleSteps = 6;
 
 	/** Number of subdivisions lengthwise along the cylindrical section */
-	UPROPERTY(EditAnywhere, Category = "Capsule", meta = (UIMin = "0"));
+	UPROPERTY(EditAnywhere, Category = "Capsule", meta = (UIMin = "0", ClampMin = "0"));
 	int32 NumSegmentSteps = 0;
 
 	/** Output mesh */
@@ -741,23 +742,23 @@ struct FMakeCylinderMeshDataflowNode : public FDataflowNode
 
 private:
 	/** Radius1 of cylinder */
-	UPROPERTY(EditAnywhere, Category = "Cylinder", meta = (UIMin = "0.01"));
+	UPROPERTY(EditAnywhere, Category = "Cylinder", meta = (UIMin = "0.1", ClampMin = "0.1"));
 	float Radius1 = 1.f;
 
 	/** Radius2 of cylinder */
-	UPROPERTY(EditAnywhere, Category = "Cylinder", meta = (UIMin = "0.01"));
+	UPROPERTY(EditAnywhere, Category = "Cylinder", meta = (UIMin = "0.1", ClampMin = "0.1"));
 	float Radius2 = 1.f;
 
 	/** Height of cylinder */
-	UPROPERTY(EditAnywhere, Category = "Cylinder", meta = (UIMin = "0.01"));
+	UPROPERTY(EditAnywhere, Category = "Cylinder", meta = (UIMin = "0.1", ClampMin = "0.1"));
 	float Height = 5.f;
 
 	/** LengthSamples of cylinder */
-	UPROPERTY(EditAnywhere, Category = "Cylinder", meta = (UIMin = "0"));
+	UPROPERTY(EditAnywhere, Category = "Cylinder", meta = (UIMin = "0", ClampMin = "0"));
 	int32 LengthSamples = 0;
 
-	/** LengthSamples of cylinder */
-	UPROPERTY(EditAnywhere, Category = "Cylinder", meta = (UIMin = "0"));
+	/** AngleSamples of cylinder */
+	UPROPERTY(EditAnywhere, Category = "Cylinder", meta = (UIMin = "4", ClampMin = "4"));
 	int32 AngleSamples = 12;
 
 	/** Output mesh */
@@ -787,19 +788,19 @@ private:
 	FVector Center = FVector(0.0);
 
 	/**  */
-	UPROPERTY(EditAnywhere, Category = "Box", meta = (UIMin = "0.1"));
+	UPROPERTY(EditAnywhere, Category = "Box", meta = (UIMin = "0.1", ClampMin = "0.1"));
 	FVector Size = FVector(5.0);
 
 	/**  */
-	UPROPERTY(EditAnywhere, Category = "Box", meta = (UIMin = "1"));
+	UPROPERTY(EditAnywhere, Category = "Box", meta = (UIMin = "1", ClampMin = "1"));
 	int32 SubdivisionsX = 3;
 
 	/**  */
-	UPROPERTY(EditAnywhere, Category = "Box", meta = (UIMin = "1"));
+	UPROPERTY(EditAnywhere, Category = "Box", meta = (UIMin = "1", ClampMin = "1"));
 	int32 SubdivisionsY = 3;
 
 	/**  */
-	UPROPERTY(EditAnywhere, Category = "Box", meta = (UIMin = "1"));
+	UPROPERTY(EditAnywhere, Category = "Box", meta = (UIMin = "1", ClampMin = "1"));
 	int32 SubdivisionsZ = 3;
 
 	/** Output mesh */
@@ -813,6 +814,235 @@ private:
 	virtual void Evaluate(UE::Dataflow::FContext& Context, const FDataflowOutput* Out) const override;
 };
 
+/**
+ * Make a plane
+ */
+USTRUCT()
+struct FMakePlaneDataflowNode : public FDataflowNode
+{
+	GENERATED_USTRUCT_BODY()
+	DATAFLOW_NODE_DEFINE_INTERNAL(FMakePlaneDataflowNode, "MakePlane", "Generators|Plane", "")
+
+private:
+	/** Base point */
+	UPROPERTY(EditAnywhere, Category = "Plane", meta = (DataflowInput));
+	FVector BasePoint = FVector(0.0);
+
+	/** Normal vector */
+	UPROPERTY(EditAnywhere, Category = "Plane", meta = (DataflowInput));
+	FVector Normal = FVector::UpVector;
+
+	/** DebugDraw settings */
+	UPROPERTY(EditAnywhere, Category = "Debug Draw", meta = (ShowOnlyInnerProperties))
+	FDataflowNodeDebugDrawSettings DebugDrawRenderSettings;
+
+	UPROPERTY(EditAnywhere, Category = "Debug Draw", meta = (UIMin = "1.0", UIMax = "10.0", ClampMin = "1.0", ClampMax = "10.0"));
+	float PlaneSizeMultiplier = 1.f;
+
+	/** Output mesh */
+	UPROPERTY(meta = (DataflowOutput));
+	FPlane Plane = FPlane(ForceInit);
+
+public:
+	FMakePlaneDataflowNode(const UE::Dataflow::FNodeParameters& InParam, FGuid InGuid = FGuid::NewGuid());
+
+private:
+	virtual void Evaluate(UE::Dataflow::FContext& Context, const FDataflowOutput* Out) const override;
+	virtual void DebugDraw(UE::Dataflow::FContext& Context, IDataflowDebugDrawInterface& DataflowRenderingInterface, const FDebugDrawParameters& DebugDrawParameters) const override;
+};
+
+/**
+ * Make a disc mesh
+ */
+USTRUCT()
+struct FMakeDiscMeshDataflowNode : public FDataflowNode
+{
+	GENERATED_USTRUCT_BODY()
+	DATAFLOW_NODE_DEFINE_INTERNAL(FMakeDiscMeshDataflowNode, "MakeDiscMesh", "Generators|Mesh", "")
+	DATAFLOW_NODE_RENDER_TYPE("SurfaceRender", FName("FDynamicMesh3"), "Mesh")
+
+private:
+	/** Radius */
+	UPROPERTY(EditAnywhere, Category = "Disc", meta = (UIMin = "0.1"));
+	float Radius = 1.f;
+
+	/** Normal vector of all vertices will be set to this value. Default is +Z axis. */
+	UPROPERTY(EditAnywhere, Category = "Disc");
+	FVector Normal = FVector::UnitZ();
+
+	/** Number of vertices around circumference */
+	UPROPERTY(EditAnywhere, Category = "Disc", meta = (UIMin = "2", ClampMin = "2"));
+	int32 AngleSamples = 12;
+
+	/** Number of vertices along radial spokes */
+	UPROPERTY(EditAnywhere, Category = "Disc", meta = (UIMin = "2", ClampMin = "2"));
+	int32 RadialSamples = 2;
+
+	/** Start of angle range spanned by disc, in degrees */
+	UPROPERTY(EditAnywhere, Category = "Disc");
+	float StartAngle = 0.f;
+
+	/** End of angle range spanned by disc, in degrees */
+	UPROPERTY(EditAnywhere, Category = "Disc");
+	float EndAngle = 360.f;
+
+	/** Output mesh */
+	UPROPERTY(meta = (DataflowOutput))
+	TObjectPtr<UDynamicMesh> Mesh;
+
+public:
+	FMakeDiscMeshDataflowNode(const UE::Dataflow::FNodeParameters& InParam, FGuid InGuid = FGuid::NewGuid());
+
+private:
+	virtual void Evaluate(UE::Dataflow::FContext& Context, const FDataflowOutput* Out) const override;
+};
+
+UENUM(BlueprintType)
+enum class EDataflowStairTypeEnum : uint8
+{
+	Linear UMETA(DisplayName = "Linear"),
+	Floating UMETA(DisplayName = "Floating"),
+	Curved UMETA(DisplayName = "Curved"),
+	Spiral UMETA(DisplayName = "Spiral"),
+};
+/**
+ * Make a stair mesh
+ */
+USTRUCT()
+struct FMakeStairMeshDataflowNode : public FDataflowNode
+{
+	GENERATED_USTRUCT_BODY()
+	DATAFLOW_NODE_DEFINE_INTERNAL(FMakeStairMeshDataflowNode, "MakeStairMesh", "Generators|Mesh", "")
+	DATAFLOW_NODE_RENDER_TYPE("SurfaceRender", FName("FDynamicMesh3"), "Mesh")
+
+private:
+	/** Type of staircase */
+	UPROPERTY(EditAnywhere, Category = "Stair", meta = (UIMin = "2"));
+	EDataflowStairTypeEnum StairType = EDataflowStairTypeEnum::Linear;
+
+	/** The number of steps in this staircase. */
+	UPROPERTY(EditAnywhere, Category = "Stair", meta = (UIMin = "2"));
+	int32 NumSteps = 8;
+
+	/** The width of each step. */
+	UPROPERTY(EditAnywhere, Category = "Stair", meta = (UIMin = "10.", ClampMin = "10"));
+	float StepWidth = 150.f;
+
+	/** The height of each step. */
+	UPROPERTY(EditAnywhere, Category = "Stair", meta = (UIMin = "10.", ClampMin = "10"));
+	float StepHeight = 20.f;
+
+	/** The height of each step. */
+	UPROPERTY(EditAnywhere, Category = "Stair", meta = (UIMin = "10.", ClampMin = "10", EditCondition = "StairType == EDataflowStairTypeEnum::Linear || StairType == EDataflowStairTypeEnum::Floating", EditConditionHides));
+	float StepDepth = 30.f;
+
+	/** Inner radius of the curved staircase */
+	UPROPERTY(EditAnywhere, Category = "Stair", meta = (UIMin = "10.", ClampMin = "10", EditCondition = "StairType == EDataflowStairTypeEnum::Curved || StairType == EDataflowStairTypeEnum::Spiral", EditConditionHides));
+	float CurveAngle = 90.f;
+
+	/** Curve angle of the staircase (in degrees) */
+	UPROPERTY(EditAnywhere, Category = "Stair", meta = (UIMin = "10.", ClampMin = "10", EditCondition = "StairType == EDataflowStairTypeEnum::Curved || StairType == EDataflowStairTypeEnum::Spiral", EditConditionHides));
+	float InnerRadius = 150.f;
+
+	/** Output mesh */
+	UPROPERTY(meta = (DataflowOutput))
+	TObjectPtr<UDynamicMesh> Mesh;
+
+public:
+	FMakeStairMeshDataflowNode(const UE::Dataflow::FNodeParameters& InParam, FGuid InGuid = FGuid::NewGuid());
+
+private:
+	virtual void Evaluate(UE::Dataflow::FContext& Context, const FDataflowOutput* Out) const override;
+};
+
+/**
+ * Make a rectangle mesh
+ */
+USTRUCT()
+struct FMakeRectangleMeshDataflowNode : public FDataflowNode
+{
+	GENERATED_USTRUCT_BODY()
+	DATAFLOW_NODE_DEFINE_INTERNAL(FMakeRectangleMeshDataflowNode, "MakeRectangleMesh", "Generators|Mesh", "")
+	DATAFLOW_NODE_RENDER_TYPE("SurfaceRender", FName("FDynamicMesh3"), "Mesh")
+
+private:
+	/** Rectangle will be translated so that center is at this point */
+	UPROPERTY(EditAnywhere, Category = "Rectangle", meta = (DataflowInput));
+	FVector Origin = FVector(0.0);
+
+	/** Normal vector of all vertices will be set to this value. Default is +Z axis. */
+	UPROPERTY(EditAnywhere, Category = "Rectangle", meta = (DataflowInput));
+	FVector Normal = FVector::UnitZ();
+
+	/** Width of rectangle */
+	UPROPERTY(EditAnywhere, Category = "Rectangle", meta = (UIMin = "0.1", ClampMin = "0.1"));
+	float Width = 5.f;
+
+	/** Height of rectangle */
+	UPROPERTY(EditAnywhere, Category = "Rectangle", meta = (UIMin = "0.1", ClampMin = "0.1"));
+	float Height = 5.f;
+
+	/** Number of vertices along Width axis */
+	UPROPERTY(EditAnywhere, Category = "Rectangle", meta = (UIMin = "2", ClampMin = "2"));
+	int32 WidthVertexCount = 3;
+
+	/** Number of vertices along Height axis */
+	UPROPERTY(EditAnywhere, Category = "Rectangle", meta = (UIMin = "2", ClampMin = "2"));
+	int32 HeightVertexCount = 3;
+
+	/** Output mesh */
+	UPROPERTY(meta = (DataflowOutput))
+	TObjectPtr<UDynamicMesh> Mesh;
+
+public:
+	FMakeRectangleMeshDataflowNode(const UE::Dataflow::FNodeParameters& InParam, FGuid InGuid = FGuid::NewGuid());
+
+private:
+	virtual void Evaluate(UE::Dataflow::FContext& Context, const FDataflowOutput* Out) const override;
+};
+
+/**
+ * Make a torus mesh
+ */
+USTRUCT()
+struct FMakeTorusMeshDataflowNode : public FDataflowNode
+{
+	GENERATED_USTRUCT_BODY()
+	DATAFLOW_NODE_DEFINE_INTERNAL(FMakeTorusMeshDataflowNode, "MakeTorusMesh", "Generators|Mesh", "")
+	DATAFLOW_NODE_RENDER_TYPE("SurfaceRender", FName("FDynamicMesh3"), "Mesh")
+
+private:
+	/** Torus will be translated so that center is at this point */
+	UPROPERTY(EditAnywhere, Category = "Torus", meta = (DataflowInput));
+	FVector Origin = FVector(0.0);
+
+	/** Radius of the profile */
+	UPROPERTY(EditAnywhere, Category = "Torus", meta = (UIMin = "0.01", ClampMin = "0.01"));
+	float Radius1 = 4.f;
+
+	/** Number of vertices on the profile */
+	UPROPERTY(EditAnywhere, Category = "Torus", meta = (UIMin = "3", ClampMin = "3"));
+	int32 ProfileVertexCount = 12;
+
+	/** Radius of sweep curve */
+	UPROPERTY(EditAnywhere, Category = "Torus", meta = (UIMin = "0.01", ClampMin = "0.01"));
+	float Radius2 = 10.f;
+
+	/** Number of vertices on the sweep curve */
+	UPROPERTY(EditAnywhere, Category = "Torus", meta = (UIMin = "3", ClampMin = "3"));
+	int32 SweepVertexCount = 12;
+
+	/** Output mesh */
+	UPROPERTY(meta = (DataflowOutput))
+	TObjectPtr<UDynamicMesh> Mesh;
+
+public:
+	FMakeTorusMeshDataflowNode(const UE::Dataflow::FNodeParameters& InParam, FGuid InGuid = FGuid::NewGuid());
+
+private:
+	virtual void Evaluate(UE::Dataflow::FContext& Context, const FDataflowOutput* Out) const override;
+};
+
 namespace UE::Dataflow
 {
 	void GeometryCollectionMakeNodes();
diff --git a/Engine/Source/Runtime/Experimental/Dataflow/Engine/Private/Dataflow/DataflowDebugDraw.cpp b/Engine/Source/Runtime/Experimental/Dataflow/Engine/Private/Dataflow/DataflowDebugDraw.cpp
--- a/Engine/Source/Runtime/Experimental/Dataflow/Engine/Private/Dataflow/DataflowDebugDraw.cpp
+++ b/Engine/Source/Runtime/Experimental/Dataflow/Engine/Private/Dataflow/DataflowDebugDraw.cpp
@@ -182,3 +182,26 @@ void FDataflowDebugDraw::DrawSphere(const FVector& Center, double Radius) const
 {
 	DebugRenderSceneProxy->Spheres.Add(FDebugRenderSceneProxy::FSphere(Radius, Center, ColorWithTranslucency, DrawType));
 }
+
+/* ----------------------------------------------------------------------------------------------------------------------- */
+
+void FDataflowNodeDebugDrawSettings::SetDebugDrawSettings(IDataflowDebugDrawInterface& DataflowRenderingInterface) const
+{
+	DataflowRenderingInterface.SetWireframe(true);
+	DataflowRenderingInterface.SetWorldPriority();
+	DataflowRenderingInterface.SetLineWidth(LineWidthMultiplier);
+	if (RenderType == EDataflowDebugDrawRenderType::Shaded)
+	{
+		DataflowRenderingInterface.SetShaded(true);
+		DataflowRenderingInterface.SetTranslucent(bTranslucent);
+		DataflowRenderingInterface.SetWireframe(true);
+	}
+	else
+	{
+		DataflowRenderingInterface.SetShaded(false);
+		DataflowRenderingInterface.SetWireframe(true);
+	}
+	DataflowRenderingInterface.SetWorldPriority();
+	DataflowRenderingInterface.SetColor(Color);
+}
+
diff --git a/Engine/Source/Runtime/Experimental/Dataflow/Engine/Public/Dataflow/DataflowDebugDraw.h b/Engine/Source/Runtime/Experimental/Dataflow/Engine/Public/Dataflow/DataflowDebugDraw.h
--- a/Engine/Source/Runtime/Experimental/Dataflow/Engine/Public/Dataflow/DataflowDebugDraw.h
+++ b/Engine/Source/Runtime/Experimental/Dataflow/Engine/Public/Dataflow/DataflowDebugDraw.h
@@ -54,3 +54,27 @@ enum class EDataflowDebugDrawRenderType : uint8
 	Shaded UMETA(DisplayName = "Shaded"),
 };
 
+/**
+ * DebugDraw basic common settings
+ */
+USTRUCT()
+struct FDataflowNodeDebugDrawSettings
+{
+	GENERATED_USTRUCT_BODY()
+
+	UPROPERTY(EditAnywhere, Category = "Debug Draw")
+	EDataflowDebugDrawRenderType RenderType = EDataflowDebugDrawRenderType::Wireframe;
+
+	UPROPERTY(EditAnywhere, Category = "Debug Draw", meta = (EditCondition = "RenderType==EDataflowDebugDrawRenderType::Shaded"));
+	bool bTranslucent = true;
+
+	UPROPERTY(EditAnywhere, Category = "Debug Draw");
+	FLinearColor Color = FLinearColor::Gray;
+
+	UPROPERTY(EditAnywhere, Category = "Debug Draw", meta = (UIMin = "0.1", UIMax = "10.0", ClampMin = "0.1", ClampMax = "10.0"));
+	float LineWidthMultiplier = 1.f;
+
+	DATAFLOWENGINE_API void SetDebugDrawSettings(IDataflowDebugDrawInterface& DataflowRenderingInterface) const;
+};
+
+
