Clarify the conditions where TargetNumSettings is relevant on various convex-decomposition-related dataflow nodes (GenerateClusterConvexHullsFromChildrenHulls, GenerateClusterConvexHullsFromLeafHulls, MergeConvexHulls), in the comment/tooltip

diff --git a/Engine/Plugins/Experimental/GeometryCollectionPlugin/Source/GeometryCollectionNodes/Public/Dataflow/GeometryCollectionUtilityNodes.h b/Engine/Plugins/Experimental/GeometryCollectionPlugin/Source/GeometryCollectionNodes/Public/Dataflow/GeometryCollectionUtilityNodes.h
--- a/Engine/Plugins/Experimental/GeometryCollectionPlugin/Source/GeometryCollectionNodes/Public/Dataflow/GeometryCollectionUtilityNodes.h
+++ b/Engine/Plugins/Experimental/GeometryCollectionPlugin/Source/GeometryCollectionNodes/Public/Dataflow/GeometryCollectionUtilityNodes.h
@@ -316,7 +316,7 @@ public:
 	UPROPERTY(EditAnywhere, Category = NegativeSpace, meta = (EditCondition = "bProtectNegativeSpace && SampleMethod == ENegativeSpaceSampleMethodDataflowEnum::VoxelSearch", EditConditionHides))
 	bool bOnlyConnectedToHull = false;
 
-	/** Approximate number of spheres to consider when covering negative space */
+	/** Approximate number of spheres to consider when covering negative space. Only applicable with the Uniform Sample Method or if Require Search Sample Coverage is disabled. */
 	UPROPERTY(EditAnywhere, Category = NegativeSpace, meta = (DataflowInput, ClampMin = 1, EditCondition = "bProtectNegativeSpace", EditConditionHides))
 	int32 TargetNumSamples = 50;
 
@@ -397,7 +397,7 @@ public:
 	UPROPERTY(EditAnywhere, Category = NegativeSpace, meta = (EditCondition = "bProtectNegativeSpace && SampleMethod == ENegativeSpaceSampleMethodDataflowEnum::VoxelSearch", EditConditionHides))
 	bool bOnlyConnectedToHull = false;
 
-	/** Approximate number of spheres to consider when covering negative space */
+	/** Approximate number of spheres to consider when covering negative space. Only applicable with the Uniform Sample Method or if Require Search Sample Coverage is disabled. */
 	UPROPERTY(EditAnywhere, Category = NegativeSpace, meta = (DataflowInput, ClampMin = 1, EditCondition = "bProtectNegativeSpace", EditConditionHides))
 	int32 TargetNumSamples = 50;
 
@@ -499,7 +499,7 @@ public:
 	UPROPERTY(EditAnywhere, Category = NegativeSpace, meta = (EditCondition = "bProtectNegativeSpace && SampleMethod == ENegativeSpaceSampleMethodDataflowEnum::VoxelSearch", EditConditionHides))
 	bool bOnlyConnectedToHull = false;
 
-	/** Approximate number of spheres to consider when covering negative space */
+	/** Approximate number of spheres to consider when covering negative space. Only applicable with the Uniform Sample Method or if Require Search Sample Coverage is disabled. */
 	UPROPERTY(EditAnywhere, Category = NegativeSpace, meta = (DataflowInput, ClampMin = 1, EditCondition = "bProtectNegativeSpace", EditConditionHides))
 	int32 TargetNumSamples = 50;
 
