Dataflow:
- Added displaying the normal of the plane to DebugDraw

diff --git a/Engine/Plugins/Experimental/GeometryCollectionPlugin/Source/GeometryCollectionNodes/Private/Dataflow/GeometryCollectionMakeNodes.cpp b/Engine/Plugins/Experimental/GeometryCollectionPlugin/Source/GeometryCollectionNodes/Private/Dataflow/GeometryCollectionMakeNodes.cpp
--- a/Engine/Plugins/Experimental/GeometryCollectionPlugin/Source/GeometryCollectionNodes/Private/Dataflow/GeometryCollectionMakeNodes.cpp
+++ b/Engine/Plugins/Experimental/GeometryCollectionPlugin/Source/GeometryCollectionNodes/Private/Dataflow/GeometryCollectionMakeNodes.cpp
@@ -568,6 +568,9 @@ void FMakePlaneDataflowNode::DebugDraw(UE::Dataflow::FContext& Context, IDataflo
 	{
 		DebugDrawRenderSettings.SetDebugDrawSettings(DataflowRenderingInterface);
 
+		const FVector InBasePoint = GetValue(Context, &BasePoint);
+		FVector InNormal = GetValue(Context, &Normal);
+
 		FSimpleDebugDrawMesh Mesh;
 		Mesh.MakeRectangleMesh(FVector(0.0), PlaneSizeMultiplier * 10.f, PlaneSizeMultiplier * 10.f, 11, 11);
 
@@ -576,7 +579,7 @@ void FMakePlaneDataflowNode::DebugDraw(UE::Dataflow::FContext& Context, IDataflo
 
 		FTransform PlaneTransform = FTransform::Identity;
 		PlaneTransform.SetRotation(Quat);
-		PlaneTransform.SetTranslation(BasePoint);
+		PlaneTransform.SetTranslation(InBasePoint);
 
 		for (int32 VertexIdx = 0; VertexIdx < Mesh.GetMaxVertexIndex(); ++VertexIdx)
 		{
@@ -584,6 +587,10 @@ void FMakePlaneDataflowNode::DebugDraw(UE::Dataflow::FContext& Context, IDataflo
 		}
 
 		DataflowRenderingInterface.DrawMesh(Mesh);
+
+		// Draw normal
+		InNormal.Normalize();
+		DataflowRenderingInterface.DrawLine(InBasePoint, InBasePoint + InNormal * 2.f);
 	}
 }
 
