Dataflow:
- Added DebugDraw to CreateLeafConvexHulls, GenerateClusterConvexHullsFromLeafHulls, GenerateClusterConvexHullsFromChildrenHulls to display convex hulls
- Bug fixes

diff --git a/Engine/Plugins/Experimental/GeometryCollectionPlugin/Source/GeometryCollectionNodes/Private/Dataflow/GeometryCollectionUtilityNodes.cpp b/Engine/Plugins/Experimental/GeometryCollectionPlugin/Source/GeometryCollectionNodes/Private/Dataflow/GeometryCollectionUtilityNodes.cpp
--- a/Engine/Plugins/Experimental/GeometryCollectionPlugin/Source/GeometryCollectionNodes/Private/Dataflow/GeometryCollectionUtilityNodes.cpp
+++ b/Engine/Plugins/Experimental/GeometryCollectionPlugin/Source/GeometryCollectionNodes/Private/Dataflow/GeometryCollectionUtilityNodes.cpp
@@ -8,14 +8,18 @@
 #include "GeometryCollection/GeometryCollectionConvexUtility.h"
 #include "GeometryCollection/Facades/CollectionHierarchyFacade.h"
 #include "GeometryCollection/Facades/CollectionTransformSelectionFacade.h"
+#include "Dataflow/DataflowDebugDrawInterface.h"
+#include "Dataflow/DataflowDebugDraw.h"
+#include "Dataflow/DataflowDynamicMeshDebugDrawMesh.h"
+#include "FractureEngineConvex.h"
+#include "DynamicMesh/DynamicMesh3.h"
 
 #include "Operations/MeshSelfUnion.h"
 #include "MeshQueries.h"
 
 #include "MeshSimplification.h"
 
-
-//#include UE_INLINE_GENERATED_CPP_BY_NAME(GeometryCollectionUtilityNodes)
+#include UE_INLINE_GENERATED_CPP_BY_NAME(GeometryCollectionUtilityNodes)
 
 namespace UE::Dataflow
 {
@@ -66,14 +70,34 @@ void FMakeDataflowConvexDecompositionSettingsNode::Evaluate(UE::Dataflow::FConte
 	}
 }
 
+/* --------------------------------------------------------------------------------------------------------------------------- */
+
+namespace UE::Dataflow::Convex
+{
+	static FLinearColor GetRandomColor(const int32 RandomSeed, int32 Idx)
+	{
+		FRandomStream RandomStream(RandomSeed * 23 + Idx * 4078);
+
+		const uint8 R = static_cast<uint8>(RandomStream.FRandRange(128, 255));
+		const uint8 G = static_cast<uint8>(RandomStream.FRandRange(128, 255));
+		const uint8 B = static_cast<uint8>(RandomStream.FRandRange(128, 255));
+
+		return FLinearColor(FColor(R, G, B, 255));
+	}
+}
+
 FCreateLeafConvexHullsDataflowNode::FCreateLeafConvexHullsDataflowNode(const UE::Dataflow::FNodeParameters& InParam, FGuid InGuid)
 	: FDataflowNode(InParam, InGuid)
 {
 	RegisterInputConnection(&Collection);
-	RegisterInputConnection(&OptionalSelectionFilter);
-	RegisterInputConnection(&SimplificationDistanceThreshold);
-	RegisterInputConnection(&ConvexDecompositionSettings);
+	RegisterInputConnection(&OptionalSelectionFilter).SetCanHidePin(true).SetPinIsHidden(true);
+	RegisterInputConnection(&SimplificationDistanceThreshold).SetCanHidePin(true).SetPinIsHidden(true);
+	RegisterInputConnection(&ConvexDecompositionSettings).SetCanHidePin(true).SetPinIsHidden(true);
 	RegisterOutputConnection(&Collection);
+
+	DebugDrawRenderSettings.RenderType = EDataflowDebugDrawRenderType::Wireframe;
+	DebugDrawRenderSettings.Color = FLinearColor::Green;
+	DebugDrawRenderSettings.LineWidthMultiplier = 5.0;
 }
 
 void FCreateLeafConvexHullsDataflowNode::Evaluate(UE::Dataflow::FContext& Context, const FDataflowOutput* Out) const
@@ -119,6 +143,38 @@ void FCreateLeafConvexHullsDataflowNode::Evaluate(UE::Dataflow::FContext& Contex
 	}
 }
 
+void FCreateLeafConvexHullsDataflowNode::DebugDraw(UE::Dataflow::FContext& Context, IDataflowDebugDrawInterface& DataflowRenderingInterface, const FDebugDrawParameters& DebugDrawParameters) const
+{
+	using namespace UE::Geometry;
+
+	if ((DebugDrawParameters.bNodeIsSelected || DebugDrawParameters.bNodeIsPinned))
+	{
+		DebugDrawRenderSettings.SetDebugDrawSettings(DataflowRenderingInterface);
+
+		if (const FDataflowOutput* Output = FindOutput(&Collection))
+		{
+			const FManagedArrayCollection& InCollection = Output->GetValue(Context, Collection);
+
+			TArray<FDynamicMesh3> HullsMeshes;
+			UE::FractureEngine::Convex::GetConvexHullsAsDynamicMeshes(InCollection, HullsMeshes, false, {});
+
+			int32 Idx = 0;
+			for (const FDynamicMesh3& Mesh : HullsMeshes)
+			{
+				FDynamicMeshDebugDrawMesh DebugdrawMesh(&Mesh);
+				DataflowRenderingInterface.DrawMesh(DebugdrawMesh);
+
+				if (bRandomizeColor)
+				{
+					DataflowRenderingInterface.SetColor(UE::Dataflow::Convex::GetRandomColor(ColorRandomSeed, Idx++));
+				}
+			}
+		}
+	}
+}
+
+/* --------------------------------------------------------------------------------------------------------------------------- */
+
 FSimplifyConvexHullsDataflowNode::FSimplifyConvexHullsDataflowNode(const UE::Dataflow::FNodeParameters& InParam, FGuid InGuid)
 	: FDataflowNode(InParam, InGuid)
 {
@@ -199,6 +255,8 @@ void FCreateNonOverlappingConvexHullsDataflowNode::Evaluate(UE::Dataflow::FConte
 	}
 }
 
+/* --------------------------------------------------------------------------------------------------------------------------- */
+
 // local helper to convert the dataflow enum
 static UE::Geometry::FNegativeSpaceSampleSettings::ESampleMethod ConvertNegativeSpaceSampleMethodDataflowEnum(ENegativeSpaceSampleMethodDataflowEnum SampleMethod)
 {
@@ -216,18 +274,22 @@ FGenerateClusterConvexHullsFromLeafHullsDataflowNode::FGenerateClusterConvexHull
 	: FDataflowNode(InParam, InGuid)
 {
 	RegisterInputConnection(&Collection);
-	RegisterInputConnection(&ConvexCount);
-	RegisterInputConnection(&ErrorTolerance);
-	RegisterInputConnection(&OptionalSelectionFilter);
-	RegisterInputConnection(&bAllowMergingLeafHulls);
-	RegisterInputConnection(&bProtectNegativeSpace);
-	RegisterInputConnection(&TargetNumSamples);
-	RegisterInputConnection(&MinSampleSpacing);
-	RegisterInputConnection(&NegativeSpaceTolerance);
-	RegisterInputConnection(&MinRadius);
+	RegisterInputConnection(&ConvexCount).SetCanHidePin(true).SetPinIsHidden(true);
+	RegisterInputConnection(&ErrorTolerance).SetCanHidePin(true).SetPinIsHidden(true);
+	RegisterInputConnection(&OptionalSelectionFilter).SetCanHidePin(true).SetPinIsHidden(true);
+	RegisterInputConnection(&bAllowMergingLeafHulls).SetCanHidePin(true).SetPinIsHidden(true);
+	RegisterInputConnection(&bProtectNegativeSpace).SetCanHidePin(true).SetPinIsHidden(true);
+	RegisterInputConnection(&TargetNumSamples).SetCanHidePin(true).SetPinIsHidden(true);
+	RegisterInputConnection(&MinSampleSpacing).SetCanHidePin(true).SetPinIsHidden(true);
+	RegisterInputConnection(&NegativeSpaceTolerance).SetCanHidePin(true).SetPinIsHidden(true);
+	RegisterInputConnection(&MinRadius).SetCanHidePin(true).SetPinIsHidden(true);
 
 	RegisterOutputConnection(&Collection);
 	RegisterOutputConnection(&SphereCovering);
+
+	DebugDrawRenderSettings.RenderType = EDataflowDebugDrawRenderType::Wireframe;
+	DebugDrawRenderSettings.Color = FLinearColor::Green;
+	DebugDrawRenderSettings.LineWidthMultiplier = 5.0;
 }
 
 void FGenerateClusterConvexHullsFromLeafHullsDataflowNode::Evaluate(UE::Dataflow::FContext& Context, const FDataflowOutput* Out) const
@@ -302,22 +364,58 @@ void FGenerateClusterConvexHullsFromLeafHullsDataflowNode::Evaluate(UE::Dataflow
 	}
 }
 
+void FGenerateClusterConvexHullsFromLeafHullsDataflowNode::DebugDraw(UE::Dataflow::FContext& Context, IDataflowDebugDrawInterface& DataflowRenderingInterface, const FDebugDrawParameters& DebugDrawParameters) const
+{
+	using namespace UE::Geometry;
+
+	if ((DebugDrawParameters.bNodeIsSelected || DebugDrawParameters.bNodeIsPinned))
+	{
+		DebugDrawRenderSettings.SetDebugDrawSettings(DataflowRenderingInterface);
+
+		if (const FDataflowOutput* Output = FindOutput(&Collection))
+		{
+			const FManagedArrayCollection& InCollection = Output->GetValue(Context, Collection);
+
+			TArray<FDynamicMesh3> HullsMeshes;
+			UE::FractureEngine::Convex::GetConvexHullsAsDynamicMeshes(InCollection, HullsMeshes, false, {});
+
+			int32 Idx = 0;
+			for (const FDynamicMesh3& Mesh : HullsMeshes)
+			{
+				FDynamicMeshDebugDrawMesh DebugdrawMesh(&Mesh);
+				DataflowRenderingInterface.DrawMesh(DebugdrawMesh);
+
+				if (bRandomizeColor)
+				{
+					DataflowRenderingInterface.SetColor(UE::Dataflow::Convex::GetRandomColor(ColorRandomSeed, Idx++));
+				}
+			}
+		}
+	}
+}
+
+/* --------------------------------------------------------------------------------------------------------------------------- */
+
 FGenerateClusterConvexHullsFromChildrenHullsDataflowNode::FGenerateClusterConvexHullsFromChildrenHullsDataflowNode(const UE::Dataflow::FNodeParameters& InParam, FGuid InGuid)
 	: FDataflowNode(InParam, InGuid)
 {
 	RegisterInputConnection(&Collection);
-	RegisterInputConnection(&ConvexCount);
-	RegisterInputConnection(&ErrorTolerance);
-	RegisterInputConnection(&OptionalSelectionFilter);
-	RegisterInputConnection(&bAllowMergingLeafHulls);
-	RegisterInputConnection(&bProtectNegativeSpace);
-	RegisterInputConnection(&TargetNumSamples);
-	RegisterInputConnection(&MinSampleSpacing);
-	RegisterInputConnection(&NegativeSpaceTolerance);
-	RegisterInputConnection(&MinRadius);
+	RegisterInputConnection(&ConvexCount).SetCanHidePin(true).SetPinIsHidden(true);
+	RegisterInputConnection(&ErrorTolerance).SetCanHidePin(true).SetPinIsHidden(true);
+	RegisterInputConnection(&OptionalSelectionFilter).SetCanHidePin(true).SetPinIsHidden(true);
+	RegisterInputConnection(&bAllowMergingLeafHulls).SetCanHidePin(true).SetPinIsHidden(true);
+	RegisterInputConnection(&bProtectNegativeSpace).SetCanHidePin(true).SetPinIsHidden(true);
+	RegisterInputConnection(&TargetNumSamples).SetCanHidePin(true).SetPinIsHidden(true);
+	RegisterInputConnection(&MinSampleSpacing).SetCanHidePin(true).SetPinIsHidden(true);
+	RegisterInputConnection(&NegativeSpaceTolerance).SetCanHidePin(true).SetPinIsHidden(true);
+	RegisterInputConnection(&MinRadius).SetCanHidePin(true).SetPinIsHidden(true);
 
 	RegisterOutputConnection(&Collection);
 	RegisterOutputConnection(&SphereCovering);
+
+	DebugDrawRenderSettings.RenderType = EDataflowDebugDrawRenderType::Wireframe;
+	DebugDrawRenderSettings.Color = FLinearColor::Green;
+	DebugDrawRenderSettings.LineWidthMultiplier = 5.0;
 }
 
 void FGenerateClusterConvexHullsFromChildrenHullsDataflowNode::Evaluate(UE::Dataflow::FContext& Context, const FDataflowOutput* Out) const
@@ -392,6 +490,38 @@ void FGenerateClusterConvexHullsFromChildrenHullsDataflowNode::Evaluate(UE::Data
 	}
 }
 
+void FGenerateClusterConvexHullsFromChildrenHullsDataflowNode::DebugDraw(UE::Dataflow::FContext& Context, IDataflowDebugDrawInterface& DataflowRenderingInterface, const FDebugDrawParameters& DebugDrawParameters) const
+{
+	using namespace UE::Geometry;
+
+	if ((DebugDrawParameters.bNodeIsSelected || DebugDrawParameters.bNodeIsPinned))
+	{
+		DebugDrawRenderSettings.SetDebugDrawSettings(DataflowRenderingInterface);
+
+		if (const FDataflowOutput* Output = FindOutput(&Collection))
+		{
+			const FManagedArrayCollection& InCollection = Output->GetValue(Context, Collection);
+
+			TArray<FDynamicMesh3> HullsMeshes;
+			UE::FractureEngine::Convex::GetConvexHullsAsDynamicMeshes(InCollection, HullsMeshes, false, {});
+
+			int32 Idx = 0;
+			for (const FDynamicMesh3& Mesh : HullsMeshes)
+			{
+				FDynamicMeshDebugDrawMesh DebugdrawMesh(&Mesh);
+				DataflowRenderingInterface.DrawMesh(DebugdrawMesh);
+
+				if (bRandomizeColor)
+				{
+					DataflowRenderingInterface.SetColor(UE::Dataflow::Convex::GetRandomColor(ColorRandomSeed, Idx++));
+				}
+			}
+		}
+	}
+}
+
+/* --------------------------------------------------------------------------------------------------------------------------- */
+
 FMergeConvexHullsDataflowNode::FMergeConvexHullsDataflowNode(const UE::Dataflow::FNodeParameters& InParam, FGuid InGuid)
 	: FDataflowNode(InParam, InGuid)
 {
diff --git a/Engine/Plugins/Experimental/GeometryCollectionPlugin/Source/GeometryCollectionNodes/Public/Dataflow/GeometryCollectionUtilityNodes.h b/Engine/Plugins/Experimental/GeometryCollectionPlugin/Source/GeometryCollectionNodes/Public/Dataflow/GeometryCollectionUtilityNodes.h
--- a/Engine/Plugins/Experimental/GeometryCollectionPlugin/Source/GeometryCollectionNodes/Public/Dataflow/GeometryCollectionUtilityNodes.h
+++ b/Engine/Plugins/Experimental/GeometryCollectionPlugin/Source/GeometryCollectionNodes/Public/Dataflow/GeometryCollectionUtilityNodes.h
@@ -8,6 +8,7 @@
 #include "GeometryCollection/ManagedArrayCollection.h"
 #include "GeometryCollection/GeometryCollectionConvexUtility.h"
 #include "CompGeom/ConvexDecomposition3.h"
+#include "Dataflow/DataflowDebugDraw.h"
 
 #include "FractureEngineConvex.h"
 #include "FractureEngineUtility.h"
@@ -114,8 +115,9 @@ struct FCreateLeafConvexHullsDataflowNode : public FDataflowNode
 {
 	GENERATED_USTRUCT_BODY()
 	DATAFLOW_NODE_DEFINE_INTERNAL(FCreateLeafConvexHullsDataflowNode, "CreateLeafConvexHulls", "GeometryCollection|Utilities", "")
+	DATAFLOW_NODE_RENDER_TYPE("SurfaceRender", FGeometryCollection::StaticType(), "Collection")
 
-public:
+private:
 	UPROPERTY(meta = (DataflowInput, DataflowOutput))
 	FManagedArrayCollection Collection;
 
@@ -146,9 +148,22 @@ public:
 	UPROPERTY(EditAnywhere, Category = Convex, meta = (DataflowInput))
 	FDataflowConvexDecompositionSettings ConvexDecompositionSettings;
 
-	FCreateLeafConvexHullsDataflowNode(const UE::Dataflow::FNodeParameters& InParam, FGuid InGuid = FGuid::NewGuid());
+	UPROPERTY(EditAnywhere, Category = "Debug Draw", meta = (ShowOnlyInnerProperties))
+	FDataflowNodeDebugDrawSettings DebugDrawRenderSettings;
+
+	/** Randomize color per convex hull */
+	UPROPERTY(EditAnywhere, Category = "Debug Draw")
+	bool bRandomizeColor = true;
+
+	/** Random seed */
+	UPROPERTY(EditAnywhere, Category = "Debug Draw", meta = (ClampMin = "0"))
+	int32 ColorRandomSeed = 0;
 
 	virtual void Evaluate(UE::Dataflow::FContext& Context, const FDataflowOutput* Out) const override;
+	virtual void DebugDraw(UE::Dataflow::FContext& Context, IDataflowDebugDrawInterface& DataflowRenderingInterface, const FDebugDrawParameters& DebugDrawParameters) const override;
+
+public:
+	FCreateLeafConvexHullsDataflowNode(const UE::Dataflow::FNodeParameters& InParam, FGuid InGuid = FGuid::NewGuid());
 };
 
 USTRUCT(meta = (DataflowGeometryCollection))
@@ -263,8 +278,9 @@ struct FGenerateClusterConvexHullsFromLeafHullsDataflowNode : public FDataflowNo
 {
 	GENERATED_USTRUCT_BODY()
 	DATAFLOW_NODE_DEFINE_INTERNAL(FGenerateClusterConvexHullsFromLeafHullsDataflowNode, "GenerateClusterConvexHullsFromLeafHulls", "GeometryCollection|Utilities", "")
+	DATAFLOW_NODE_RENDER_TYPE("SurfaceRender", FGeometryCollection::StaticType(), "Collection")
 
-public:
+private:
 	UPROPERTY(meta = (DataflowInput, DataflowOutput))
 	FManagedArrayCollection Collection;
 
@@ -332,9 +348,23 @@ public:
 	UPROPERTY(EditAnywhere, Category = NegativeSpace, meta = (DataflowInput, ClampMin = 0, Units = cm, EditCondition = "bProtectNegativeSpace", EditConditionHides))
 	double MinRadius = 10.0;
 
-	FGenerateClusterConvexHullsFromLeafHullsDataflowNode(const UE::Dataflow::FNodeParameters& InParam, FGuid InGuid = FGuid::NewGuid());
+	UPROPERTY(EditAnywhere, Category = "Debug Draw", meta = (ShowOnlyInnerProperties))
+	FDataflowNodeDebugDrawSettings DebugDrawRenderSettings;
+
+	/** Randomize color per convex hull */
+	UPROPERTY(EditAnywhere, Category = "Debug Draw")
+	bool bRandomizeColor = true;
+
+	/** Random seed */
+	UPROPERTY(EditAnywhere, Category = "Debug Draw", meta = (ClampMin = "0"))
+	int32 ColorRandomSeed = 0;
 
 	virtual void Evaluate(UE::Dataflow::FContext& Context, const FDataflowOutput* Out) const override;
+	virtual void DebugDraw(UE::Dataflow::FContext& Context, IDataflowDebugDrawInterface& DataflowRenderingInterface, const FDebugDrawParameters& DebugDrawParameters) const override;
+
+public:
+	FGenerateClusterConvexHullsFromLeafHullsDataflowNode(const UE::Dataflow::FNodeParameters& InParam, FGuid InGuid = FGuid::NewGuid());
+
 
 };
 
@@ -348,8 +378,9 @@ struct FGenerateClusterConvexHullsFromChildrenHullsDataflowNode : public FDatafl
 {
 	GENERATED_USTRUCT_BODY()
 	DATAFLOW_NODE_DEFINE_INTERNAL(FGenerateClusterConvexHullsFromChildrenHullsDataflowNode, "GenerateClusterConvexHullsFromChildrenHulls", "GeometryCollection|Utilities", "")
+	DATAFLOW_NODE_RENDER_TYPE("SurfaceRender", FGeometryCollection::StaticType(), "Collection")
 
-public:
+private:
 	UPROPERTY(meta = (DataflowInput, DataflowOutput))
 	FManagedArrayCollection Collection;
 	
@@ -413,9 +444,23 @@ public:
 	UPROPERTY(EditAnywhere, Category = NegativeSpace, meta = (DataflowInput, ClampMin = 0, Units = cm, EditCondition = "bProtectNegativeSpace", EditConditionHides))
 	double MinRadius = 10.0;
 
-	FGenerateClusterConvexHullsFromChildrenHullsDataflowNode(const UE::Dataflow::FNodeParameters& InParam, FGuid InGuid = FGuid::NewGuid());
+	UPROPERTY(EditAnywhere, Category = "Debug Draw", meta = (ShowOnlyInnerProperties))
+	FDataflowNodeDebugDrawSettings DebugDrawRenderSettings;
+
+	/** Randomize color per convex hull */
+	UPROPERTY(EditAnywhere, Category = "Debug Draw")
+	bool bRandomizeColor = true;
+
+	/** Random seed */
+	UPROPERTY(EditAnywhere, Category = "Debug Draw", meta = (ClampMin = "0"))
+	int32 ColorRandomSeed = 0;
 
 	virtual void Evaluate(UE::Dataflow::FContext& Context, const FDataflowOutput* Out) const override;
+	virtual void DebugDraw(UE::Dataflow::FContext& Context, IDataflowDebugDrawInterface& DataflowRenderingInterface, const FDebugDrawParameters& DebugDrawParameters) const override;
+
+public:
+	FGenerateClusterConvexHullsFromChildrenHullsDataflowNode(const UE::Dataflow::FNodeParameters& InParam, FGuid InGuid = FGuid::NewGuid());
+
 
 };
 
diff --git a/Engine/Source/Runtime/Experimental/Dataflow/Core/Private/Dataflow/DataflowDynamicMeshDebugDrawMesh.cpp b/Engine/Source/Runtime/Experimental/Dataflow/Core/Private/Dataflow/DataflowDynamicMeshDebugDrawMesh.cpp
--- a/Engine/Source/Runtime/Experimental/Dataflow/Core/Private/Dataflow/DataflowDynamicMeshDebugDrawMesh.cpp
+++ b/Engine/Source/Runtime/Experimental/Dataflow/Core/Private/Dataflow/DataflowDynamicMeshDebugDrawMesh.cpp
@@ -32,6 +32,16 @@ FVector FDynamicMeshDebugDrawMesh::GetVertexPosition(int32 VertexIndex) const
 	return FVector(0.0);
 }
 
+FVector FDynamicMeshDebugDrawMesh::GetVertexNormal(int32 VertexIndex) const
+{
+	if (IsValidVertex(VertexIndex))
+	{
+		return FVector(DynamicMesh->GetVertexNormal(VertexIndex));
+	}
+
+	return FVector(0.0);
+}
+
 int32 FDynamicMeshDebugDrawMesh::GetMaxTriangleIndex() const
 {
 	if (DynamicMesh)
diff --git a/Engine/Source/Runtime/Experimental/Dataflow/Core/Public/Dataflow/DataflowDynamicMeshDebugDrawMesh.h b/Engine/Source/Runtime/Experimental/Dataflow/Core/Public/Dataflow/DataflowDynamicMeshDebugDrawMesh.h
--- a/Engine/Source/Runtime/Experimental/Dataflow/Core/Public/Dataflow/DataflowDynamicMeshDebugDrawMesh.h
+++ b/Engine/Source/Runtime/Experimental/Dataflow/Core/Public/Dataflow/DataflowDynamicMeshDebugDrawMesh.h
@@ -6,13 +6,14 @@
 
 struct FDynamicMeshDebugDrawMesh : public IDataflowDebugDrawInterface::IDebugDrawMesh
 {
-	FDynamicMeshDebugDrawMesh(const class UE::Geometry::FDynamicMesh3* DynamicMesh);
+	DATAFLOWCORE_API FDynamicMeshDebugDrawMesh(const class UE::Geometry::FDynamicMesh3* DynamicMesh);
 
 	DATAFLOWCORE_API virtual ~FDynamicMeshDebugDrawMesh() = default;
 
 	DATAFLOWCORE_API virtual int32 GetMaxVertexIndex() const override;
 	DATAFLOWCORE_API virtual bool IsValidVertex(int32 VertexIndex) const override;
 	DATAFLOWCORE_API virtual FVector GetVertexPosition(int32 VertexIndex) const override;
+	DATAFLOWCORE_API virtual FVector GetVertexNormal(int32 VertexIndex) const override;
 
 	DATAFLOWCORE_API virtual int32 GetMaxTriangleIndex() const override;
 	DATAFLOWCORE_API virtual bool IsValidTriangle(int32 VertexIndex) const override;
