ChaosFlesh:
-changed ComputeMuscleActivation warning
-added more log info on origin/insertion overlap
-fixed editCondition error in FGeometryCollectionTransferVertexAttributeNode

diff --git a/Engine/Plugins/Experimental/ChaosFlesh/Source/ChaosFleshNodes/Private/Dataflow/ChaosFleshComputeMuscleActivationNode.cpp b/Engine/Plugins/Experimental/ChaosFlesh/Source/ChaosFleshNodes/Private/Dataflow/ChaosFleshComputeMuscleActivationNode.cpp
--- a/Engine/Plugins/Experimental/ChaosFlesh/Source/ChaosFleshNodes/Private/Dataflow/ChaosFleshComputeMuscleActivationNode.cpp
+++ b/Engine/Plugins/Experimental/ChaosFlesh/Source/ChaosFleshNodes/Private/Dataflow/ChaosFleshComputeMuscleActivationNode.cpp
@@ -36,7 +36,7 @@ void FComputeMuscleActivationDataNode_v2::Evaluate(UE::Dataflow::FContext& Conte
 		GeometryCollection::Facades::FMuscleActivationFacade FMuscleActivation(InCollection);
 		if (!FMuscleActivation.SetUpMuscleActivation(InOriginIndices, InInsertionIndices))
 		{
-			Context.Error(TEXT("Setup failed, please check the Log for more info."), this, Out);
+			Context.Warning(TEXT("Setup failed, please check the Log for more info."), this, Out);
 		}
 		Out->SetValue(MoveTemp(InCollection), Context);
 	}
diff --git a/Engine/Plugins/Experimental/GeometryCollectionPlugin/Source/GeometryCollectionNodes/Public/Dataflow/GeometryCollectionTransferVertexAttributeNode.h b/Engine/Plugins/Experimental/GeometryCollectionPlugin/Source/GeometryCollectionNodes/Public/Dataflow/GeometryCollectionTransferVertexAttributeNode.h
--- a/Engine/Plugins/Experimental/GeometryCollectionPlugin/Source/GeometryCollectionNodes/Public/Dataflow/GeometryCollectionTransferVertexAttributeNode.h
+++ b/Engine/Plugins/Experimental/GeometryCollectionPlugin/Source/GeometryCollectionNodes/Public/Dataflow/GeometryCollectionTransferVertexAttributeNode.h
@@ -65,6 +65,22 @@ enum class EDataflowTransferVertexAttributeNodeBoundingVolume : uint8
 	Dataflow_Max UMETA(Hidden)
 };
 
+UENUM(BlueprintType)
+enum class EDataflowTransferVertexAttributeNodeTransferMethod : uint8
+{
+	/** Transfers vertex attribute between matched geometries.*/
+	Component  UMETA(DisplayName = "Component Transfer"),
+
+	/** Transfers vertex attribute globally. Warning: attributes on one geometry might bleed over to neighbor geometries.*/
+	Global UMETA(DisplayName = "Global Transfer"),
+
+	/** No distance falloff */
+	None UMETA(DisplayName = "None"),
+	//~~~
+	//256th entry
+	Dataflow_Max UMETA(Hidden)
+};
+
 /**
  * Transfer float properties from a source collection to a target collection.
  */
@@ -76,6 +92,15 @@ struct FGeometryCollectionTransferVertexAttributeNode : public FDataflowNode
 	DATAFLOW_NODE_RENDER_TYPE("SurfaceWeightsRender", FGeometryCollection::StaticType(), "Collection", "AttributeKey")
 
 public:
+	/* Transfer method used during the most recent node evaluation. 
+	* Component Transfer is used when all geometries from the source collection have matched names with the target collection.
+	* Otherwise, Global Transfer is used.
+	* Geometries are matched when the geometry's BoneName can be found as the start of the BoneName of a geometry in the target collection.
+	* Use TransformNameSuffix to add extra string to the source geometry's BoneName to avoid multiple matched names.
+	* For example, source geometry has name SK_10 and target geometry has name SK_10_tet1
+	* For all names, Check BoneName attribute in Transform group in the collection.*/
+	UPROPERTY(VisibleAnywhere, Category = "Transfer Method")
+	EDataflowTransferVertexAttributeNodeTransferMethod TransferMethod = EDataflowTransferVertexAttributeNodeTransferMethod::None;
 
 	/* Target collection to transfer vertex attribute to. */
 	UPROPERTY(Meta = (DataflowInput, DataflowOutput, DataflowPassthrough = "Collection"))
@@ -103,7 +128,7 @@ public:
 
 	/* Threshold based on distance from source triangle.Values past the threshold will falloff.[Defaults to 1 percent of triangle size(0.01)] */
 	UPROPERTY(EditAnywhere, Category = "Thresholds", meta = (EditCondition = 
-		"BoundingVolumeType == EDataflowTransferVertexAttributeNodeBoundingVolume::Triangle && EDataflowTransferVertexAttributeNodeFalloff Falloff = EDataflowTransferVertexAttributeNodeFalloff::None"
+		"BoundingVolumeType == EDataflowTransferVertexAttributeNodeBoundingVolume::Triangle && EDataflowTransferVertexAttributeNodeFalloff Falloff == EDataflowTransferVertexAttributeNodeFalloff::None"
 		, EditConditionHides))
 	float FalloffThreshold = 0.01f;
 
@@ -117,7 +142,7 @@ public:
 		EditConditionHides))
 	float BoundMultiplier = 0.01f;
 
-	/* Suffix of transform names for geometry matching during transfer[default: _Tet]. In CreateTetrahedron node we add _Tet to tetrahedral geometries.*/
+	/* Suffix of transform names added to the source geometry's BoneName for geometry matching during transfer[default: _Tet]. In CreateTetrahedron node we add _Tet to tetrahedral geometries.*/
 	UPROPERTY(EditAnywhere, Category = "Transform Name")
 	FString TransformNameSuffix = "_Tet";
 
@@ -148,6 +173,15 @@ struct FGeometryCollectionTransferVertexSkinWeightsNode : public FDataflowNode
 	DATAFLOW_NODE_RENDER_TYPE("SurfaceRender", FGeometryCollection::StaticType(), "Collection")
 
 public:
+	/* Transfer method used during the most recent node evaluation.
+	* Component Transfer is used when all geometries from the source collection have matched names with the target collection.
+	* Otherwise, Global Transfer is used.
+	* Geometries are matched when the geometry's BoneName can be found as the start of the BoneName of a geometry in the target collection.
+	* Use TransformNameSuffix to add extra string to the source geometry's BoneName to avoid multiple matched names.
+	* For example, source geometry has name SK_10 and target geometry has name SK_10_tet1
+	* For all names, Check BoneName attribute in Transform group in the collection.*/
+	UPROPERTY(VisibleAnywhere, Category = "Transfer Method")
+	EDataflowTransferVertexAttributeNodeTransferMethod TransferMethod = EDataflowTransferVertexAttributeNodeTransferMethod::None;
 
 	/* Target collection to transfer vertex attribute to. */
 	UPROPERTY(Meta = (DataflowInput, DataflowOutput, DataflowPassthrough = "Collection"))
diff --git a/Engine/Source/Runtime/Experimental/Chaos/Private/GeometryCollection/Facades/CollectionMuscleActivationFacade.cpp b/Engine/Source/Runtime/Experimental/Chaos/Private/GeometryCollection/Facades/CollectionMuscleActivationFacade.cpp
--- a/Engine/Source/Runtime/Experimental/Chaos/Private/GeometryCollection/Facades/CollectionMuscleActivationFacade.cpp
+++ b/Engine/Source/Runtime/Experimental/Chaos/Private/GeometryCollection/Facades/CollectionMuscleActivationFacade.cpp
@@ -228,41 +228,65 @@ namespace GeometryCollection::Facades
 		// Vertices and fiber field
 		if (InOrigin.Num() == 0)
 		{
-			UE_LOG(LogDataflowMuscleActivationFacade, Error, TEXT("No origins given"));
+			UE_LOG(LogDataflowMuscleActivationFacade, Error, TEXT("Muscle activation setup failed: No origins given"));
 			return false;
 		}
 		if (Insertion.Num() == 0)
 		{
-			UE_LOG(LogDataflowMuscleActivationFacade, Error, TEXT("No insertions given"));
+			UE_LOG(LogDataflowMuscleActivationFacade, Error, TEXT("Muscle activation setup failed: No insertions given"));
 			return false;
 		}
 		if (!ConstCollection.FindAttribute<FVector3f>("Vertex", FGeometryCollection::VerticesGroup))
 		{
-			UE_LOG(LogDataflowMuscleActivationFacade, Error, TEXT("Collection has no Vertex attribute"));
+			UE_LOG(LogDataflowMuscleActivationFacade, Error, TEXT("Muscle activation setup failed: Collection has no Vertex attribute"));
 			return false;
 		}
 		if (!ConstCollection.FindAttribute<FIntVector4>("Tetrahedron", "Tetrahedral"))
 		{
-			UE_LOG(LogDataflowMuscleActivationFacade, Error, TEXT("Collection has no Tetrahedron attribute"));
+			UE_LOG(LogDataflowMuscleActivationFacade, Error, TEXT("Muscle activation setup failed: Collection has no Tetrahedron attribute"));
 			return false;
 		}
 		if (!ConstCollection.FindAttribute<FVector3f>("FiberDirection", "Tetrahedral"))
 		{
-			UE_LOG(LogDataflowMuscleActivationFacade, Error, TEXT("Collection has no FiberDirection attribute"));
+			UE_LOG(LogDataflowMuscleActivationFacade, Error, TEXT("Muscle activation setup failed: Collection has no FiberDirection attribute"));
 			return false;
 		}
 		// Check if origins and insertions have overlaps
 		int32 OverlapCount = 0;
+		const TManagedArray<int32>* TransformIndex = ConstCollection.FindAttribute<int32>("TransformIndex", FGeometryCollection::GeometryGroup);
+		const TManagedArray<FString>* BoneName = ConstCollection.FindAttribute<FString>("BoneName", FGeometryCollection::TransformGroup);
+		const TManagedArray<int32>* VertexStart = ConstCollection.FindAttribute<int32>("VertexStart", FGeometryCollection::GeometryGroup);
+		const TManagedArray<int32>* VertexCount = ConstCollection.FindAttribute<int32>("VertexCount", FGeometryCollection::GeometryGroup);
+		TSet<int32> OverlapGeometries;
 		for (int32 OriginIdx : InOrigin)
 		{
 			if (Insertion.Contains(OriginIdx))
 			{
 				OverlapCount++;
+				if (VertexStart && VertexCount)
+				{
+					for (int32 GeomIdx = 0; GeomIdx < VertexStart->Num(); ++GeomIdx)
+					{
+						if ((*VertexStart)[GeomIdx] <= OriginIdx && OriginIdx < (*VertexStart)[GeomIdx] + (*VertexCount)[GeomIdx])
+						{
+							OverlapGeometries.Add(GeomIdx);
+							break;
+						}
+					}
+				}
 			}
 		}
 		if (OverlapCount)
 		{
-			UE_LOG(LogDataflowMuscleActivationFacade, Error, TEXT("Origins and Insertions have %d common indices out of total %d, please check if they are from different sources."), OverlapCount, InOrigin.Num());
+			UE_LOG(LogDataflowMuscleActivationFacade, Error, TEXT("Muscle activation setup failed: origins and insertions have %d common indices out of total %d, please check if they are from different sources."), OverlapCount, InOrigin.Num());
+			if (TransformIndex && BoneName)
+			{
+				for (const int32 GeomIdx : OverlapGeometries)
+				{
+					if (ensure(TransformIndex->IsValidIndex(GeomIdx) && BoneName->IsValidIndex((*TransformIndex)[GeomIdx])))
+					UE_LOG(LogDataflowMuscleActivationFacade, Error, TEXT("Overlapped origins and insertions are from geometry indexed %d, bone name %s."), GeomIdx, *(*BoneName)[(*TransformIndex)[GeomIdx]]);
+				}
+			}
 			return false;
 		}
 		TArray<int32> Origin = InOrigin;
