ChaosFlesh:
-Optimized ComputeFiberField node (10X perf): use muscle mesh instead of full tet mesh + parallel Laplacian operator
-Quick optimization on KinematicMuscleAttachment node
-Added transfer method option for TransferVertexAttribute node to force default paired geometry transfer instead of accidental global transfer. It is valid that the geometries from source and target do not match exactly, as an extremely small SKM geomtry will not produce any tetrahedral mesh geometry.

diff --git a/Engine/Plugins/Experimental/ChaosFlesh/Source/ChaosFleshNodes/Private/Dataflow/ChaosFleshComputeFiberFieldNode.cpp b/Engine/Plugins/Experimental/ChaosFlesh/Source/ChaosFleshNodes/Private/Dataflow/ChaosFleshComputeFiberFieldNode.cpp
--- a/Engine/Plugins/Experimental/ChaosFlesh/Source/ChaosFleshNodes/Private/Dataflow/ChaosFleshComputeFiberFieldNode.cpp
+++ b/Engine/Plugins/Experimental/ChaosFlesh/Source/ChaosFleshNodes/Private/Dataflow/ChaosFleshComputeFiberFieldNode.cpp
@@ -5,6 +5,7 @@
 #include "Chaos/Math/Poisson.h"
 #include "ChaosFlesh/ChaosFlesh.h"
 #include "ChaosFlesh/TetrahedralCollection.h"
+#include "GeometryCollection/Facades/CollectionMeshFacade.h"
 #include "GeometryCollection/Facades/CollectionMuscleActivationFacade.h"
 
 #include UE_INLINE_GENERATED_CPP_BY_NAME(ChaosFleshComputeFiberFieldNode)
@@ -104,6 +105,10 @@ void FComputeFiberFieldNode::Evaluate(UE::Dataflow::FContext& Context, const FDa
 					Out->SetValue(MoveTemp(InCollection), Context);
 					return;
 				}
+				else
+				{
+					InOriginIndices = Origin->GetConstArray();
+				}
 			}
 
 			// Insertion vertices
@@ -124,18 +129,80 @@ void FComputeFiberFieldNode::Evaluate(UE::Dataflow::FContext& Context, const FDa
 					Out->SetValue(MoveTemp(InCollection), Context);
 					return;
 				}
+				else
+				{
+					InInsertionIndices = Insertion->GetConstArray();
+				}
+			}
+		}
+
+		// Only solve for fiber field on muscle geometries
+		TSet<int32> MuscleGeometries;
+		GeometryCollection::Facades::FCollectionMeshFacade MeshFacade(InCollection);
+		TArray<int32> GeometryIndex = MeshFacade.GetGeometryGroupIndexArray();
+		for (int32 OriginIdx : InOriginIndices)
+		{
+			if (GeometryIndex.IsValidIndex(OriginIdx))
+			{
+				MuscleGeometries.Add(GeometryIndex[OriginIdx]);
+			}
+		}
+		for (int32 OriginIdx : InInsertionIndices)
+		{
+			if (GeometryIndex.IsValidIndex(OriginIdx))
+			{
+				MuscleGeometries.Add(GeometryIndex[OriginIdx]);
+			}
+		}
+		TArray<int32> MuscleElementIndices;
+		TArray<FIntVector4> MuscleElements;
+
+		TManagedArray<int32>* TetrahedronStart = InCollection.FindAttribute<int32>(
+			FTetrahedralCollection::TetrahedronStartAttribute, FTetrahedralCollection::GeometryGroup);
+		TManagedArray<int32>* TetrahedronCount = InCollection.FindAttribute<int32>(
+			FTetrahedralCollection::TetrahedronCountAttribute, FTetrahedralCollection::GeometryGroup);
+		TArray<TArray<int32>> MuscleConstraints;
+		if (TetrahedronStart && TetrahedronCount)
+		{
+			for (int32 GeometryIdx : MuscleGeometries)
+			{
+				if (TetrahedronStart->IsValidIndex(GeometryIdx))
+				{
+					for (int32 ElemIdx = (*TetrahedronStart)[GeometryIdx]; ElemIdx < (*TetrahedronStart)[GeometryIdx] + (*TetrahedronCount)[GeometryIdx]; ++ElemIdx)
+					{
+						MuscleElementIndices.Add(ElemIdx);
+						MuscleElements.Add((*Elements)[ElemIdx]);
+					}
+				}
 			}
 		}
+		MuscleConstraints.SetNum(MuscleElements.Num());
+		for (int32 LocalIdx = 0; LocalIdx < MuscleElements.Num(); ++LocalIdx)
+		{
+			for (int32 c = 0; c < 4; ++c)
+			{
+				MuscleConstraints[LocalIdx].Add(MuscleElements[LocalIdx][c]);
+			}
+		}
+		TArray<TArray<int32>> MuscleIncidentElementsLocalIndex;
+		TArray<TArray<int32>> MuscleIncidentElements = Chaos::Utilities::ComputeIncidentElements(MuscleConstraints, &MuscleIncidentElementsLocalIndex);
 
 		//
 		// Do the thing
 		//
-
-		TArray<FVector3f> FiberDirs;
+		TArray<FVector3f> MuscleFiberDirs;
 		TArray<float> MuscleAttachmentScalarFieldTArray; //continuous field where origin = 1, insertion = 2, othernodes = 0
-		ComputeFiberField(*Elements, *Vertex, *IncidentElements, *IncidentElementsLocalIndex, 
-				Origin ? Origin->GetConstArray() : InOriginIndices,
-				Insertion ? Insertion->GetConstArray() : InInsertionIndices, FiberDirs, MuscleAttachmentScalarFieldTArray);
+		Chaos::ComputeFiberField<float>(
+			MuscleElements,
+			Vertex->GetConstArray(),
+			MuscleIncidentElements,
+			MuscleIncidentElementsLocalIndex,
+			InOriginIndices,
+			InInsertionIndices,
+			MuscleFiberDirs,
+			MuscleAttachmentScalarFieldTArray,
+			MaxIterations,
+			Tolerance);
 
 		//
 		// Set output(s)
@@ -143,7 +210,11 @@ void FComputeFiberFieldNode::Evaluate(UE::Dataflow::FContext& Context, const FDa
 
 		TManagedArray<FVector3f>& FiberDirections =
 			InCollection.AddAttribute<FVector3f>("FiberDirection", FTetrahedralCollection::TetrahedralGroup);
-		FiberDirections = MoveTemp(FiberDirs);
+		FiberDirections.Fill(FVector3f(0, 0, 0));
+		for (int32 LocalIdx = 0; LocalIdx < MuscleElements.Num(); ++LocalIdx)
+		{
+			FiberDirections[MuscleElementIndices[LocalIdx]] = MuscleFiberDirs[LocalIdx];
+		}
 
 		if (bShowMuscleColor)
 		{
@@ -152,7 +223,7 @@ void FComputeFiberFieldNode::Evaluate(UE::Dataflow::FContext& Context, const FDa
 			for (int32 i = 0; i < Color.Num(); ++i)
 			{
 				float Value = MuscleAttachmentScalarFieldTArray[i];
-				if (Value >= 1) // 1 <= Value <= 2 if muscle with origin and insertion
+				if (1 <= Value && Value <= 2) // 1 <= Value <= 2 if muscle with origin and insertion
 				{
 					Color[i] = FLinearColor(FVector(Value - 1, 0, 2 - Value));
 				}
@@ -177,29 +248,6 @@ FComputeFiberFieldNode::GetNonZeroIndices(const TArray<uint8>& Map) const
 	return Indices;
 }
 
-void FComputeFiberFieldNode::ComputeFiberField(
-	const TManagedArray<FIntVector4>& Elements,
-	const TManagedArray<FVector3f>& Vertex,
-	const TManagedArray<TArray<int32>>& IncidentElements,
-	const TManagedArray<TArray<int32>>& IncidentElementsLocalIndex,
-	const TArray<int32>& Origin,
-	const TArray<int32>& Insertion,
-	TArray<FVector3f>& Directions,
-	TArray<float>& ScalarField) const
-{
-	Chaos::ComputeFiberField<float>(
-		Elements.GetConstArray(),
-		Vertex.GetConstArray(),
-		IncidentElements.GetConstArray(),
-		IncidentElementsLocalIndex.GetConstArray(),
-		Origin,
-		Insertion,
-		Directions,
-		ScalarField,
-		MaxIterations,
-		Tolerance);
-}
-
 void FComputeFiberStreamlineNode::Evaluate(UE::Dataflow::FContext& Context, const FDataflowOutput* Out) const
 {
 	//
diff --git a/Engine/Plugins/Experimental/ChaosFlesh/Source/ChaosFleshNodes/Private/Dataflow/ChaosFleshKinematicMuscleAttachmentsNode.cpp b/Engine/Plugins/Experimental/ChaosFlesh/Source/ChaosFleshNodes/Private/Dataflow/ChaosFleshKinematicMuscleAttachmentsNode.cpp
--- a/Engine/Plugins/Experimental/ChaosFlesh/Source/ChaosFleshNodes/Private/Dataflow/ChaosFleshKinematicMuscleAttachmentsNode.cpp
+++ b/Engine/Plugins/Experimental/ChaosFlesh/Source/ChaosFleshNodes/Private/Dataflow/ChaosFleshKinematicMuscleAttachmentsNode.cpp
@@ -23,20 +23,19 @@ void FKinematicMuscleAttachmentsDataflowNode::Evaluate(UE::Dataflow::FContext& C
 			if (FindInput(&OriginVertexIndicesIn) && FindInput(&OriginVertexIndicesIn)->GetConnection() && FindInput(&InsertionVertexIndicesIn) && FindInput(&InsertionVertexIndicesIn)->GetConnection())
 			{
 				TArray<int32> BoundVerts;
-				TArray<float> BoundWeights;
 
 				for (int32 SelectionIndex : GetValue<TArray<int32>>(Context, &OriginVertexIndicesIn))
 				{
 					if (0 <= SelectionIndex && SelectionIndex < Vertices->Num())
 					{
-						BoundVerts.Add(SelectionIndex);
+						BoundVerts.AddUnique(SelectionIndex);
 					}
 				}
 				for (int32 SelectionIndex : GetValue<TArray<int32>>(Context, &InsertionVertexIndicesIn))
 				{
 					if (0 <= SelectionIndex && SelectionIndex < Vertices->Num())
 					{
-						BoundVerts.Add(SelectionIndex);
+						BoundVerts.AddUnique(SelectionIndex);
 					}
 				}
 				if (TObjectPtr<USkeletalMesh> SkeletalMesh = GetValue<TObjectPtr<USkeletalMesh>>(Context, &SkeletalMeshIn))
@@ -64,9 +63,10 @@ void FKinematicMuscleAttachmentsDataflowNode::Evaluate(UE::Dataflow::FContext& C
 						GeometryCollection::Facades::FCollectionTransformFacade TransformFacade(InCollection);
 						const TMap<FString, int32> BoneNameToIndex = TransformFacade.BoneNameIndexMap();
 						auto DoubleVert = [](FVector3f V) { return FVector3d(V.X, V.Y, V.Z); };
+						//ToDo(Chaosflesh): change this to particle bounding box query
 						for (int32 i = 0; i < BoundVerts.Num(); i++)
 						{
-							int ClosestPointIndex = 0;
+							int32 ClosestPointIndex = 0;
 							double MinDistance = DBL_MAX;
 							for (uint32 j = 0; j < PositionVertexBuffer.GetNumVertices(); j++)
 							{
@@ -128,6 +128,7 @@ void FKinematicMuscleAttachmentsDataflowNode::Evaluate(UE::Dataflow::FContext& C
 				{
 					if (BoundVerts.Num())
 					{
+						TArray<float> BoundWeights;
 						BoundWeights.Init(1.0, BoundVerts.Num());
 						GeometryCollection::Facades::FKinematicBindingFacade Kinematics(InCollection);
 						Kinematics.AddKinematicBinding(Kinematics.SetBoneBindings(INDEX_NONE, BoundVerts, BoundWeights));
diff --git a/Engine/Plugins/Experimental/ChaosFlesh/Source/ChaosFleshNodes/Private/Dataflow/ChaosFleshVisualizeFiberFieldNode.cpp b/Engine/Plugins/Experimental/ChaosFlesh/Source/ChaosFleshNodes/Private/Dataflow/ChaosFleshVisualizeFiberFieldNode.cpp
--- a/Engine/Plugins/Experimental/ChaosFlesh/Source/ChaosFleshNodes/Private/Dataflow/ChaosFleshVisualizeFiberFieldNode.cpp
+++ b/Engine/Plugins/Experimental/ChaosFlesh/Source/ChaosFleshNodes/Private/Dataflow/ChaosFleshVisualizeFiberFieldNode.cpp
@@ -26,7 +26,7 @@ void FVisualizeFiberFieldNode::Evaluate(UE::Dataflow::FContext& Context, const F
 						ensureMsgf(Elements->Num() == FiberDirections->Num(), TEXT("Fiber direction has different size than elements"));
 						for (int32 ElemIndex = 0; ElemIndex < Elements->Num(); ElemIndex++)
 						{
-							if ((*FiberDirections)[ElemIndex].Length() > 1e-10)
+							if ((*FiberDirections)[ElemIndex].Length() > UE_KINDA_SMALL_NUMBER) // This is supposed to be unit length
 							{
 								FVector3f VectorStart = { 0,0,0 };
 								FLinearColor VectorColor(EForceInit::ForceInitToZero);
diff --git a/Engine/Plugins/Experimental/GeometryCollectionPlugin/Source/GeometryCollectionNodes/Private/Dataflow/GeometryCollectionTransferVertexAttributeNode.cpp b/Engine/Plugins/Experimental/GeometryCollectionPlugin/Source/GeometryCollectionNodes/Private/Dataflow/GeometryCollectionTransferVertexAttributeNode.cpp
--- a/Engine/Plugins/Experimental/GeometryCollectionPlugin/Source/GeometryCollectionNodes/Private/Dataflow/GeometryCollectionTransferVertexAttributeNode.cpp
+++ b/Engine/Plugins/Experimental/GeometryCollectionPlugin/Source/GeometryCollectionNodes/Private/Dataflow/GeometryCollectionTransferVertexAttributeNode.cpp
@@ -223,9 +223,9 @@ namespace UE::Private
 		{
 			for (int32 i = 0; i < SourceName->Num(); ++i)
 			{
+				const FString TestName = FString::Printf(TEXT("%s%s"), *(*SourceName)[i], *TransformNameSuffix);
 				for (int32 j = 0; j < TargetName->Num(); ++j)
 				{
-					const FString TestName = FString::Printf(TEXT("%s%s"), *(*SourceName)[i], *TransformNameSuffix);
 					if ((*TargetName)[j].StartsWith(TestName))
 					{
 						Mapping.Add(FIntVector2((*SourceGeometryGroup)[i], (*TargetGeometryGroup)[j]));
@@ -521,27 +521,45 @@ namespace UE::Private
 		return MoveTemp(TransferDataArray);
 	}
 
-	static TArray<FTransferData> ComputeTransferData(const FManagedArrayCollection& SourceCollection, const FManagedArrayCollection& TargetCollection, const FTransferProperties& TransferProperties)
+	//store barycentric weight info TransferDataArray and returns transfer method used
+	static TArray<FTransferData> ComputeTransferData(const FManagedArrayCollection& SourceCollection, const FManagedArrayCollection& TargetCollection, const FTransferProperties& TransferProperties, const EDataflowTransferVertexAttributeNodeTransferMethod& TransferMethod)
 	{
-		TArray<FTransferData> TransferDataArray; //store barycentric weight info
+		TArray<FTransferData> TransferDataArray;
 		const FTransferFacade Target(TargetCollection);
 		const FTransferFacade Source(SourceCollection);
 		if (Target.IsValid() && Source.IsValid())
 		{
-			TArray<FIntVector2> AlignedGeometry = FindSourceToTargetGeometryMap(SourceCollection, TargetCollection, TransferProperties.TransformNameSuffix);
-			if (AlignedGeometry.Num() == SourceCollection.NumElements(FGeometryCollection::GeometryGroup))
+			if (TransferMethod == EDataflowTransferVertexAttributeNodeTransferMethod::Component)
 			{
+				TArray<FIntVector2> AlignedGeometry = FindSourceToTargetGeometryMap(SourceCollection, TargetCollection, TransferProperties.TransformNameSuffix);
+				if (AlignedGeometry.Num() != SourceCollection.NumElements(FGeometryCollection::GeometryGroup))
+				{
+					TSet<int32> MatchedSourceGeometries;
+					for (const FIntVector2& Pair : AlignedGeometry)
+					{
+						MatchedSourceGeometries.Add(Pair.X);
+					}
+					for (int32 SourceIdx = 0; SourceIdx < SourceCollection.NumElements(FGeometryCollection::GeometryGroup); ++SourceIdx)
+					{
+						if (!MatchedSourceGeometries.Contains(SourceIdx))
+						{
+							UE_LOG(LogChaosDataflow, Warning, TEXT("Source geometry indexed %d is not matched to any target geometry. Check if this is intentional."), SourceIdx);
+						}
+					}
+				}
 				TransferDataArray = PairedGeometryTransfer(AlignedGeometry, Source, Target, TransferProperties);
 			}
-			else
+			else if (TransferMethod == EDataflowTransferVertexAttributeNodeTransferMethod::Global)
 			{
 				TransferDataArray = GlobalAssetTransfer(Source, Target, TransferProperties);
 			}
 		}
-		return MoveTemp(TransferDataArray);
+		return TransferDataArray;
 	}
 }
 
+using namespace UE::Private;
+
 void FGeometryCollectionTransferVertexAttributeNode::Evaluate(UE::Dataflow::FContext& Context, const FDataflowOutput* Out) const
 {
 	FCollectionAttributeKey Key = GetValue(Context, &AttributeKey, AttributeKey);
@@ -550,10 +568,10 @@ void FGeometryCollectionTransferVertexAttributeNode::Evaluate(UE::Dataflow::FCon
 	{
 		FManagedArrayCollection TargetCollection = GetValue(Context, &Collection);
 		const FManagedArrayCollection& SourceCollection = GetValue(Context, &FromCollection);
-		UE::Private::FTransferFacade Target(TargetCollection);
-		const UE::Private::FTransferFacade Source(SourceCollection);
+		FTransferFacade Target(TargetCollection);
+		const FTransferFacade Source(SourceCollection);
 
-		const UE::Private::FTransferProperties TransferProperties = {
+		const FTransferProperties TransferProperties = {
 			BoundingVolumeType,
 			SourceScale,
 			Falloff,
@@ -565,14 +583,22 @@ void FGeometryCollectionTransferVertexAttributeNode::Evaluate(UE::Dataflow::FCon
 
 		if (Target.IsValid() && Source.IsValid())
 		{
-			const TArray<UE::Private::FTransferData> TransferDataArray = ComputeTransferData(SourceCollection, TargetCollection, TransferProperties);
+			TArray<FTransferData> TransferDataArray = ComputeTransferData(SourceCollection, TargetCollection, TransferProperties, TransferMethod);
+			if (TransferMethod == EDataflowTransferVertexAttributeNodeTransferMethod::Global)
+			{
+				Context.Warning(TEXT("Transfer method: global transfer"), this, Out);
+			}
+			else if (TransferMethod == EDataflowTransferVertexAttributeNodeTransferMethod::None)
+			{
+				Context.Warning(TEXT("Transfer failed"), this, Out);
+			}
 			if (const TManagedArray<float>* SourceAttributeFloatArray = Source.GetAttributeArray<float>(Key.Attribute, Key.Group))
 			{
 				TManagedArray<float>* TargetAttributeArray = Target.GetAttributeArray<float>(Key.Attribute, Key.Group);
 				TargetAttributeArray->Fill(0.f);
 				for (int32 VertexIndex = 0; VertexIndex < TargetAttributeArray->Num(); VertexIndex++)
 				{
-					const UE::Private::FTransferData& TransferData = TransferDataArray[VertexIndex];
+					const FTransferData& TransferData = TransferDataArray[VertexIndex];
 					if (TransferData.SourceIndices.Num())
 					{
 						for (int32 SourceLocalIdx = 0; SourceLocalIdx < TransferData.SourceIndices.Num(); ++SourceLocalIdx)
@@ -589,7 +615,7 @@ void FGeometryCollectionTransferVertexAttributeNode::Evaluate(UE::Dataflow::FCon
 				TargetAttributeArray->Fill(FLinearColor(0, 0, 0, 0));
 				for (int32 VertexIndex = 0; VertexIndex < TargetAttributeArray->Num(); VertexIndex++)
 				{
-					const UE::Private::FTransferData& TransferData = TransferDataArray[VertexIndex];
+					const FTransferData& TransferData = TransferDataArray[VertexIndex];
 					if (TransferData.SourceIndices.Num())
 					{
 						for (int32 SourceLocalIdx = 0; SourceLocalIdx < TransferData.SourceIndices.Num(); ++SourceLocalIdx)
@@ -620,10 +646,10 @@ void FGeometryCollectionTransferVertexSkinWeightsNode::Evaluate(UE::Dataflow::FC
 	{
 		FManagedArrayCollection TargetCollection = GetValue(Context, &Collection);
 		const FManagedArrayCollection& SourceCollection = GetValue(Context, &FromCollection);
-		UE::Private::FTransferFacade Target(TargetCollection);
-		const UE::Private::FTransferFacade Source(SourceCollection);
+		FTransferFacade Target(TargetCollection);
+		const FTransferFacade Source(SourceCollection);
 
-		const UE::Private::FTransferProperties TransferProperties = {
+		const FTransferProperties TransferProperties = {
 			BoundingVolumeType,
 			SourceScale,
 			Falloff,
@@ -635,8 +661,7 @@ void FGeometryCollectionTransferVertexSkinWeightsNode::Evaluate(UE::Dataflow::FC
 
 		if (Target.IsValid() && Source.IsValid())
 		{
-			const TArray<UE::Private::FTransferData> TransferDataArray = ComputeTransferData(SourceCollection, TargetCollection, TransferProperties);
-
+			TArray<FTransferData> TransferDataArray = ComputeTransferData(SourceCollection, TargetCollection, TransferProperties, TransferMethod);
 			GeometryCollection::Facades::FVertexBoneWeightsFacade SourceVertexBoneWeightsFacade(SourceCollection);
 			GeometryCollection::Facades::FVertexBoneWeightsFacade TargetVertexBoneWeightsFacade(TargetCollection);
 			const TManagedArray<TArray<int32>>* SourceBoneIndices = SourceVertexBoneWeightsFacade.FindBoneIndices();
@@ -646,10 +671,10 @@ void FGeometryCollectionTransferVertexSkinWeightsNode::Evaluate(UE::Dataflow::FC
 				//
 				// Compute the bone index mappings. This allows the transfer operator to retarget weights to the correct skeleton.
 				//
-				const TMap<int32, int32> SourceBoneToTargetBone = UE::Private::FindSourceToTargetTransformMap(SourceCollection, TargetCollection);
+				const TMap<int32, int32> SourceBoneToTargetBone = FindSourceToTargetTransformMap(SourceCollection, TargetCollection);
 				for (int32 VertexIndex = 0; VertexIndex < Target.Vertex.Num(); VertexIndex++)
 				{
-					const UE::Private::FTransferData& TransferData = TransferDataArray[VertexIndex];
+					const FTransferData& TransferData = TransferDataArray[VertexIndex];
 					if (TransferData.SourceIndices.Num())
 					{
 						TMap<int32, float> BoneWeightBucket;
diff --git a/Engine/Plugins/Experimental/GeometryCollectionPlugin/Source/GeometryCollectionNodes/Public/Dataflow/GeometryCollectionTransferVertexAttributeNode.h b/Engine/Plugins/Experimental/GeometryCollectionPlugin/Source/GeometryCollectionNodes/Public/Dataflow/GeometryCollectionTransferVertexAttributeNode.h
--- a/Engine/Plugins/Experimental/GeometryCollectionPlugin/Source/GeometryCollectionNodes/Public/Dataflow/GeometryCollectionTransferVertexAttributeNode.h
+++ b/Engine/Plugins/Experimental/GeometryCollectionPlugin/Source/GeometryCollectionNodes/Public/Dataflow/GeometryCollectionTransferVertexAttributeNode.h
@@ -69,7 +69,7 @@ UENUM(BlueprintType)
 enum class EDataflowTransferVertexAttributeNodeTransferMethod : uint8
 {
 	/** Transfers vertex attribute between matched geometries.*/
-	Component  UMETA(DisplayName = "Component Transfer"),
+	Component  UMETA(DisplayName = "Paired Geometry Transfer"),
 
 	/** Transfers vertex attribute globally. Warning: attributes on one geometry might bleed over to neighbor geometries.*/
 	Global UMETA(DisplayName = "Global Transfer"),
@@ -83,6 +83,12 @@ enum class EDataflowTransferVertexAttributeNodeTransferMethod : uint8
 
 /**
  * Transfer float properties from a source collection to a target collection.
+ * Component Transfer is used when all geometries from the source collection have matched names with the target collection.
+ * Otherwise, Global Transfer is used.
+ * Geometries are matched when the geometry's BoneName can be found as the start of the BoneName of a geometry in the target collection.
+ * Use TransformNameSuffix to add extra string to the source geometry's BoneName to avoid multiple matched names.
+ * For example, source geometry has name SK_10 and target geometry has name SK_10_tet1
+ * For all names, Check BoneName attribute in Transform group in the collection.
  */
 USTRUCT(meta = (DataflowGeometryCollection))
 struct FGeometryCollectionTransferVertexAttributeNode : public FDataflowNode
@@ -92,15 +98,6 @@ struct FGeometryCollectionTransferVertexAttributeNode : public FDataflowNode
 	DATAFLOW_NODE_RENDER_TYPE("SurfaceWeightsRender", FGeometryCollection::StaticType(), "Collection", "AttributeKey")
 
 public:
-	/* Transfer method used during the most recent node evaluation. 
-	* Component Transfer is used when all geometries from the source collection have matched names with the target collection.
-	* Otherwise, Global Transfer is used.
-	* Geometries are matched when the geometry's BoneName can be found as the start of the BoneName of a geometry in the target collection.
-	* Use TransformNameSuffix to add extra string to the source geometry's BoneName to avoid multiple matched names.
-	* For example, source geometry has name SK_10 and target geometry has name SK_10_tet1
-	* For all names, Check BoneName attribute in Transform group in the collection.*/
-	UPROPERTY(VisibleAnywhere, Category = "Transfer Method")
-	EDataflowTransferVertexAttributeNodeTransferMethod TransferMethod = EDataflowTransferVertexAttributeNodeTransferMethod::None;
 
 	/* Target collection to transfer vertex attribute to. */
 	UPROPERTY(Meta = (DataflowInput, DataflowOutput, DataflowPassthrough = "Collection"))
@@ -114,6 +111,10 @@ public:
 	UPROPERTY(EditAnywhere, Category = "Dataflow", Meta = (DataflowInput, DataflowOutput, DisplayName = "AttributeKey", DataflowPassthrough = "AttributeKey"))
 	FCollectionAttributeKey AttributeKey = FCollectionAttributeKey(FString(""), FString("Vertices"));
 	
+	/* Transfer method [default: Paired Geometry Transfer] */
+	UPROPERTY(EditAnywhere, Category = "Method")
+	EDataflowTransferVertexAttributeNodeTransferMethod TransferMethod = EDataflowTransferVertexAttributeNodeTransferMethod::Component;
+
 	/* Bounding volume type for source assets[default: Triangle] */
 	UPROPERTY(EditAnywhere, Category = "Thresholds")
 	EDataflowTransferVertexAttributeNodeBoundingVolume BoundingVolumeType = EDataflowTransferVertexAttributeNodeBoundingVolume::Triangle;
@@ -164,6 +165,12 @@ private:
 
 /**
  * Transfer skin weights from a source collection to a target collection.
+ * Component Transfer is used when all geometries from the source collection have matched names with the target collection.
+ * Otherwise, Global Transfer is used.
+ * Geometries are matched when the geometry's BoneName can be found as the start of the BoneName of a geometry in the target collection.
+ * Use TransformNameSuffix to add extra string to the source geometry's BoneName to avoid multiple matched names.
+ * For example, source geometry has name SK_10 and target geometry has name SK_10_tet1
+ * For all names, Check BoneName attribute in Transform group in the collection.
  */
 USTRUCT(meta = (DataflowGeometryCollection))
 struct FGeometryCollectionTransferVertexSkinWeightsNode : public FDataflowNode
@@ -173,15 +180,6 @@ struct FGeometryCollectionTransferVertexSkinWeightsNode : public FDataflowNode
 	DATAFLOW_NODE_RENDER_TYPE("SurfaceRender", FGeometryCollection::StaticType(), "Collection")
 
 public:
-	/* Transfer method used during the most recent node evaluation.
-	* Component Transfer is used when all geometries from the source collection have matched names with the target collection.
-	* Otherwise, Global Transfer is used.
-	* Geometries are matched when the geometry's BoneName can be found as the start of the BoneName of a geometry in the target collection.
-	* Use TransformNameSuffix to add extra string to the source geometry's BoneName to avoid multiple matched names.
-	* For example, source geometry has name SK_10 and target geometry has name SK_10_tet1
-	* For all names, Check BoneName attribute in Transform group in the collection.*/
-	UPROPERTY(VisibleAnywhere, Category = "Transfer Method")
-	EDataflowTransferVertexAttributeNodeTransferMethod TransferMethod = EDataflowTransferVertexAttributeNodeTransferMethod::None;
 
 	/* Target collection to transfer vertex attribute to. */
 	UPROPERTY(Meta = (DataflowInput, DataflowOutput, DataflowPassthrough = "Collection"))
@@ -191,6 +189,10 @@ public:
 	UPROPERTY(Meta = (DataflowInput, DisplayName = "FromCollection"))
 	FManagedArrayCollection FromCollection;
 
+	/* Transfer method [default: Paired Geometry Transfer] */
+	UPROPERTY(EditAnywhere, Category = "Method")
+	EDataflowTransferVertexAttributeNodeTransferMethod TransferMethod = EDataflowTransferVertexAttributeNodeTransferMethod::Component;
+
 	/* Bounding volume type for source assets[default: Triangle] */
 	UPROPERTY(EditAnywhere, Category = "Thresholds")
 	EDataflowTransferVertexAttributeNodeBoundingVolume BoundingVolumeType = EDataflowTransferVertexAttributeNodeBoundingVolume::Triangle;
diff --git a/Engine/Source/Runtime/Experimental/Chaos/Private/GeometryCollection/Facades/CollectionMuscleActivationFacade.cpp b/Engine/Source/Runtime/Experimental/Chaos/Private/GeometryCollection/Facades/CollectionMuscleActivationFacade.cpp
--- a/Engine/Source/Runtime/Experimental/Chaos/Private/GeometryCollection/Facades/CollectionMuscleActivationFacade.cpp
+++ b/Engine/Source/Runtime/Experimental/Chaos/Private/GeometryCollection/Facades/CollectionMuscleActivationFacade.cpp
@@ -284,7 +284,9 @@ namespace GeometryCollection::Facades
 				for (const int32 GeomIdx : OverlapGeometries)
 				{
 					if (ensure(TransformIndex->IsValidIndex(GeomIdx) && BoneName->IsValidIndex((*TransformIndex)[GeomIdx])))
-					UE_LOG(LogDataflowMuscleActivationFacade, Error, TEXT("Overlapped origins and insertions are from geometry indexed %d, bone name %s."), GeomIdx, *(*BoneName)[(*TransformIndex)[GeomIdx]]);
+					{
+						UE_LOG(LogDataflowMuscleActivationFacade, Error, TEXT("Overlapped origins and insertions are from transform indexed %d, bone name %s."), (*TransformIndex)[GeomIdx], *(*BoneName)[(*TransformIndex)[GeomIdx]]);
+					}
 				}
 			}
 			return false;
@@ -316,17 +318,22 @@ namespace GeometryCollection::Facades
 			}
 		}
 		ComponentInsertions.SetNum(ComponentOrigins.Num());
+		TSet<int32> InsertionOnlyComponents;
 		for (int32 i = 0; i < Insertion.Num(); i++)
 		{
 			if (!ComponentToMuscleIndex.Contains(ComponentIndex[Insertion[i]]))
 			{
-				UE_LOG(LogDataflowMuscleActivationFacade, Warning, TEXT("Some muscle has only insertions but no origins."));
+				InsertionOnlyComponents.Add(ComponentIndex[Insertion[i]]);
 			}
 			else
 			{
 				ComponentInsertions[ComponentToMuscleIndex[ComponentIndex[Insertion[i]]]].Add(Insertion[i]);
 			}
 		}
+		for (int32 ComponentIdx: InsertionOnlyComponents)
+		{
+			UE_LOG(LogDataflowMuscleActivationFacade, Warning, TEXT("Geometry %d has only insertions but no origins."), ComponentIdx);
+		}
 		MuscleActivationElements.SetNum(ComponentOrigins.Num());
 		for (int32 ElemIdx = 0; ElemIdx < Elements.Num(); ElemIdx++)
 		{
diff --git a/Engine/Source/Runtime/Experimental/Chaos/Public/Chaos/Math/Krylov.h b/Engine/Source/Runtime/Experimental/Chaos/Public/Chaos/Math/Krylov.h
--- a/Engine/Source/Runtime/Experimental/Chaos/Public/Chaos/Math/Krylov.h
+++ b/Engine/Source/Runtime/Experimental/Chaos/Public/Chaos/Math/Krylov.h
@@ -20,6 +20,7 @@ void LanczosCG(
 	bool check_residual = false,
 	int min_parallel_batch_size = 1000) 
 {
+	//ToDo(ChaosFlesh): take an optional use_list of vertex indices to bypass unused verts
 	auto dotProduct = [](const TArray<T>& x, const TArray<T>& y) 
 	{
 		T result = T(0);
@@ -220,9 +221,9 @@ void LanczosCG(
 	TArray<TV> c; c.SetNum(x.Num());
 
 	T beta = FGenericPlatformMath::Sqrt(T(dotProduct(b, b)));
-	if (beta < res) 
+	if (beta <= res) 
 	{
-		UE_LOG(LogChaos, VeryVerbose, TEXT("Lanczos residual = %f. Lanczos converged in 0 iterations."), beta)
+		UE_LOG(LogChaos, Warning, TEXT("Lanczos residual = %f. Lanczos converged in 0 iterations."), beta)
 		scale(x, T(0));
 		return;
 	}
diff --git a/Engine/Source/Runtime/Experimental/Chaos/Public/Chaos/Math/Poisson.h b/Engine/Source/Runtime/Experimental/Chaos/Public/Chaos/Math/Poisson.h
--- a/Engine/Source/Runtime/Experimental/Chaos/Public/Chaos/Math/Poisson.h
+++ b/Engine/Source/Runtime/Experimental/Chaos/Public/Chaos/Math/Poisson.h
@@ -336,8 +336,7 @@ Laplacian(
 	//first compute element contributions
 	TArray<T> element_contributions; element_contributions.SetNum(Mesh.Num() * 4);
 	const int32* MeshPtr = &Mesh[0][0]; // Assumes Mesh is a contiguous array of int32.
-	//PhysicsParallelFor(int32(Mesh.Num()*4 / (d + 1)), [&](const int32 e)
-	for(int32 e=0; e < Mesh.Num(); e++)
+	PhysicsParallelFor(Mesh.Num(), [&](const int32 e)
 	{
 		T Deinv[d * d]{
 			De_inverse[d * d * e + 0],
@@ -366,17 +365,17 @@ Laplacian(
 				element_contributions[(d + 1) * e + ie] += Ae_ieje * u[Elem[je]];
 			}
 		}
-	//});
-	}
+	});
 
 	//now gather from elements
-	//PhysicsParallelFor(IncidentElements.Num(), [&](const int32 i)
-	for (int32 i = 0; i < IncidentElements.Num(); i++)
+	PhysicsParallelFor(IncidentElements.Num(), [&](const int32 i)
 	{
+		if (!IncidentElements[i].Num())
+		{
+			return;
+		}
 		const TArray<int32>& IncidentElements_i = IncidentElements[i];
 		const TArray<int32>& IncidentElementsLocalIndex_i = IncidentElementsLocalIndex[i];
-		if (!IncidentElements_i.Num())
-			return;
 		const int32 p = MeshPtr[MinFlatIndex(IncidentElements_i, IncidentElementsLocalIndex_i)];
 		for (int32 e = 0; e < IncidentElements_i.Num(); e++)
 		{
@@ -384,8 +383,7 @@ Laplacian(
 			const int32 Local = IncidentElementsLocalIndex_i[e];
 			Lu[p] += element_contributions[(Elem * 4) + Local];
 		}
-	}
-	//});
+	});
 }
 
 template<class T, class TV_INT, int d=3>
