add newer convex decomp settings to the dataflow convex decomp settings:
- for leaf convex hull, expose the option to protect negative space w/ the recommended settings + control over the key settings
- for the older hull merging methods that exposed more options, add the option to choose the newer version of voxel search (leaving defaults alone to not change behavior)
- also expose the newer proprocess options to the index buffer convex decomp init method, so that the dataflow decomp can handle degenerate inputs as well.

diff --git a/Engine/Plugins/Experimental/GeometryCollectionPlugin/Source/GeometryCollectionNodes/Private/Dataflow/GeometryCollectionUtilityNodes.cpp b/Engine/Plugins/Experimental/GeometryCollectionPlugin/Source/GeometryCollectionNodes/Private/Dataflow/GeometryCollectionUtilityNodes.cpp
--- a/Engine/Plugins/Experimental/GeometryCollectionPlugin/Source/GeometryCollectionNodes/Private/Dataflow/GeometryCollectionUtilityNodes.cpp
+++ b/Engine/Plugins/Experimental/GeometryCollectionPlugin/Source/GeometryCollectionNodes/Private/Dataflow/GeometryCollectionUtilityNodes.cpp
@@ -46,12 +46,17 @@ namespace UE::Dataflow
 FMakeDataflowConvexDecompositionSettingsNode::FMakeDataflowConvexDecompositionSettingsNode(const UE::Dataflow::FNodeParameters& InParam, FGuid InGuid)
 	: FDataflowNode(InParam, InGuid)
 {
-	RegisterInputConnection(&MinSizeToDecompose);
-	RegisterInputConnection(&MaxGeoToHullVolumeRatioToDecompose);
-	RegisterInputConnection(&ErrorTolerance);
-	RegisterInputConnection(&MaxHullsPerGeometry);
-	RegisterInputConnection(&MinThicknessTolerance);
-	RegisterInputConnection(&NumAdditionalSplits);
+	RegisterInputConnection(&MinSizeToDecompose).SetCanHidePin(true).SetPinIsHidden(true);
+	RegisterInputConnection(&MaxGeoToHullVolumeRatioToDecompose).SetCanHidePin(true).SetPinIsHidden(true);
+	RegisterInputConnection(&ErrorTolerance).SetCanHidePin(true).SetPinIsHidden(true);
+	RegisterInputConnection(&MaxHullsPerGeometry).SetCanHidePin(true).SetPinIsHidden(true);
+	RegisterInputConnection(&MinThicknessTolerance).SetCanHidePin(true).SetPinIsHidden(true);
+	RegisterInputConnection(&NumAdditionalSplits).SetCanHidePin(true).SetPinIsHidden(true);
+	RegisterInputConnection(&bProtectNegativeSpace).SetCanHidePin(true).SetPinIsHidden(true);
+	RegisterInputConnection(&bOnlyConnectedToHull).SetCanHidePin(true).SetPinIsHidden(true);
+	RegisterInputConnection(&NegativeSpaceTolerance).SetCanHidePin(true).SetPinIsHidden(true);
+	RegisterInputConnection(&NegativeSpaceMinRadius).SetCanHidePin(true).SetPinIsHidden(true);
+
 	RegisterOutputConnection(&DecompositionSettings);
 }
 
@@ -66,6 +71,11 @@ void FMakeDataflowConvexDecompositionSettingsNode::Evaluate(UE::Dataflow::FConte
 		OutSettings.MaxHullsPerGeometry = GetValue(Context, &MaxHullsPerGeometry);
 		OutSettings.MinThicknessTolerance = GetValue(Context, &MinThicknessTolerance);
 		OutSettings.NumAdditionalSplits = GetValue(Context, &NumAdditionalSplits);
+		OutSettings.bProtectNegativeSpace = GetValue(Context, &bProtectNegativeSpace);
+		OutSettings.bOnlyConnectedToHull = GetValue(Context, &bOnlyConnectedToHull);
+		OutSettings.NegativeSpaceTolerance = GetValue(Context, &NegativeSpaceTolerance);
+		OutSettings.NegativeSpaceMinRadius = GetValue(Context, &NegativeSpaceMinRadius);
+
 		SetValue(Context, OutSettings, &DecompositionSettings);
 	}
 }
@@ -136,6 +146,11 @@ void FCreateLeafConvexHullsDataflowNode::Evaluate(UE::Dataflow::FContext& Contex
 			LeafSettings.DecompositionSettings.MaxHullsPerGeometry = InDecompSettings.MaxHullsPerGeometry;
 			LeafSettings.DecompositionSettings.MinThicknessTolerance = InDecompSettings.MinThicknessTolerance;
 			LeafSettings.DecompositionSettings.NumAdditionalSplits = InDecompSettings.NumAdditionalSplits;
+			LeafSettings.DecompositionSettings.bProtectNegativeSpace = InDecompSettings.bProtectNegativeSpace;
+			LeafSettings.DecompositionSettings.bOnlyConnectedToHull = InDecompSettings.bOnlyConnectedToHull;
+			LeafSettings.DecompositionSettings.NegativeSpaceMinRadius = InDecompSettings.NegativeSpaceMinRadius;
+			LeafSettings.DecompositionSettings.NegativeSpaceTolerance = InDecompSettings.NegativeSpaceTolerance;
+
 			LeafSettings.bComputeIntersectionsBeforeHull = bComputeIntersectionsBeforeHull;
 			FGeometryCollectionConvexUtility::GenerateLeafConvexHulls(*GeomCollection, bRestrictToSelection, SelectedBones, LeafSettings);
 			SetValue(Context, (const FManagedArrayCollection&)(*GeomCollection), &Collection);
@@ -266,6 +281,8 @@ static UE::Geometry::FNegativeSpaceSampleSettings::ESampleMethod ConvertNegative
 		return UE::Geometry::FNegativeSpaceSampleSettings::ESampleMethod::Uniform;
 	case ENegativeSpaceSampleMethodDataflowEnum::VoxelSearch:
 		return UE::Geometry::FNegativeSpaceSampleSettings::ESampleMethod::VoxelSearch;
+	case ENegativeSpaceSampleMethodDataflowEnum::NavigableVoxelSearch:
+		return UE::Geometry::FNegativeSpaceSampleSettings::ESampleMethod::NavigableVoxelSearch;
 	}
 	return UE::Geometry::FNegativeSpaceSampleSettings::ESampleMethod::Uniform;
 }
diff --git a/Engine/Plugins/Experimental/GeometryCollectionPlugin/Source/GeometryCollectionNodes/Public/Dataflow/GeometryCollectionUtilityNodes.h b/Engine/Plugins/Experimental/GeometryCollectionPlugin/Source/GeometryCollectionNodes/Public/Dataflow/GeometryCollectionUtilityNodes.h
--- a/Engine/Plugins/Experimental/GeometryCollectionPlugin/Source/GeometryCollectionNodes/Public/Dataflow/GeometryCollectionUtilityNodes.h
+++ b/Engine/Plugins/Experimental/GeometryCollectionPlugin/Source/GeometryCollectionNodes/Public/Dataflow/GeometryCollectionUtilityNodes.h
@@ -48,12 +48,12 @@ public:
 	float MaxGeoToHullVolumeRatioToDecompose = 1.f;
 
 	// Stop splitting when hulls have error less than this (expressed in cm; will be cubed for volumetric error).
-	// Note: ErrorTolerance must be > 0 or MaxHullsPerGeometry > 1, or decomposition will not be performed.
-	UPROPERTY(EditAnywhere, Category = Decomposition, meta = (ClampMin = 0.0))
+	// Note: Decomposition will only be performed if: bProtectNegativeSpace is true, ErrorTolerance is > 0, or MaxHullsPerGeometry > 1
+	UPROPERTY(EditAnywhere, Category = Decomposition, meta = (ClampMin = 0.0, Units = cm, EditCondition = "!bProtectNegativeSpace"))
 	float ErrorTolerance = 0.f;
 
 	// If greater than zero, maximum number of convex hulls to use in each convex decomposition.
-	//Note: ErrorTolerance must be > 0 or MaxHullsPerGeometry > 1, or decomposition will not be performed.
+	// Note: Decomposition will only be performed if: bProtectNegativeSpace is true, ErrorTolerance is > 0, or MaxHullsPerGeometry > 1
 	UPROPERTY(EditAnywhere, Category = Decomposition, meta = (ClampMin = -1))
 	int32 MaxHullsPerGeometry = -1;
 
@@ -62,8 +62,24 @@ public:
 	float MinThicknessTolerance = 0.f;
 
 	// Control the search effort spent per convex decomposition: larger values will require more computation but may find better convex decompositions
-	UPROPERTY(EditAnywhere, Category = Decomposition, meta = (ClampMin = 0))
+	UPROPERTY(EditAnywhere, Category = Decomposition, meta = (ClampMin = 0, EditCondition = "!bProtectNegativeSpace"))
 	int32 NumAdditionalSplits = 4;
+
+	/** Whether to drive decomposition by finding a negative space that should not be covered by convex hulls. If enabled, ErrorTolerance and NumAdditionalSplits will not be used. */
+	UPROPERTY(EditAnywhere, Category = NegativeSpace)
+	bool bProtectNegativeSpace = false;
+
+	/** When protecting negative space, only look for space that is connected out to the convex hull. This removes inaccessable internal negative space from consideration. */
+	UPROPERTY(EditAnywhere, Category = NegativeSpace, meta = (EditCondition = "bProtectNegativeSpace"))
+	bool bOnlyConnectedToHull = true;
+
+	/** Amount of space to leave between convex hulls and protected negative space */
+	UPROPERTY(EditAnywhere, Category = NegativeSpace, meta = (ClampMin = .01, UIMin = .1, Units = cm, EditCondition = "bProtectNegativeSpace"))
+	float NegativeSpaceTolerance = 2.0;
+
+	/** Spheres smaller than this are not included in the negative space */
+	UPROPERTY(EditAnywhere, Category = NegativeSpace, meta = (ClampMin = 0, Units = cm, EditCondition = "bProtectNegativeSpace"))
+	float NegativeSpaceMinRadius = 10.0;
 };
 
 
@@ -85,12 +101,12 @@ public:
 	float MaxGeoToHullVolumeRatioToDecompose = 1.f;
 
 	// Stop splitting when hulls have error less than this (expressed in cm; will be cubed for volumetric error).
-	// Note: ErrorTolerance must be > 0 or MaxHullsPerGeometry > 1, or decomposition will not be performed.
-	UPROPERTY(EditAnywhere, Category = Decomposition, meta = (ClampMin = 0.0, DataflowInput))
+	// Note: Decomposition will only be performed if: bProtectNegativeSpace is true, ErrorTolerance is > 0, or MaxHullsPerGeometry > 1
+	UPROPERTY(EditAnywhere, Category = Decomposition, meta = (ClampMin = 0.0, Units = cm, DataflowInput, EditCondition = "!bProtectNegativeSpace"))
 	float ErrorTolerance = 0.f;
 
 	// If greater than zero, maximum number of convex hulls to use in each convex decomposition.
-	//Note: ErrorTolerance must be > 0 or MaxHullsPerGeometry > 1, or decomposition will not be performed.
+	// Note: Decomposition will only be performed if: bProtectNegativeSpace is true, ErrorTolerance is > 0, or MaxHullsPerGeometry > 1
 	UPROPERTY(EditAnywhere, Category = Decomposition, meta = (ClampMin = -1, DataflowInput))
 	int32 MaxHullsPerGeometry = -1;
 
@@ -99,9 +115,25 @@ public:
 	float MinThicknessTolerance = 0.f;
 
 	// Control the search effort spent per convex decomposition: larger values will require more computation but may find better convex decompositions
-	UPROPERTY(EditAnywhere, Category = Decomposition, meta = (ClampMin = 0, DataflowInput))
+	UPROPERTY(EditAnywhere, Category = Decomposition, meta = (ClampMin = 0, DataflowInput, EditCondition = "!bProtectNegativeSpace"))
 	int32 NumAdditionalSplits = 4;
 
+	/** Whether to drive decomposition by finding a negative space that should not be covered by convex hulls. If enabled, ErrorTolerance and NumAdditionalSplits will not be used. */
+	UPROPERTY(EditAnywhere, Category = NegativeSpace, meta = (DataflowInput))
+	bool bProtectNegativeSpace = false;
+
+	/** When protecting negative space, only look for space that is connected out to the convex hull. This removes inaccessable internal negative space from consideration. */
+	UPROPERTY(EditAnywhere, Category = NegativeSpace, meta = (DataflowInput, EditCondition = "bProtectNegativeSpace"))
+	bool bOnlyConnectedToHull = true;
+
+	/** Amount of space to leave between convex hulls and protected negative space */
+	UPROPERTY(EditAnywhere, Category = NegativeSpace, meta = (DataflowInput, ClampMin = .01, UIMin = .1, Units = cm, EditCondition = "bProtectNegativeSpace"))
+	float NegativeSpaceTolerance = 2.0;
+
+	/** Spheres smaller than this are not included in the negative space */
+	UPROPERTY(EditAnywhere, Category = NegativeSpace, meta = (DataflowInput, ClampMin = 0, Units = cm, EditCondition = "bProtectNegativeSpace"))
+	float NegativeSpaceMinRadius = 10.0;
+
 	UPROPERTY(meta = (DataflowOutput))
 	FDataflowConvexDecompositionSettings DecompositionSettings;
 
@@ -265,7 +297,9 @@ enum class ENegativeSpaceSampleMethodDataflowEnum : uint8
 	// Place sample spheres in a uniform grid pattern
 	Uniform,
 	// Use voxel-based subtraction and offsetting methods to specifically target concavities
-	VoxelSearch
+	VoxelSearch,
+	// Use a variant of VoxelSearch that aims to limit negative space to the space that can be accessed by a ball of radius >= MinRadius
+	NavigableVoxelSearch
 };
 
 /**
@@ -333,11 +367,11 @@ private:
 	ENegativeSpaceSampleMethodDataflowEnum SampleMethod = ENegativeSpaceSampleMethodDataflowEnum::Uniform;
 
 	/** Whether to require that all candidate locations identified by Voxel Search are covered by negative space samples, up to the specified Min Sample Spacing. Only applies to Voxel Search. */
-	UPROPERTY(EditAnywhere, Category = NegativeSpace, meta = (EditCondition = "bProtectNegativeSpace && SampleMethod == ENegativeSpaceSampleMethodDataflowEnum::VoxelSearch", EditConditionHides))
+	UPROPERTY(EditAnywhere, Category = NegativeSpace, meta = (EditCondition = "bProtectNegativeSpace && SampleMethod != ENegativeSpaceSampleMethodDataflowEnum::Uniform", EditConditionHides))
 	bool bRequireSearchSampleCoverage = false;
 
 	/** When performing Voxel Search, only look for negative space that is connected out to the convex hull. This removes inaccessable internal negative space from consideration. Only applies to Voxel Search. */
-	UPROPERTY(EditAnywhere, Category = NegativeSpace, meta = (EditCondition = "bProtectNegativeSpace && SampleMethod == ENegativeSpaceSampleMethodDataflowEnum::VoxelSearch", EditConditionHides))
+	UPROPERTY(EditAnywhere, Category = NegativeSpace, meta = (EditCondition = "bProtectNegativeSpace && SampleMethod != ENegativeSpaceSampleMethodDataflowEnum::Uniform", EditConditionHides))
 	bool bOnlyConnectedToHull = false;
 
 	/** Approximate number of spheres to consider when covering negative space. Only applicable with the Uniform Sample Method or if Require Search Sample Coverage is disabled. */
@@ -437,11 +471,11 @@ private:
 	ENegativeSpaceSampleMethodDataflowEnum SampleMethod = ENegativeSpaceSampleMethodDataflowEnum::Uniform;
 
 	/** Whether to require that all candidate locations identified by Voxel Search are covered by negative space samples, up to the specified Min Sample Spacing. Only applies to Voxel Search. */
-	UPROPERTY(EditAnywhere, Category = NegativeSpace, meta = (EditCondition = "bProtectNegativeSpace && SampleMethod == ENegativeSpaceSampleMethodDataflowEnum::VoxelSearch", EditConditionHides))
+	UPROPERTY(EditAnywhere, Category = NegativeSpace, meta = (EditCondition = "bProtectNegativeSpace && SampleMethod != ENegativeSpaceSampleMethodDataflowEnum::Uniform", EditConditionHides))
 	bool bRequireSearchSampleCoverage = false;
 
 	/** When performing Voxel Search, only look for negative space that is connected out to the convex hull. This removes inaccessable internal negative space from consideration. Only applies to Voxel Search. */
-	UPROPERTY(EditAnywhere, Category = NegativeSpace, meta = (EditCondition = "bProtectNegativeSpace && SampleMethod == ENegativeSpaceSampleMethodDataflowEnum::VoxelSearch", EditConditionHides))
+	UPROPERTY(EditAnywhere, Category = NegativeSpace, meta = (EditCondition = "bProtectNegativeSpace && SampleMethod != ENegativeSpaceSampleMethodDataflowEnum::Uniform", EditConditionHides))
 	bool bOnlyConnectedToHull = false;
 
 	/** Approximate number of spheres to consider when covering negative space. Only applicable with the Uniform Sample Method or if Require Search Sample Coverage is disabled. */
@@ -561,11 +595,11 @@ public:
 	ENegativeSpaceSampleMethodDataflowEnum SampleMethod = ENegativeSpaceSampleMethodDataflowEnum::Uniform;
 
 	/** Whether to require that all candidate locations identified by Voxel Search are covered by negative space samples, up to the specified Min Sample Spacing. Only applies to Voxel Search. */
-	UPROPERTY(EditAnywhere, Category = NegativeSpace, meta = (EditCondition = "bProtectNegativeSpace && SampleMethod == ENegativeSpaceSampleMethodDataflowEnum::VoxelSearch", EditConditionHides))
+	UPROPERTY(EditAnywhere, Category = NegativeSpace, meta = (EditCondition = "bProtectNegativeSpace && SampleMethod != ENegativeSpaceSampleMethodDataflowEnum::Uniform", EditConditionHides))
 	bool bRequireSearchSampleCoverage = false;
 
 	/** When performing Voxel Search, only look for negative space that is connected out to the convex hull. This removes inaccessable internal negative space from consideration. Only applies to Voxel Search. */
-	UPROPERTY(EditAnywhere, Category = NegativeSpace, meta = (EditCondition = "bProtectNegativeSpace && SampleMethod == ENegativeSpaceSampleMethodDataflowEnum::VoxelSearch", EditConditionHides))
+	UPROPERTY(EditAnywhere, Category = NegativeSpace, meta = (EditCondition = "bProtectNegativeSpace && SampleMethod != ENegativeSpaceSampleMethodDataflowEnum::Uniform", EditConditionHides))
 	bool bOnlyConnectedToHull = false;
 
 	/** Approximate number of spheres to consider when covering negative space. Only applicable with the Uniform Sample Method or if Require Search Sample Coverage is disabled. */
diff --git a/Engine/Source/Runtime/Experimental/Chaos/Private/GeometryCollection/GeometryCollectionConvexUtility.cpp b/Engine/Source/Runtime/Experimental/Chaos/Private/GeometryCollection/GeometryCollectionConvexUtility.cpp
--- a/Engine/Source/Runtime/Experimental/Chaos/Private/GeometryCollection/GeometryCollectionConvexUtility.cpp
+++ b/Engine/Source/Runtime/Experimental/Chaos/Private/GeometryCollection/GeometryCollectionConvexUtility.cpp
@@ -1428,7 +1428,7 @@ void HullsFromGeometry(
 
 			if (OptionalDecompositionSettings)
 			{
-				bool bAttemptDecomposition = OptionalDecompositionSettings->MaxHullsPerGeometry > 1 || OptionalDecompositionSettings->ErrorTolerance > 0.0;
+				bool bAttemptDecomposition = OptionalDecompositionSettings->MaxHullsPerGeometry > 1 || OptionalDecompositionSettings->ErrorTolerance > 0.0 || OptionalDecompositionSettings->bProtectNegativeSpace;
 				double InitialGeoVolume = 0.0, InitialHullVolume = 0.0;
 				if (bAttemptDecomposition && (OptionalDecompositionSettings->MinGeoVolumeToDecompose > 0.0 || OptionalDecompositionSettings->MaxGeoToHullVolumeRatioToDecompose < 1.0))
 				{
@@ -1460,19 +1460,71 @@ void HullsFromGeometry(
 				if (bAttemptDecomposition)
 				{
 					UE::Geometry::FConvexDecomposition3 Decomposition;
+					UE::Geometry::FConvexDecomposition3::FPreprocessMeshOptions PreprocessOptions;
+					PreprocessOptions.bMergeEdges = true;
+					PreprocessOptions.ThickenInputAfterHullFailure = OptionalDecompositionSettings->ThickenHullOnFailure;
+
 					if (GeometryIntersectedWithExternalHull)
 					{
-						Decomposition.InitializeFromMesh(*GeometryIntersectedWithExternalHull, true);
+						Decomposition.InitializeFromMesh(*GeometryIntersectedWithExternalHull, PreprocessOptions);
 					}
 					else
 					{
 						int32 VertexStart = Geometry.VertexStart[GeomIdx];
 						TArrayView<const FVector3f> VerticesView(Geometry.Vertex.GetData() + VertexStart, Geometry.VertexCount[GeomIdx]);
 						TArrayView<const FIntVector3> FacesView(Geometry.Indices.GetData() + Geometry.FaceStart[GeomIdx], Geometry.FaceCount[GeomIdx]);
-						Decomposition.InitializeFromIndexMesh(VerticesView, FacesView, true, -VertexStart);
+						Decomposition.InitializeFromIndexMesh(VerticesView, FacesView, PreprocessOptions, -VertexStart);
+					}
+
+					const bool bIsSolid = Decomposition.IsInputSolid();
+					Decomposition.bTreatAsSolid = bIsSolid;
+
+					if (OptionalDecompositionSettings->bProtectNegativeSpace)
+					{
+						// Use settings tuned for navigation-driven decomposition
+
+						UE::Geometry::FNegativeSpaceSampleSettings Settings;
+						Settings.ApplyDefaults();
+						Settings.bOnlyConnectedToHull = OptionalDecompositionSettings->bOnlyConnectedToHull;
+						Settings.MinRadius = OptionalDecompositionSettings->NegativeSpaceMinRadius;
+						Settings.ReduceRadiusMargin = OptionalDecompositionSettings->NegativeSpaceTolerance;
+						Settings.bRequireSearchSampleCoverage = true;
+						Settings.bAllowSamplesInsideMesh = !bIsSolid;
+
+						Decomposition.MaxConvexEdgePlanes = 4;
+						Decomposition.bSplitDisconnectedComponents = false;
+						Decomposition.ConvexEdgeAngleMoreSamplesThreshold = 180;
+						Decomposition.ThickenAfterHullFailure = PreprocessOptions.ThickenInputAfterHullFailure;
+
+						Decomposition.InitializeNegativeSpace(Settings);
+
+						constexpr int32 MaxAllowedSplits = 1000000; // more parts than any expected / reasonable decomposition
+						int32 TargetNumSplits = OptionalDecompositionSettings->MaxHullsPerGeometry > 0 ? OptionalDecompositionSettings->MaxHullsPerGeometry - 1 : MaxAllowedSplits + 1;
+						for (int32 Split = 0; Split < TargetNumSplits; Split++)
+						{
+							int32 NumSplit = Decomposition.SplitWorst(false, -1, true, Settings.ReduceRadiusMargin * .5);
+
+							if (NumSplit == 0)
+							{
+								break;
+							}
+
+							if (!ensureMsgf(Split < MaxAllowedSplits, TEXT("Convex decomposition split the input %d times; likely stuck in a loop"), Split))
+							{
+								break;
+							}
+						}
+
+						Decomposition.FixHullOverlapsInNegativeSpace();
+						int32 TargetNumPieces = OptionalDecompositionSettings->MaxHullsPerGeometry > 0 ? OptionalDecompositionSettings->MaxHullsPerGeometry : -1;
+						Decomposition.MergeBest(TargetNumPieces, 0, OptionalDecompositionSettings->MinThicknessTolerance, true);
 					}
-					Decomposition.Compute(OptionalDecompositionSettings->MaxHullsPerGeometry, OptionalDecompositionSettings->NumAdditionalSplits,
-						OptionalDecompositionSettings->ErrorTolerance, OptionalDecompositionSettings->MinThicknessTolerance, OptionalDecompositionSettings->MaxHullsPerGeometry);
+					else
+					{
+						Decomposition.Compute(OptionalDecompositionSettings->MaxHullsPerGeometry, OptionalDecompositionSettings->NumAdditionalSplits,
+							OptionalDecompositionSettings->ErrorTolerance, OptionalDecompositionSettings->MinThicknessTolerance, OptionalDecompositionSettings->MaxHullsPerGeometry);
+					}
+
 					int32 NumHulls = Decomposition.NumHulls();
 					if ((NumHulls > 0 && !Hull) || NumHulls > 1)
 					{
diff --git a/Engine/Source/Runtime/Experimental/Chaos/Public/GeometryCollection/GeometryCollectionConvexUtility.h b/Engine/Source/Runtime/Experimental/Chaos/Public/GeometryCollection/GeometryCollectionConvexUtility.h
--- a/Engine/Source/Runtime/Experimental/Chaos/Public/GeometryCollection/GeometryCollectionConvexUtility.h
+++ b/Engine/Source/Runtime/Experimental/Chaos/Public/GeometryCollection/GeometryCollectionConvexUtility.h
@@ -182,24 +182,17 @@ public:
 	// Settings to control convex decompositions. Note the default values are set to *not* perform any decomposition
 	struct FConvexDecompositionSettings
 	{
-		FConvexDecompositionSettings() :
-			MinGeoVolumeToDecompose(0.f),
-			MaxGeoToHullVolumeRatioToDecompose(1.f),
-			MaxHullsPerGeometry(1),
-			ErrorTolerance(0.f),
-			MinThicknessTolerance(0.f),
-			NumAdditionalSplits(4)
-		{}
+		FConvexDecompositionSettings() = default;
 
 		///
 		/// These settings are filters that control whether convex decomposition is performed
 		///
 		
 		// If greater than zero, the minimum geometry volume to consider for convex decomposition
-		float MinGeoVolumeToDecompose;
+		float MinGeoVolumeToDecompose = 0.f;
 
 		// If the geo volume / hull volume ratio is greater than this, do not consider convex decomposition
-		float MaxGeoToHullVolumeRatioToDecompose;
+		float MaxGeoToHullVolumeRatioToDecompose = 1.f;
 
 
 		///
@@ -207,17 +200,33 @@ public:
 		///
 		
 		// If > 0, specify the maximum number of convex hulls to create.  Hull merges will be attempted until this number of hulls remains, or no merges are possible.
-		int32 MaxHullsPerGeometry;
+		int32 MaxHullsPerGeometry = 1;
 
 		// Stop splitting when hulls have error less than this (expressed in cm; will be cubed for volumetric error). Overrides NumOutputHulls if non-zero.
-		float ErrorTolerance;
+		float ErrorTolerance = 0.f;
 
 		// Optionally specify a minimum thickness (in cm) for convex parts; parts below this thickness will always be merged away. Will merge hulls with greater error than ErrorTolerance when needed.
-		float MinThicknessTolerance;
+		float MinThicknessTolerance = 0.f;
 
 		// How far to go beyond the target number of hulls for initial decomposition (before merges) -- larger values will require more computation but can find better convex decompositions
-		int32 NumAdditionalSplits;
+		int32 NumAdditionalSplits = 4;
+
+		// How much to thicken a hull on failure to construct -- to allow 3D decomposition of planar or degenerate meshes and parts-of-meshes
+		//~ Note Chaos::FConvex has its own process to thicken hulls that are smaller than .1 units across (specified by Chaos::TriQuadPrismInflation)
+		//~ Applying some smaller thickening earlier in the decomposition process just allows degenerate parts to make it through the decomposition.
+		float ThickenHullOnFailure = UE_KINDA_SMALL_NUMBER;
+
+		/** Whether to use a sphere cover to define negative space that should not be covered by convex hulls */
+		bool bProtectNegativeSpace = false;
+
+		/** When protecting negative space, only look for space that is connected out to the convex hull. This removes inaccessable internal negative space from consideration. */
+		bool bOnlyConnectedToHull = true;
+
+		/** Amount of space to leave between convex hulls and protected negative space */
+		float NegativeSpaceTolerance = 2.0f;
 
+		/** Spheres smaller than this are not included in the negative space */
+		float NegativeSpaceMinRadius = 10.0f;
 	};
 
 	// Settings to control how convex hulls are generated for rigid/leaf nodes, from the geometry and/or imported collision shapes
diff --git a/Engine/Source/Runtime/GeometryCore/Private/CompGeom/ConvexDecomposition3.cpp b/Engine/Source/Runtime/GeometryCore/Private/CompGeom/ConvexDecomposition3.cpp
--- a/Engine/Source/Runtime/GeometryCore/Private/CompGeom/ConvexDecomposition3.cpp
+++ b/Engine/Source/Runtime/GeometryCore/Private/CompGeom/ConvexDecomposition3.cpp
@@ -1617,6 +1617,20 @@ FConvexDecomposition3::FConvexPart::FConvexPart(const FDynamicMesh3& SourceMesh,
 	InitializeFromInternalGeo(Options, TransformOut);
 }
 
+FConvexDecomposition3::FConvexPart::FConvexPart(TArrayView<const FVector3f> Vertices, TArrayView<const FIntVector3> Faces, const FConvexDecomposition3::FPreprocessMeshOptions& Options, FTransformSRT3d& TransformOut, int32 FaceVertexOffset)
+{
+	for (const FVector3f& V : Vertices)
+	{
+		InternalGeo.AppendVertex((FVector)V);
+	}
+	for (FIntVector F : Faces)
+	{
+		InternalGeo.AppendTriangle(FIndex3i(F.X + FaceVertexOffset, F.Y + FaceVertexOffset, F.Z + FaceVertexOffset));
+	}
+
+	InitializeFromInternalGeo(Options, TransformOut);
+}
+
 FConvexDecomposition3::FConvexPart::FConvexPart(TArrayView<const FVector3f> Vertices, TArrayView<const FIntVector3> Faces, bool bMergeEdges, FTransformSRT3d& TransformOut, int32 FaceVertexOffset)
 {
 	for (const FVector3f& V : Vertices)
@@ -1676,9 +1690,17 @@ void FConvexDecomposition3::FConvexPart::InitializeFromInternalGeo(const FConvex
 
 
 void FConvexDecomposition3::InitializeFromIndexMesh(TArrayView<const FVector3f> Vertices, TArrayView<const FIntVector> Faces, bool bMergeEdges, int32 FaceVertexOffset)
+{
+	FConvexDecomposition3::FPreprocessMeshOptions Preprocess;
+	Preprocess.bMergeEdges = bMergeEdges;
+	InitializeFromIndexMesh(Vertices, Faces, Preprocess, FaceVertexOffset);
+}
+
+
+void FConvexDecomposition3::InitializeFromIndexMesh(TArrayView<const FVector3f> Vertices, TArrayView<const FIntVector> Faces, const FConvexDecomposition3::FPreprocessMeshOptions& Preprocess, int32 FaceVertexOffset)
 {
 	Decomposition.Empty();
-	FConvexPart* Convex = new FConvexPart(Vertices, Faces, bMergeEdges, ResultTransform, FaceVertexOffset);
+	FConvexPart* Convex = new FConvexPart(Vertices, Faces, Preprocess, ResultTransform, FaceVertexOffset);
 	if (Convex->IsFailed())
 	{
 		delete Convex;
diff --git a/Engine/Source/Runtime/GeometryCore/Public/CompGeom/ConvexDecomposition3.h b/Engine/Source/Runtime/GeometryCore/Public/CompGeom/ConvexDecomposition3.h
--- a/Engine/Source/Runtime/GeometryCore/Public/CompGeom/ConvexDecomposition3.h
+++ b/Engine/Source/Runtime/GeometryCore/Public/CompGeom/ConvexDecomposition3.h
@@ -348,6 +348,14 @@ public:
 	 * @param FaceVertexOffset	Indices from the Faces array are optionally offset by this value.  Useful e.g. to take slices of the multi-geometry vertex and face buffers of FGeometryCollection.
 	 */
 	GEOMETRYCORE_API void InitializeFromIndexMesh(TArrayView<const FVector3f> Vertices, TArrayView<const FIntVector> Faces, bool bMergeEdges, int32 FaceVertexOffset = 0);
+	/**
+	 * Initialize convex decomposition with a triangle index mesh
+	 * @param Vertices			Vertex buffer for mesh to decompose
+	 * @param Faces				Triangle buffer for mesh to decompose
+	 * @param Options			Options for pre-processing the mesh before decomposition (e.g., to weld matching edges, thicken degenerate inputs, or to run a general pre-process lambda function)
+	 * @param FaceVertexOffset	Indices from the Faces array are optionally offset by this value.  Useful e.g. to take slices of the multi-geometry vertex and face buffers of FGeometryCollection.
+	 */
+	GEOMETRYCORE_API void InitializeFromIndexMesh(TArrayView<const FVector3f> Vertices, TArrayView<const FIntVector> Faces, const FPreprocessMeshOptions& Options, int32 FaceVertexOffset = 0);
 
 	/**
 	 * Find negative space that should be protected. Uses the mesh passed on construction or to InitializeFromMesh or InitializeFromIndexMesh
@@ -583,6 +591,7 @@ public:
 		FConvexPart() {}
 		FConvexPart(const FDynamicMesh3& SourceMesh, bool bMergeEdges, FTransformSRT3d& TransformOut);
 		FConvexPart(const FDynamicMesh3& SourceMesh, const FConvexDecomposition3::FPreprocessMeshOptions& Options, FTransformSRT3d& TransformOut);
+		FConvexPart(TArrayView<const FVector3f> Vertices, TArrayView<const FIntVector3> Faces, const FConvexDecomposition3::FPreprocessMeshOptions& Options, FTransformSRT3d& TransformOut, int32 FaceVertexOffset = 0);
 		FConvexPart(TArrayView<const FVector3f> Vertices, TArrayView<const FIntVector3> Faces, bool bMergeEdges, FTransformSRT3d& TransformOut, int32 FaceVertexOffset = 0);
 
 		// Allow direct construction of a compact part (e.g. to allow construction of a pre-existing convex hull)
