fix issue where dataflow blueprint nodes would redundantly register the blueprint OnChanged delegate handle, causing the dataflow graph to evaluate mutliple times when the blueprint was changed

diff --git a/Engine/Plugins/Experimental/GeometryCollectionPlugin/Source/GeometryCollectionNodes/Private/Dataflow/GeometryCollectionMeshNodes.cpp b/Engine/Plugins/Experimental/GeometryCollectionPlugin/Source/GeometryCollectionNodes/Private/Dataflow/GeometryCollectionMeshNodes.cpp
--- a/Engine/Plugins/Experimental/GeometryCollectionPlugin/Source/GeometryCollectionNodes/Private/Dataflow/GeometryCollectionMeshNodes.cpp
+++ b/Engine/Plugins/Experimental/GeometryCollectionPlugin/Source/GeometryCollectionNodes/Private/Dataflow/GeometryCollectionMeshNodes.cpp
@@ -616,6 +616,10 @@ void FMeshProcessorDataflowNodeBase::SetupBlueprintEvent()
 	{
 		if (UBlueprint* Blueprint = Cast<UBlueprint>(MeshProcessor->ClassGeneratedBy))
 		{
+			if (!ensure(!BlueprintChangeDelegateHandle.IsValid()))
+			{
+				TeardownBlueprintEvent();
+			}
 			BlueprintChangeDelegateHandle = Blueprint->OnChanged().AddLambda([this](UBlueprint* BP)
 			{
 				Invalidate();
@@ -633,6 +637,7 @@ void FMeshProcessorDataflowNodeBase::TeardownBlueprintEvent()
 		if (UBlueprint* Blueprint = Cast<UBlueprint>(MeshProcessor->ClassGeneratedBy))
 		{
 			Blueprint->OnChanged().Remove(BlueprintChangeDelegateHandle);
+			BlueprintChangeDelegateHandle.Reset();
 		}
 	}
 #endif
diff --git a/Engine/Plugins/Experimental/GeometryCollectionPlugin/Source/GeometryCollectionNodes/Public/Dataflow/GeometryCollectionMeshNodes.h b/Engine/Plugins/Experimental/GeometryCollectionPlugin/Source/GeometryCollectionNodes/Public/Dataflow/GeometryCollectionMeshNodes.h
--- a/Engine/Plugins/Experimental/GeometryCollectionPlugin/Source/GeometryCollectionNodes/Public/Dataflow/GeometryCollectionMeshNodes.h
+++ b/Engine/Plugins/Experimental/GeometryCollectionPlugin/Source/GeometryCollectionNodes/Public/Dataflow/GeometryCollectionMeshNodes.h
@@ -98,7 +98,11 @@ protected:
 	virtual void PostSerialize(const FArchive& Ar) override final
 	{
 		Super::PostSerialize(Ar);
-		SetupBlueprintEvent();
+		if (Ar.IsLoading())
+		{
+			TeardownBlueprintEvent();
+			SetupBlueprintEvent();
+		}
 	}
 
 private:
