Dataflow: add node for duplicating a UV layer in a DataflowMesh

diff --git a/Engine/Plugins/Experimental/GeometryCollectionPlugin/Source/GeometryCollectionNodes/Private/Dataflow/GeometryCollectionMeshNodes.cpp b/Engine/Plugins/Experimental/GeometryCollectionPlugin/Source/GeometryCollectionNodes/Private/Dataflow/GeometryCollectionMeshNodes.cpp
--- a/Engine/Plugins/Experimental/GeometryCollectionPlugin/Source/GeometryCollectionNodes/Private/Dataflow/GeometryCollectionMeshNodes.cpp
+++ b/Engine/Plugins/Experimental/GeometryCollectionPlugin/Source/GeometryCollectionNodes/Private/Dataflow/GeometryCollectionMeshNodes.cpp
@@ -53,6 +53,7 @@ namespace UE::Dataflow
 		DATAFLOW_NODE_REGISTER_CREATION_FACTORY(FMeshAppendDataflowNode);
 		DATAFLOW_NODE_REGISTER_CREATION_FACTORY(FDataflowMeshAppendDataflowNode);
 		DATAFLOW_NODE_REGISTER_CREATION_FACTORY(FMakeDataflowMeshDataflowNode);
+		DATAFLOW_NODE_REGISTER_CREATION_FACTORY(FDuplicateMeshUVChannelNode);
 		DATAFLOW_NODE_REGISTER_CREATION_FACTORY(FSplitDataflowMeshDataflowNode);
 		DATAFLOW_NODE_REGISTER_CREATION_FACTORY(FMeshCopyToPointsDataflowNode);
 		DATAFLOW_NODE_REGISTER_CREATION_FACTORY(FGetMeshDataDataflowNode);
@@ -506,6 +507,51 @@ void FSplitDataflowMeshDataflowNode::Evaluate(UE::Dataflow::FContext& Context, c
 	}
 }
 
+FDuplicateMeshUVChannelNode::FDuplicateMeshUVChannelNode(const UE::Dataflow::FNodeParameters& InParam, FGuid InGuid)
+	: FDataflowNode(InParam, InGuid)
+{
+	RegisterInputConnection(&Mesh);
+	RegisterOutputConnection(&Mesh, &Mesh);
+	RegisterOutputConnection(&NewUVChannel);
+}
+
+void FDuplicateMeshUVChannelNode::Evaluate(UE::Dataflow::FContext& Context, const FDataflowOutput* Out) const
+{
+	int32 NewUVLayerIndex = -1;
+
+	if (Out->IsA(&Mesh) || Out->IsA(&NewUVChannel))
+	{
+		if (TObjectPtr<UDataflowMesh> InMesh = GetValue(Context, &Mesh))
+		{
+			if (const FDynamicMesh3* const InDynamicMesh = InMesh->GetDynamicMesh())
+			{
+				if (InDynamicMesh->HasAttributes() && SourceUVChannel >= 0 && SourceUVChannel < InDynamicMesh->Attributes()->NumUVLayers())
+				{
+					UE::Geometry::FDynamicMesh3 OutDynamicMesh;
+					OutDynamicMesh.Copy(*InDynamicMesh);
+					OutDynamicMesh.EnableAttributes();
+
+					NewUVLayerIndex = OutDynamicMesh.Attributes()->NumUVLayers();
+					OutDynamicMesh.Attributes()->SetNumUVLayers(NewUVLayerIndex + 1);
+
+					const UE::Geometry::FDynamicMeshUVOverlay* const SourceUVLayer = OutDynamicMesh.Attributes()->GetUVLayer(SourceUVChannel);
+					OutDynamicMesh.Attributes()->GetUVLayer(NewUVLayerIndex)->Copy(*SourceUVLayer);
+
+					TObjectPtr<UDataflowMesh> OutMesh = NewObject<UDataflowMesh>();
+					OutMesh->SetDynamicMesh(MoveTemp(OutDynamicMesh));
+					OutMesh->SetMaterials(InMesh->GetMaterials());
+					SetValue(Context, OutMesh, &Mesh);
+					SetValue(Context, NewUVLayerIndex, &NewUVChannel);
+					return;
+				}
+			}
+		}
+	}
+
+	SafeForwardInput(Context, &Mesh, &Mesh);
+	SetValue(Context, NewUVLayerIndex, &NewUVChannel);
+}
+
 
 void FMeshCopyToPointsDataflowNode::Evaluate(UE::Dataflow::FContext& Context, const FDataflowOutput* Out) const
 {
diff --git a/Engine/Plugins/Experimental/GeometryCollectionPlugin/Source/GeometryCollectionNodes/Public/Dataflow/GeometryCollectionMeshNodes.h b/Engine/Plugins/Experimental/GeometryCollectionPlugin/Source/GeometryCollectionNodes/Public/Dataflow/GeometryCollectionMeshNodes.h
--- a/Engine/Plugins/Experimental/GeometryCollectionPlugin/Source/GeometryCollectionNodes/Public/Dataflow/GeometryCollectionMeshNodes.h
+++ b/Engine/Plugins/Experimental/GeometryCollectionPlugin/Source/GeometryCollectionNodes/Public/Dataflow/GeometryCollectionMeshNodes.h
@@ -561,6 +561,39 @@ private:
 };
 
 
+/**
+ * Create a new UV layer/channel in a UDataflowMesh
+ */
+USTRUCT(meta = (DataflowGeometryCollection))
+struct FDuplicateMeshUVChannelNode : public FDataflowNode
+{
+	GENERATED_USTRUCT_BODY()
+	DATAFLOW_NODE_DEFINE_INTERNAL(FDuplicateMeshUVChannelNode, "DuplicateMeshUVChannelNode", "Mesh|Utilities", "Mesh UV DataflowMesh")
+	DATAFLOW_NODE_RENDER_TYPE("SurfaceRender", FName("UDataflowMesh"), "Mesh", "UVChannel")
+
+public:
+
+	FDuplicateMeshUVChannelNode(const UE::Dataflow::FNodeParameters& InParam, FGuid InGuid = FGuid::NewGuid());
+
+private:
+
+	/** DataflowMesh input/output */
+	UPROPERTY(meta = (DataflowInput, DataflowOutput, DataflowPassthrough = "Mesh", DataflowIntrinsic))
+	TObjectPtr<UDataflowMesh> Mesh;
+
+	/** Index of the source UV channel */
+	UPROPERTY(EditAnywhere, Category = Options)
+	int32 SourceUVChannel = -1;
+
+	/** Index of the added UV channel */
+	UPROPERTY(meta = (DataflowOutput))
+	int32 NewUVChannel = -1;
+
+	virtual void Evaluate(UE::Dataflow::FContext& Context, const FDataflowOutput* Out) const override;
+};
+
+
+
 /**
  *
  * Split a UDataflow mesh into a UDynamicMesh and a material array
