ChaosFlesh:
- Deleted unused ReferencePositions variable from ChaosDeformableSolver
- Corrected hide conditions for TransferVertexAttribute node

diff --git a/Engine/Plugins/Experimental/GeometryCollectionPlugin/Source/GeometryCollectionNodes/Public/Dataflow/GeometryCollectionTransferVertexAttributeNode.h b/Engine/Plugins/Experimental/GeometryCollectionPlugin/Source/GeometryCollectionNodes/Public/Dataflow/GeometryCollectionTransferVertexAttributeNode.h
--- a/Engine/Plugins/Experimental/GeometryCollectionPlugin/Source/GeometryCollectionNodes/Public/Dataflow/GeometryCollectionTransferVertexAttributeNode.h
+++ b/Engine/Plugins/Experimental/GeometryCollectionPlugin/Source/GeometryCollectionNodes/Public/Dataflow/GeometryCollectionTransferVertexAttributeNode.h
@@ -129,7 +129,7 @@ public:
 
 	/* Threshold based on distance from source triangle.Values past the threshold will falloff.[Defaults to 1 percent of triangle size(0.01)] */
 	UPROPERTY(EditAnywhere, Category = "Thresholds", meta = (EditCondition = 
-		"BoundingVolumeType == EDataflowTransferVertexAttributeNodeBoundingVolume::Triangle && EDataflowTransferVertexAttributeNodeFalloff Falloff == EDataflowTransferVertexAttributeNodeFalloff::None"
+		"Falloff != EDataflowTransferVertexAttributeNodeFalloff::None"
 		, EditConditionHides))
 	float FalloffThreshold = 0.01f;
 
@@ -206,7 +206,7 @@ public:
 	EDataflowTransferVertexAttributeNodeFalloff Falloff = EDataflowTransferVertexAttributeNodeFalloff::None;
 
 	/* Threshold based on distance from source triangle.Values past the threshold will falloff.[Defaults to 1 percent of triangle size(0.01)] */
-	UPROPERTY(EditAnywhere, Category = "Thresholds", meta = (EditCondition = "BoundingVolumeType == EDataflowTransferVertexAttributeNodeBoundingVolume::Triangle", EditConditionHides))
+	UPROPERTY(EditAnywhere, Category = "Thresholds", meta = (EditCondition = "Falloff != EDataflowTransferVertexAttributeNodeFalloff::None", EditConditionHides))
 	float FalloffThreshold = 0.01f;
 
 	/* Edge multiplier for the Bounding Volume Hierarchy(BVH) target's particle search radius. */
diff --git a/Engine/Source/Runtime/Experimental/Chaos/Public/Chaos/Deformable/ChaosDeformableSolver.h b/Engine/Source/Runtime/Experimental/Chaos/Public/Chaos/Deformable/ChaosDeformableSolver.h
--- a/Engine/Source/Runtime/Experimental/Chaos/Public/Chaos/Deformable/ChaosDeformableSolver.h
+++ b/Engine/Source/Runtime/Experimental/Chaos/Public/Chaos/Deformable/ChaosDeformableSolver.h
@@ -236,7 +236,6 @@ namespace Chaos::Softs
 		TUniquePtr <TArray<FSolverReal>> AllSkinEMeshArray;
 		TUniquePtr <TArray<int32>> ParticleComponentIndex;
 		TMap<int32, TSet<int32>> ParticleTriangleExclusionMap;
-		TUniquePtr <FSolverParticles> ReferencePositions;
 		//Muscle Activation Variables
 		TUniquePtr<Softs::FMuscleActivationConstraints<Softs::FSolverReal, Softs::FSolverParticles>> MuscleActivationConstraints;		
 		TMap<FThreadingProxy::FKey, int32> MuscleIndexOffset;
