Dataflow : make SetValue more resilient to value / reference type mismatch
- case detected was a Value using a base type while the reference was of a derived type, causing the stored type to be smaller than the expected type
- fixed CollectionSelectByAttr node that was passing a generic selection object in more specialized one

diff --git a/Engine/Plugins/Experimental/GeometryCollectionPlugin/Source/GeometryCollectionNodes/Private/Dataflow/GeometryCollectionSelectionNodes.cpp b/Engine/Plugins/Experimental/GeometryCollectionPlugin/Source/GeometryCollectionNodes/Private/Dataflow/GeometryCollectionSelectionNodes.cpp
--- a/Engine/Plugins/Experimental/GeometryCollectionPlugin/Source/GeometryCollectionNodes/Private/Dataflow/GeometryCollectionSelectionNodes.cpp
+++ b/Engine/Plugins/Experimental/GeometryCollectionPlugin/Source/GeometryCollectionNodes/Private/Dataflow/GeometryCollectionSelectionNodes.cpp
@@ -1469,14 +1469,14 @@ static void CreateSelectionFromAttr(const FManagedArrayCollection& InCollection,
 
 void FCollectionSelectionByAttrDataflowNode::Evaluate(UE::Dataflow::FContext& Context, const FDataflowOutput* Out) const
 {
-	if (Out->IsA<FDataflowVertexSelection>(&VertexSelection) ||
-		Out->IsA<FDataflowFaceSelection>(&FaceSelection) ||
-		Out->IsA<FDataflowTransformSelection>(&TransformSelection) ||
-		Out->IsA<FDataflowGeometrySelection>(&GeometrySelection) ||
-		Out->IsA<FDataflowMaterialSelection>(&MaterialSelection))
+	if (Out->IsA(&VertexSelection) ||
+		Out->IsA(&FaceSelection) ||
+		Out->IsA(&TransformSelection) ||
+		Out->IsA(&GeometrySelection) ||
+		Out->IsA(&MaterialSelection))
 	{
-		const FManagedArrayCollection& InCollection = GetValue<FManagedArrayCollection>(Context, &Collection);
-		FCollectionAttributeKey InAttributeKey = GetValue<FCollectionAttributeKey>(Context, &AttributeKey);
+		const FManagedArrayCollection& InCollection = GetValue(Context, &Collection);
+		FCollectionAttributeKey InAttributeKey = GetValue(Context, &AttributeKey);
 		FName GroupName = UE::Dataflow::Private::GetAttributeFromEnumAsName(Group);
 		FName AttributeName = FName(Attribute);
 		if (IsConnected(&AttributeKey))
@@ -1491,31 +1491,62 @@ void FCollectionSelectionByAttrDataflowNode::Evaluate(UE::Dataflow::FContext& Co
 			{
 				const int32 NumFaces = InCollection.NumElements(GroupName);
 
-				FDataflowSelection NewSelection;
-				NewSelection.Initialize(NumFaces, false);
+				FDataflowSelection NewGenericSelection;
+				NewGenericSelection.Initialize(NumFaces, false);
 
 				CreateSelectionFromAttr(InCollection,
 					GroupName,
 					AttributeName,
 					Value,
 					Operation,
-					NewSelection);
+					NewGenericSelection);
+
+				using namespace UE::Dataflow::Private;
+
+				FDataflowVertexSelection OutVertexSelection;
+				if (GroupName == UE::Dataflow::Private::GetAttributeFromEnumAsName(ESelectionByAttrGroup::Vertices))
+				{
+					OutVertexSelection.Initialize(NewGenericSelection);
+				}
+				SetValue(Context, MoveTemp(OutVertexSelection), &VertexSelection);
+
+				FDataflowFaceSelection OutFaceSelection;
+				if (GroupName == UE::Dataflow::Private::GetAttributeFromEnumAsName(ESelectionByAttrGroup::Faces))
+				{
+					OutFaceSelection.Initialize(NewGenericSelection);
+				}
+				SetValue(Context, MoveTemp(OutFaceSelection), &FaceSelection);
 
-				SetValue(Context, GroupName == UE::Dataflow::Private::GetAttributeFromEnumAsName(ESelectionByAttrGroup::Vertices) ? MoveTemp(NewSelection) : FDataflowSelection(), &VertexSelection);
-				SetValue(Context, GroupName == UE::Dataflow::Private::GetAttributeFromEnumAsName(ESelectionByAttrGroup::Faces) ? MoveTemp(NewSelection) : FDataflowSelection(), &FaceSelection);
-				SetValue(Context, GroupName == UE::Dataflow::Private::GetAttributeFromEnumAsName(ESelectionByAttrGroup::Transform) ? MoveTemp(NewSelection) : FDataflowSelection(), &TransformSelection);
-				SetValue(Context, GroupName == UE::Dataflow::Private::GetAttributeFromEnumAsName(ESelectionByAttrGroup::Geometry) ? MoveTemp(NewSelection) : FDataflowSelection(), &GeometrySelection);
-				SetValue(Context, GroupName == UE::Dataflow::Private::GetAttributeFromEnumAsName(ESelectionByAttrGroup::Material) ? MoveTemp(NewSelection) : FDataflowSelection(), &MaterialSelection);
+				FDataflowTransformSelection OutTransformSelection;
+				if (GroupName == UE::Dataflow::Private::GetAttributeFromEnumAsName(ESelectionByAttrGroup::Transform))
+				{
+					OutTransformSelection.Initialize(NewGenericSelection);
+				}
+				SetValue(Context, MoveTemp(OutTransformSelection), &TransformSelection);
+
+				FDataflowGeometrySelection OutGeometrySelection;
+				if (GroupName == UE::Dataflow::Private::GetAttributeFromEnumAsName(ESelectionByAttrGroup::Geometry))
+				{
+					OutGeometrySelection.Initialize(NewGenericSelection);
+				}
+				SetValue(Context, MoveTemp(OutGeometrySelection), &GeometrySelection);
+
+				FDataflowMaterialSelection OutMaterialSelection;
+				if (GroupName == UE::Dataflow::Private::GetAttributeFromEnumAsName(ESelectionByAttrGroup::Material))
+				{
+					OutMaterialSelection.Initialize(NewGenericSelection);
+				}
+				SetValue(Context, MoveTemp(OutMaterialSelection), &MaterialSelection);
 
 				return;
 			}
 		}
 
-		SetValue(Context, FDataflowSelection(), &VertexSelection);
-		SetValue(Context, FDataflowSelection(), &FaceSelection);
-		SetValue(Context, FDataflowSelection(), &TransformSelection);
-		SetValue(Context, FDataflowSelection(), &GeometrySelection);
-		SetValue(Context, FDataflowSelection(), &MaterialSelection);
+		SetValue(Context, FDataflowVertexSelection(), &VertexSelection);
+		SetValue(Context, FDataflowFaceSelection(), &FaceSelection);
+		SetValue(Context, FDataflowTransformSelection(), &TransformSelection);
+		SetValue(Context, FDataflowGeometrySelection(), &GeometrySelection);
+		SetValue(Context, FDataflowMaterialSelection(), &MaterialSelection);
 	}
 	else if (Out->IsA<FManagedArrayCollection>(&Collection))
 	{
diff --git a/Engine/Source/Runtime/Experimental/Dataflow/Core/Public/Dataflow/DataflowNode.h b/Engine/Source/Runtime/Experimental/Dataflow/Core/Public/Dataflow/DataflowNode.h
--- a/Engine/Source/Runtime/Experimental/Dataflow/Core/Public/Dataflow/DataflowNode.h
+++ b/Engine/Source/Runtime/Experimental/Dataflow/Core/Public/Dataflow/DataflowNode.h
@@ -531,9 +531,12 @@ PRAGMA_ENABLE_DEPRECATION_WARNINGS
 	*   @param Value : The value to store in the contexts data store. 
 	*   @param Reference : Pointer to a member of this node that corresponds with the output to set. 
 	*/
-	template<class T, typename = std::enable_if_t<!std::is_base_of_v<FDataflowAnyType, T>>>
-	void SetValue(UE::Dataflow::FContext& Context, T&& Value, const typename TDecay<T>::Type* Reference) const
+	template<class T, class U, typename = std::enable_if_t<!std::is_base_of_v<FDataflowAnyType, U>>>
+	void SetValue(UE::Dataflow::FContext& Context, T&& Value, const U* Reference) const
 	{
+		using FInValueType = typename TDecay<decltype(Value)>::Type;
+		using FInReferenceType = typename TDecay<decltype(*Reference)>::Type;
+		static_assert(std::is_same_v<FInValueType, FInReferenceType>);
 		if (const FDataflowOutput* Output = FindOutput(Reference))
 		{
 			Output->template SetValue<T>(Forward<T>(Value), Context);
