Dataflow:
- modify TextureToImage to handle GPU textures
- add an ImageToTexture node to create a transient UTexture from an image
- modify MaterialTextureOverride to accept texture object as input

diff --git a/Engine/Plugins/Experimental/GeometryCollectionPlugin/Source/GeometryCollectionNodes/Private/Dataflow/GeometryCollectionMaterialInterfaceNodes.cpp b/Engine/Plugins/Experimental/GeometryCollectionPlugin/Source/GeometryCollectionNodes/Private/Dataflow/GeometryCollectionMaterialInterfaceNodes.cpp
--- a/Engine/Plugins/Experimental/GeometryCollectionPlugin/Source/GeometryCollectionNodes/Private/Dataflow/GeometryCollectionMaterialInterfaceNodes.cpp
+++ b/Engine/Plugins/Experimental/GeometryCollectionPlugin/Source/GeometryCollectionNodes/Private/Dataflow/GeometryCollectionMaterialInterfaceNodes.cpp
@@ -337,56 +337,63 @@ FMaterialInterfaceTextureOverrideDataflowNode::FMaterialInterfaceTextureOverride
 {
 	RegisterInputConnection(&Material);
 	RegisterOutputConnection(&Material, &Material);
+	RegisterInputConnection(&TargetTexture);
+	RegisterInputConnection(&OverrideTexture);
 }
 
 void FMaterialInterfaceTextureOverrideDataflowNode::Evaluate(UE::Dataflow::FContext& Context, const FDataflowOutput* Out) const
 {
 #if WITH_EDITOR
-	if (Out->IsA(&Material) && TargetTexture && OverrideTexture)
+	if (Out->IsA(&Material) && TargetTexture)
 	{
-		if (TObjectPtr<const UMaterialInterface> SourceMaterial = GetValue<TObjectPtr<UMaterialInterface>>(Context, &Material))
+		const TObjectPtr<UTexture2D> InOverrideTexture = GetValue(Context, &OverrideTexture);
+
+		if (InOverrideTexture)
 		{
-			while (const UMaterialInstanceDynamic* const InMID = Cast<const UMaterialInstanceDynamic>(SourceMaterial))
+			if (TObjectPtr<const UMaterialInterface> SourceMaterial = GetValue<TObjectPtr<UMaterialInterface>>(Context, &Material))
 			{
-				SourceMaterial = InMID->Parent;
-			}
+				while (const UMaterialInstanceDynamic* const InMID = Cast<const UMaterialInstanceDynamic>(SourceMaterial))
+				{
+					SourceMaterial = InMID->Parent;
+				}
 
-			if (const UMaterialInstanceConstant* const SourceMIC = Cast<const UMaterialInstanceConstant>(SourceMaterial))
-			{
-				if (UMaterialInstanceConstant* const DuplicateMaterial = DuplicateObject<UMaterialInstanceConstant>(SourceMIC, nullptr))
-				{				
-					for (FTextureParameterValue& Param : DuplicateMaterial->TextureParameterValues)
+				if (const UMaterialInstanceConstant* const SourceMIC = Cast<const UMaterialInstanceConstant>(SourceMaterial))
+				{
+					if (UMaterialInstanceConstant* const DuplicateMaterial = DuplicateObject<UMaterialInstanceConstant>(SourceMIC, nullptr))
 					{
-						if (Param.ParameterValue == TargetTexture)
+						for (FTextureParameterValue& Param : DuplicateMaterial->TextureParameterValues)
 						{
-							DuplicateMaterial->SetTextureParameterValueEditorOnly(Param.ParameterInfo, OverrideTexture);
+							if (Param.ParameterValue == TargetTexture)
+							{
+								DuplicateMaterial->SetTextureParameterValueEditorOnly(Param.ParameterInfo, InOverrideTexture);
+							}
 						}
-					}
 
-					if (const TObjectPtr<UMaterialInterface> OutMI = CastChecked<UMaterialInterface>(DuplicateMaterial))		// should be a straightforward upcast
-					{
-						// TODO: Make a MaterialTerminalNode so we can also save this out as an asset
+						if (const TObjectPtr<UMaterialInterface> OutMI = CastChecked<UMaterialInterface>(DuplicateMaterial))		// should be a straightforward upcast
+						{
+							// TODO: Make a MaterialTerminalNode so we can also save this out as an asset
 
-						SetValue(Context, OutMI, &Material);
-						return;
+							SetValue(Context, OutMI, &Material);
+							return;
+						}
+					}
+					else
+					{
+						Context.Warning(TEXT("Error creating duplicate material"), this, Out);
 					}
 				}
 				else
 				{
-					Context.Warning(TEXT("Error creating duplicate material"), this, Out);
+					// TODO: handle non-UMaterialInstanceConstant materials
+
+					Context.Warning(TEXT("Input material is not a UMaterialInstanceConstant"), this, Out);
 				}
 			}
 			else
 			{
-				// TODO: handle non-UMaterialInstanceConstant materials
-
-				Context.Warning(TEXT("Input material is not a UMaterialInstanceConstant"), this, Out);
+				Context.Warning(TEXT("Input material or its parent is not UMaterialInstanceConstant"), this, Out);
 			}
 		}
-		else
-		{
-			Context.Warning(TEXT("Input material or its parent is not UMaterialInstanceConstant"), this, Out);
-		}
 	}
 #else	// WITH_EDITOR
 
diff --git a/Engine/Plugins/Experimental/GeometryCollectionPlugin/Source/GeometryCollectionNodes/Public/Dataflow/GeometryCollectionMaterialInterfaceNodes.h b/Engine/Plugins/Experimental/GeometryCollectionPlugin/Source/GeometryCollectionNodes/Public/Dataflow/GeometryCollectionMaterialInterfaceNodes.h
--- a/Engine/Plugins/Experimental/GeometryCollectionPlugin/Source/GeometryCollectionNodes/Public/Dataflow/GeometryCollectionMaterialInterfaceNodes.h
+++ b/Engine/Plugins/Experimental/GeometryCollectionPlugin/Source/GeometryCollectionNodes/Public/Dataflow/GeometryCollectionMaterialInterfaceNodes.h
@@ -234,10 +234,10 @@ private:
 	UPROPERTY(meta = (DataflowInput, DataflowOutput, DataflowPassthrough = "Material"))
 	TObjectPtr<UMaterialInterface> Material;
 
-	UPROPERTY(EditAnywhere, Category = "Override")
+	UPROPERTY(EditAnywhere, Category = "Override", meta = (DataflowInput))
 	TObjectPtr<const UTexture2D> TargetTexture;
 
-	UPROPERTY(EditAnywhere, Category = "Override")
+	UPROPERTY(EditAnywhere, Category = "Override", meta=(DataflowInput))
 	TObjectPtr<UTexture2D> OverrideTexture;
 
 	//~ Begin FDataflowNode interface
diff --git a/Engine/Source/Runtime/Experimental/Dataflow/Engine/Private/Dataflow/DataflowTextureAssetNodes.cpp b/Engine/Source/Runtime/Experimental/Dataflow/Engine/Private/Dataflow/DataflowTextureAssetNodes.cpp
--- a/Engine/Source/Runtime/Experimental/Dataflow/Engine/Private/Dataflow/DataflowTextureAssetNodes.cpp
+++ b/Engine/Source/Runtime/Experimental/Dataflow/Engine/Private/Dataflow/DataflowTextureAssetNodes.cpp
@@ -10,6 +10,7 @@ namespace UE::Dataflow
 	{
 		DATAFLOW_NODE_REGISTER_CREATION_FACTORY(FDataflowTextureTerminalNode);
 		DATAFLOW_NODE_REGISTER_CREATION_FACTORY(FDataflowTextureToImageNode);
+		DATAFLOW_NODE_REGISTER_CREATION_FACTORY(FDataflowImageToTextureNode);
 	}
 }
 
@@ -85,30 +86,69 @@ void FDataflowTextureToImageNode::Evaluate(UE::Dataflow::FContext& Context, cons
 {
 	if (Out->IsA(&Image))
 	{
-		if (const TObjectPtr<const UTexture2D> InTextureAsset = GetValue(Context, &TextureAsset))
+		if (const TObjectPtr<UTexture2D> InTextureAsset = GetValue(Context, &TextureAsset))
 		{
+			FImage TempImage;
 			if (const FSharedImageConstRef CpuCopy = InTextureAsset->GetCPUCopy())
 			{
-				FImage TempImage;
 				CpuCopy->CopyTo(TempImage, ERawImageFormat::RGBA32F, EGammaSpace::Linear);
-				const TArrayView64<FVector4f> Buffer = { (FVector4f*)TempImage.AsRGBA32F().GetData(), TempImage.AsRGBA32F().Num() };
-
-				FDataflowImage OutImage;
-				OutImage.CreateRGBA32F(CpuCopy->GetWidth(), CpuCopy->GetHeight());
-				OutImage.CopyRGBAPixels(Buffer);
-
-				SetValue(Context, OutImage, &Image);
-				return;
 			}
+//#if WITH_EDITORONLY_DATA
+//			else if (InTextureAsset->Source.IsValid())
+//			{
+//				FImage MipImage;
+//				InTextureAsset->Source.GetMipImage(MipImage, 0);
+//				MipImage.CopyTo(TempImage, ERawImageFormat::RGBA32F, EGammaSpace::Linear);
+//			}
+//#endif
 			else
 			{
-				// TODO: Find out if we can get access to GPU textures somehow
+				// TODO: Handle more ways to get image data from textures
+				Context.Warning(TEXT("Unable to read image from texture"));
 
-				Context.Warning(TEXT("Input texture asset not available on the CPU"), this, Out);
+				FDataflowImage EmptyImage;
+				SetValue(Context, EmptyImage, &Image);
+				return;
 			}
+
+			const TArrayView64<FVector4f> Buffer = { (FVector4f*)TempImage.AsRGBA32F().GetData(), TempImage.AsRGBA32F().Num() };
+			FDataflowImage OutImage;
+			OutImage.CreateRGBA32F(TempImage.GetWidth(), TempImage.GetHeight());
+			OutImage.CopyRGBAPixels(Buffer);
+
+			SetValue(Context, OutImage, &Image);
+			return;
 		}
 
 		FDataflowImage EmptyImage;
 		SetValue(Context, EmptyImage, &Image);
 	}
 }
+
+
+FDataflowImageToTextureNode::FDataflowImageToTextureNode(const UE::Dataflow::FNodeParameters& InParam, FGuid InGuid)
+	: FDataflowNode(InParam, InGuid)
+{
+	RegisterInputConnection(&Image);
+	RegisterInputConnection(&TextureName);
+	RegisterOutputConnection(&TransientTexture);
+}
+
+
+void FDataflowImageToTextureNode::Evaluate(UE::Dataflow::FContext& Context, const FDataflowOutput* Out) const
+{
+	if (Out->IsA(&TransientTexture))
+	{
+		const FName InName = GetValue(Context, &TextureName);
+		UObject* Package = GetTransientPackage();
+		const FName UniqueName = MakeUniqueObjectName(Package, UTexture2D::StaticClass(), InName);
+		TObjectPtr<UTexture2D> OutTexture = NewObject<UTexture2D>(Package, UniqueName, RF_Transient);
+
+		const FDataflowImage InImage = GetValue(Context, &Image);
+
+		DataflowTextureAssetNodes::Private::UpdateTexture2DFromImage(*OutTexture, InImage.GetImage());
+
+		SetValue(Context, OutTexture, &TransientTexture);
+		return;
+	}
+}
diff --git a/Engine/Source/Runtime/Experimental/Dataflow/Engine/Public/Dataflow/DataflowTextureAssetNodes.h b/Engine/Source/Runtime/Experimental/Dataflow/Engine/Public/Dataflow/DataflowTextureAssetNodes.h
--- a/Engine/Source/Runtime/Experimental/Dataflow/Engine/Public/Dataflow/DataflowTextureAssetNodes.h
+++ b/Engine/Source/Runtime/Experimental/Dataflow/Engine/Public/Dataflow/DataflowTextureAssetNodes.h
@@ -60,6 +60,33 @@ private:
 };
 
 
+/**
+* Create a transient texture asset from an image
+*/
+USTRUCT(meta = (DataflowTerminal))
+struct FDataflowImageToTextureNode : public FDataflowNode
+{
+	GENERATED_USTRUCT_BODY()
+	DATAFLOW_NODE_DEFINE_INTERNAL(FDataflowImageToTextureNode, "ImageToTexture", "Image", "Image Texture")
+
+private:
+	UPROPERTY(meta = (DataflowInput, DataflowIntrinsic))
+	FDataflowImage Image;
+
+	UPROPERTY(EditAnywhere, Category = "Texture", meta = (DataflowInput))
+	FName TextureName;
+
+	UPROPERTY(meta = (DataflowOutput))
+	TObjectPtr<UTexture2D> TransientTexture = nullptr;
+
+	virtual void Evaluate(UE::Dataflow::FContext& Context, const FDataflowOutput* Out) const override;
+
+public:
+	FDataflowImageToTextureNode(const UE::Dataflow::FNodeParameters& InParam, FGuid InGuid = FGuid::NewGuid());
+};
+
+
+
 namespace UE::Dataflow
 {
 	void RegisterTextureAssetNodes();
