Dataflow : fix add material node not properly adding the materials

diff --git a/Engine/Plugins/Experimental/GeometryCollectionPlugin/Source/GeometryCollectionNodes/Private/Dataflow/GeometryCollectionMaterialInterfaceNodes.cpp b/Engine/Plugins/Experimental/GeometryCollectionPlugin/Source/GeometryCollectionNodes/Private/Dataflow/GeometryCollectionMaterialInterfaceNodes.cpp
--- a/Engine/Plugins/Experimental/GeometryCollectionPlugin/Source/GeometryCollectionNodes/Private/Dataflow/GeometryCollectionMaterialInterfaceNodes.cpp
+++ b/Engine/Plugins/Experimental/GeometryCollectionPlugin/Source/GeometryCollectionNodes/Private/Dataflow/GeometryCollectionMaterialInterfaceNodes.cpp
@@ -171,7 +171,11 @@ void FAddToMaterialInterfaceArrayDataflowNode::Evaluate(UE::Dataflow::FContext&
 	if (Out->IsA(&MaterialArray))
 	{
 		TArray<TObjectPtr<UMaterialInterface>> InMaterialArray = GetValue(Context, &MaterialArray);
-		InMaterialArray.Append(MaterialsToAdd);
+		for (int32 Index = 0; Index < MaterialsToAdd.Num(); ++Index)
+		{
+			const TObjectPtr<UMaterialInterface> MaterialToAdd = GetValue(Context,GetConnectionReference(Index));
+			InMaterialArray.Add(MaterialToAdd);
+		}
 		SetValue(Context, MoveTemp(InMaterialArray), &MaterialArray);
 	}
 }
