make dataflow lerp(a,b,t) node not require a<b to do the lerp

diff --git a/Engine/Plugins/Experimental/GeometryCollectionPlugin/Source/GeometryCollectionNodes/Private/Dataflow/GeometryCollectionMathNodes.cpp b/Engine/Plugins/Experimental/GeometryCollectionPlugin/Source/GeometryCollectionNodes/Private/Dataflow/GeometryCollectionMathNodes.cpp
--- a/Engine/Plugins/Experimental/GeometryCollectionPlugin/Source/GeometryCollectionNodes/Private/Dataflow/GeometryCollectionMathNodes.cpp
+++ b/Engine/Plugins/Experimental/GeometryCollectionPlugin/Source/GeometryCollectionNodes/Private/Dataflow/GeometryCollectionMathNodes.cpp
@@ -448,17 +448,14 @@ void FLogeDataflowNode::Evaluate(UE::Dataflow::FContext& Context, const FDataflo
 
 void FLerpDataflowNode::Evaluate(UE::Dataflow::FContext& Context, const FDataflowOutput* Out) const
 {
-	if (Out->IsA<float>(&ReturnValue))
+	if (Out->IsA(&ReturnValue))
 	{
-		const float InA = GetValue<float>(Context, &A, A);
-		const float InB = GetValue<float>(Context, &B, B);
-		const float InAlpha = GetValue<float>(Context, &Alpha, Alpha);
+		const float InA = GetValue(Context, &A, A);
+		const float InB = GetValue(Context, &B, B);
+		const float InAlpha = GetValue(Context, &Alpha, Alpha);
 
 		float Result = InA;
-		if (InB > InA)
-		{
-			Result = FMath::Lerp(InA, InB, InAlpha);
-		}
+		Result = FMath::Lerp(InA, InB, InAlpha);
 		SetValue(Context, Result, &ReturnValue);
 	}
 }
