make collection attribute transform dataflow node use the attribute name/group params + also recompute bounding boxes if the vertices were moved

diff --git a/Engine/Plugins/Experimental/GeometryCollectionPlugin/Source/GeometryCollectionNodes/Private/Dataflow/GeometryCollectionVerticesNodes.cpp b/Engine/Plugins/Experimental/GeometryCollectionPlugin/Source/GeometryCollectionNodes/Private/Dataflow/GeometryCollectionVerticesNodes.cpp
--- a/Engine/Plugins/Experimental/GeometryCollectionPlugin/Source/GeometryCollectionNodes/Private/Dataflow/GeometryCollectionVerticesNodes.cpp
+++ b/Engine/Plugins/Experimental/GeometryCollectionPlugin/Source/GeometryCollectionNodes/Private/Dataflow/GeometryCollectionVerticesNodes.cpp
@@ -3,6 +3,8 @@
 #include "Dataflow/GeometryCollectionVerticesNodes.h"
 #include "GeometryCollection/GeometryCollection.h"
 
+#include "GeometryCollection/Facades/CollectionBoundsFacade.h"
+
 #include "Dataflow/DataflowCore.h"
 
 #include UE_INLINE_GENERATED_CPP_BY_NAME(GeometryCollectionVerticesNodes)
@@ -20,22 +22,32 @@ void FTransformCollectionAttributeDataflowNode::Evaluate(UE::Dataflow::FContext&
 {
 	if (Out->IsA<DataType>(&Collection))
 	{
-		const FTransform& Transform = GetValue<FTransform>(Context, &TransformIn);
-		FManagedArrayCollection CollectionValue = GetValue<FManagedArrayCollection>(Context, &Collection);
-		if (CollectionValue.FindAttributeTyped<FVector3f>("Vertex", FGeometryCollection::VerticesGroup))
+		const FTransform& Transform = GetValue(Context, &TransformIn);
+		FManagedArrayCollection CollectionValue = GetValue(Context, &Collection);
+		FName AttribFName(*AttributeName);
+		FName GroupFName(*GroupName);
+		if (TManagedArray<FVector3f>* VecAttrib = CollectionValue.FindAttributeTyped<FVector3f>(AttribFName, GroupFName))
 		{
-			auto UEVertd = [](FVector3f V) { return FVector3d(V.X, V.Y, V.Z); };
-			auto UEVertf = [](FVector3d V) { return FVector3f(V.X, V.Y, V.Z); };
+			FMatrix44f Matrix(LocalTransform.ToMatrixWithScale() * Transform.ToMatrixWithScale());
+			for (int32 i = 0; i < VecAttrib->Num(); i++)
+			{
+				(*VecAttrib)[i] = Matrix.TransformPosition((*VecAttrib)[i]);
+			}
 
-			FTransform FinalTransform = LocalTransform * Transform;
-			//FMatrix44d Matrix = (LocalTransform * Transform).ToMatrixWithScale();
-			TManagedArray<FVector3f>& Vertices = CollectionValue.ModifyAttribute<FVector3f>("Vertex", FGeometryCollection::VerticesGroup);
-			for (int32 i = 0; i < Vertices.Num(); i++)
+			// If we've moved vertices, also try to update bounding boxes
+			if (AttribFName == FName("Vertex") && GroupFName == FGeometryCollection::VerticesGroup)
 			{
-				//Vertices[i] = UEVertf(Matrix.TransformPosition(UEVertd(Vertices[i])));
-				Vertices[i] = UEVertf(FinalTransform.TransformPosition(UEVertd(Vertices[i])));
+				GeometryCollection::Facades::FBoundsFacade BoundsFacade(CollectionValue);
+				if (BoundsFacade.IsValid())
+				{
+					BoundsFacade.UpdateBoundingBox();
+				}
 			}
 		}
+		else
+		{
+			UE_LOG(LogChaosDataflow, Warning, TEXT("Could not find Vector3f Attribute \"%s\" in Group \"%s\""), *AttributeName, *GroupName);
+		}
 		SetValue(Context, MoveTemp(CollectionValue), &Collection);
 
 	}
