Dataflow : Exposing mesh transform options as input

diff --git a/Engine/Plugins/Experimental/GeometryCollectionPlugin/Source/GeometryCollectionNodes/Private/Dataflow/GeometryCollectionNodes.cpp b/Engine/Plugins/Experimental/GeometryCollectionPlugin/Source/GeometryCollectionNodes/Private/Dataflow/GeometryCollectionNodes.cpp
--- a/Engine/Plugins/Experimental/GeometryCollectionPlugin/Source/GeometryCollectionNodes/Private/Dataflow/GeometryCollectionNodes.cpp
+++ b/Engine/Plugins/Experimental/GeometryCollectionPlugin/Source/GeometryCollectionNodes/Private/Dataflow/GeometryCollectionNodes.cpp
@@ -529,14 +529,23 @@ void FTransformMeshDataflowNode::Evaluate(UE::Dataflow::FContext& Context, const
 			TObjectPtr<UDynamicMesh> NewMesh = NewObject<UDynamicMesh>();
 			NewMesh->SetMesh(InMesh->GetMeshRef());
 
-			FTransform NewTransform = GeometryCollection::Facades::FCollectionTransformFacade::BuildTransform(Translate,
+			const FVector& InTranslate = GetValue(Context, &Translate);
+			const FVector& InRotate = GetValue(Context, &Rotate);
+			const FVector& InScale = GetValue(Context, &Scale);
+			const float InUniformScale = GetValue(Context, &UniformScale);
+			const FVector& InRotatePivot = GetValue(Context, &RotatePivot);
+			const FVector& InScalePivot = GetValue(Context, &ScalePivot);
+			const bool InbInvertTransformation = GetValue(Context, &bInvertTransformation);
+			
+			FTransform NewTransform = GeometryCollection::Facades::FCollectionTransformFacade::BuildTransform(
+				InTranslate,
 				(uint8)RotationOrder,
-				Rotate,
-				Scale,
-				UniformScale,
-				RotatePivot,
-				ScalePivot,
-				bInvertTransformation);
+				InRotate,
+				InScale,
+				InUniformScale,
+				InRotatePivot,
+				InScalePivot,
+				InbInvertTransformation);
 
 			UE::Geometry::FDynamicMesh3& DynamicMesh = NewMesh->GetMeshRef();
 
diff --git a/Engine/Plugins/Experimental/GeometryCollectionPlugin/Source/GeometryCollectionNodes/Public/Dataflow/GeometryCollectionNodes.h b/Engine/Plugins/Experimental/GeometryCollectionPlugin/Source/GeometryCollectionNodes/Public/Dataflow/GeometryCollectionNodes.h
--- a/Engine/Plugins/Experimental/GeometryCollectionPlugin/Source/GeometryCollectionNodes/Public/Dataflow/GeometryCollectionNodes.h
+++ b/Engine/Plugins/Experimental/GeometryCollectionPlugin/Source/GeometryCollectionNodes/Public/Dataflow/GeometryCollectionNodes.h
@@ -782,7 +782,7 @@ public:
 	TObjectPtr<UDynamicMesh> Mesh;
 
 	/** Translation */
-	UPROPERTY(EditAnywhere, Category = "Transform");
+	UPROPERTY(EditAnywhere, Category = "Transform", meta = (DataflowInput));
 	FVector Translate = FVector(0.0);
 
 	/** Rotation Order */
@@ -790,11 +790,11 @@ public:
 	ERotationOrderEnum RotationOrder = ERotationOrderEnum::Dataflow_RotationOrder_XYZ;
 
 	/** Rotation */
-	UPROPERTY(EditAnywhere, Category = "Transform");
+	UPROPERTY(EditAnywhere, Category = "Transform", meta = (DataflowInput));
 	FVector Rotate = FVector(0.0);
 
 	/** Scale */
-	UPROPERTY(EditAnywhere, Category = "Transform", meta = (UIMin = 0.001f));
+	UPROPERTY(EditAnywhere, Category = "Transform", meta = (UIMin = 0.001f, DataflowInput));
 	FVector Scale = FVector(1.0);
 
 	/** Shear */
@@ -802,7 +802,7 @@ public:
 //	FVector Shear = FVector(0.0);
 
 	/** Uniform scale */
-	UPROPERTY(EditAnywhere, Category = "Transform", meta = (UIMin = 0.001f));
+	UPROPERTY(EditAnywhere, Category = "Transform", meta = (UIMin = 0.001f, DataflowInput));
 	float UniformScale = 1.f;
 
 	/** Pivot for the rotation */
@@ -810,17 +810,24 @@ public:
 	FVector RotatePivot = FVector(0.0);
 
 	/** Pivot for the scale */
-	UPROPERTY(EditAnywhere, Category = "Pivot");
+	UPROPERTY(EditAnywhere, Category = "Pivot", meta = (DataflowInput));
 	FVector ScalePivot = FVector(0.0);
 
 	/** Invert the transformation */
-	UPROPERTY(EditAnywhere, Category = "General");
+	UPROPERTY(EditAnywhere, Category = "General", meta = (DataflowInput));
 	bool bInvertTransformation = false;
 
 	FTransformMeshDataflowNode(const UE::Dataflow::FNodeParameters& InParam, FGuid InGuid = FGuid::NewGuid())
 		: FDataflowNode(InParam, InGuid)
 	{
 		RegisterInputConnection(&Mesh);
+		RegisterInputConnection(&Translate);
+		RegisterInputConnection(&Rotate);
+		RegisterInputConnection(&Scale);
+		RegisterInputConnection(&UniformScale);
+		RegisterInputConnection(&RotatePivot).SetCanHidePin(true).SetPinIsHidden(true);
+		RegisterInputConnection(&ScalePivot).SetCanHidePin(true).SetPinIsHidden(true);
+		RegisterInputConnection(&bInvertTransformation).SetCanHidePin(true).SetPinIsHidden(true);
 		RegisterOutputConnection(&Mesh, &Mesh);
 	}
 
@@ -1122,7 +1129,7 @@ public:
 	bool bSetNotSelectedBonesToOppositeState = false;
 
 	/** GeometryCollection to set anchor state on */
-	UPROPERTY(meta = (DataflowInput, DataflowOutput, DataflowPassthrough = "Collection", DataflowIntrinsic))
+	UPROPERTY(meta = (DataflowInput, DataflowOutput, DataflowPassthrough = "Collection", DataflowIntrinsic, DataflowCollectionGroups="CustomOne,CustomTwo"))
 	FManagedArrayCollection Collection;
 
 	/** Bone selection for setting the state on */
