ChaosFlesh:
TransferVertexAttribute node fixes

diff --git a/Engine/Plugins/Experimental/GeometryCollectionPlugin/Source/GeometryCollectionNodes/Private/Dataflow/GeometryCollectionTransferVertexAttributeNode.cpp b/Engine/Plugins/Experimental/GeometryCollectionPlugin/Source/GeometryCollectionNodes/Private/Dataflow/GeometryCollectionTransferVertexAttributeNode.cpp
--- a/Engine/Plugins/Experimental/GeometryCollectionPlugin/Source/GeometryCollectionNodes/Private/Dataflow/GeometryCollectionTransferVertexAttributeNode.cpp
+++ b/Engine/Plugins/Experimental/GeometryCollectionPlugin/Source/GeometryCollectionNodes/Private/Dataflow/GeometryCollectionTransferVertexAttributeNode.cpp
@@ -672,6 +672,7 @@ void FGeometryCollectionTransferVertexSkinWeightsNode::Evaluate(UE::Dataflow::FC
 				// Compute the bone index mappings. This allows the transfer operator to retarget weights to the correct skeleton.
 				//
 				const TMap<int32, int32> SourceBoneToTargetBone = FindSourceToTargetTransformMap(SourceCollection, TargetCollection);
+				TSet<int32> MissingSourceBones;
 				for (int32 VertexIndex = 0; VertexIndex < Target.Vertex.Num(); VertexIndex++)
 				{
 					const FTransferData& TransferData = TransferDataArray[VertexIndex];
@@ -696,6 +697,11 @@ void FGeometryCollectionTransferVertexSkinWeightsNode::Evaluate(UE::Dataflow::FC
 										BoneWeightBucket.Add(TargetBoneIndex, TransferData.SourceWeights[SourceLocalIdx] * BoneWeight);
 									}
 								}
+								else if (!MissingSourceBones.Contains(SourceBoneIndex))
+								{
+									MissingSourceBones.Add(SourceBoneIndex);
+									UE_LOG(LogChaosDataflow, Warning, TEXT("Target geometry does not contain Source Transform indexed [%d]."), SourceBoneIndex);
+								}
 							}
 						}
 
diff --git a/Engine/Plugins/Experimental/GeometryCollectionPlugin/Source/GeometryCollectionNodes/Public/Dataflow/GeometryCollectionTransferVertexAttributeNode.h b/Engine/Plugins/Experimental/GeometryCollectionPlugin/Source/GeometryCollectionNodes/Public/Dataflow/GeometryCollectionTransferVertexAttributeNode.h
--- a/Engine/Plugins/Experimental/GeometryCollectionPlugin/Source/GeometryCollectionNodes/Public/Dataflow/GeometryCollectionTransferVertexAttributeNode.h
+++ b/Engine/Plugins/Experimental/GeometryCollectionPlugin/Source/GeometryCollectionNodes/Public/Dataflow/GeometryCollectionTransferVertexAttributeNode.h
@@ -144,7 +144,9 @@ public:
 	float BoundMultiplier = 0.01f;
 
 	/* Suffix of transform names added to the source geometry's BoneName for geometry matching during transfer[default: _Tet]. In CreateTetrahedron node we add _Tet to tetrahedral geometries.*/
-	UPROPERTY(EditAnywhere, Category = "Transform Name")
+	UPROPERTY(EditAnywhere, Category = "Paired Geometry Transfer", meta = (
+		EditCondition = "TransferMethod == EDataflowTransferVertexAttributeNodeTransferMethod::Component",
+		EditConditionHides))
 	FString TransformNameSuffix = "_Tet";
 
 	FGeometryCollectionTransferVertexAttributeNode(const UE::Dataflow::FNodeParameters& InParam, FGuid InGuid = FGuid::NewGuid())
@@ -220,7 +222,9 @@ public:
 	float BoundMultiplier = 0.01f;
 
 	/* Suffix of transform names for geometry matching during transfer[default: _Tet]. In CreateTetrahedron node we add _Tet to tetrahedral geometries.*/
-	UPROPERTY(EditAnywhere, Category = "Transform Name")
+	UPROPERTY(EditAnywhere, Category = "Paired Geometry Transfer", meta = (
+		EditCondition = "TransferMethod == EDataflowTransferVertexAttributeNodeTransferMethod::Component",
+		EditConditionHides))
 	FString TransformNameSuffix = "_Tet";
 
 	FGeometryCollectionTransferVertexSkinWeightsNode(const UE::Dataflow::FNodeParameters& InParam, FGuid InGuid = FGuid::NewGuid())
