Dataflow:
- Added DebugDraw to all the convex Dataflow nodes
- Bug fixes

diff --git a/Engine/Plugins/Experimental/GeometryCollectionPlugin/Source/GeometryCollectionNodes/Private/Dataflow/GeometryCollectionUtilityNodes.cpp b/Engine/Plugins/Experimental/GeometryCollectionPlugin/Source/GeometryCollectionNodes/Private/Dataflow/GeometryCollectionUtilityNodes.cpp
--- a/Engine/Plugins/Experimental/GeometryCollectionPlugin/Source/GeometryCollectionNodes/Private/Dataflow/GeometryCollectionUtilityNodes.cpp
+++ b/Engine/Plugins/Experimental/GeometryCollectionPlugin/Source/GeometryCollectionNodes/Private/Dataflow/GeometryCollectionUtilityNodes.cpp
@@ -97,6 +97,26 @@ namespace UE::Dataflow::Convex
 
 		return FLinearColor(FColor(R, G, B, 255));
 	}
+
+	static void DebugDrawProc(IDataflowDebugDrawInterface& DataflowRenderingInterface, const FManagedArrayCollection& InCollection, const bool bRandomizeColor, const int32 ColorRandomSeed)
+	{
+		using namespace UE::Geometry;
+
+		TArray<FDynamicMesh3> HullsMeshes;
+		UE::FractureEngine::Convex::GetConvexHullsAsDynamicMeshes(InCollection, HullsMeshes, false, {});
+
+		int32 Idx = 0;
+		for (const FDynamicMesh3& Mesh : HullsMeshes)
+		{
+			FDynamicMeshDebugDrawMesh DebugdrawMesh(&Mesh);
+			DataflowRenderingInterface.DrawMesh(DebugdrawMesh);
+
+			if (bRandomizeColor)
+			{
+				DataflowRenderingInterface.SetColor(UE::Dataflow::Convex::GetRandomColor(ColorRandomSeed, Idx++));
+			}
+		}
+	}
 }
 
 FCreateLeafConvexHullsDataflowNode::FCreateLeafConvexHullsDataflowNode(const UE::Dataflow::FNodeParameters& InParam, FGuid InGuid)
@@ -179,20 +199,7 @@ void FCreateLeafConvexHullsDataflowNode::DebugDraw(UE::Dataflow::FContext& Conte
 		{
 			const FManagedArrayCollection& InCollection = Output->GetValue(Context, Collection);
 
-			TArray<FDynamicMesh3> HullsMeshes;
-			UE::FractureEngine::Convex::GetConvexHullsAsDynamicMeshes(InCollection, HullsMeshes, false, {});
-
-			int32 Idx = 0;
-			for (const FDynamicMesh3& Mesh : HullsMeshes)
-			{
-				FDynamicMeshDebugDrawMesh DebugdrawMesh(&Mesh);
-				DataflowRenderingInterface.DrawMesh(DebugdrawMesh);
-
-				if (bRandomizeColor)
-				{
-					DataflowRenderingInterface.SetColor(UE::Dataflow::Convex::GetRandomColor(ColorRandomSeed, Idx++));
-				}
-			}
+			UE::Dataflow::Convex::DebugDrawProc(DataflowRenderingInterface, InCollection, bRandomizeColor, ColorRandomSeed);
 		}
 	}
 }
@@ -204,11 +211,15 @@ FSimplifyConvexHullsDataflowNode::FSimplifyConvexHullsDataflowNode(const UE::Dat
 	: FDataflowNode(InParam, InGuid)
 {
 	RegisterInputConnection(&Collection);
-	RegisterInputConnection(&OptionalSelectionFilter);
-	RegisterInputConnection(&SimplificationAngleThreshold);
-	RegisterInputConnection(&SimplificationDistanceThreshold);
-	RegisterInputConnection(&MinTargetTriangleCount);
+	RegisterInputConnection(&OptionalSelectionFilter).SetCanHidePin(true).SetPinIsHidden(true);
+	RegisterInputConnection(&SimplificationAngleThreshold).SetCanHidePin(true).SetPinIsHidden(true);
+	RegisterInputConnection(&SimplificationDistanceThreshold).SetCanHidePin(true).SetPinIsHidden(true);
+	RegisterInputConnection(&MinTargetTriangleCount).SetCanHidePin(true).SetPinIsHidden(true);
 	RegisterOutputConnection(&Collection);
+
+	DebugDrawRenderSettings.RenderType = EDataflowDebugDrawRenderType::Wireframe;
+	DebugDrawRenderSettings.Color = FLinearColor::Green;
+	DebugDrawRenderSettings.LineWidthMultiplier = 5.0;
 }
 
 void FSimplifyConvexHullsDataflowNode::Evaluate(UE::Dataflow::FContext& Context, const FDataflowOutput* Out) const
@@ -244,15 +255,45 @@ void FSimplifyConvexHullsDataflowNode::Evaluate(UE::Dataflow::FContext& Context,
 	}
 }
 
+#if WITH_EDITOR
+bool FSimplifyConvexHullsDataflowNode::CanDebugDrawViewMode(const FName& ViewModeName) const
+{
+	return ViewModeName == UE::Dataflow::FDataflowConstruction3DViewMode::Name;
+}
+
+void FSimplifyConvexHullsDataflowNode::DebugDraw(UE::Dataflow::FContext& Context, IDataflowDebugDrawInterface& DataflowRenderingInterface, const FDebugDrawParameters& DebugDrawParameters) const
+{
+	using namespace UE::Geometry;
+
+	if ((DebugDrawParameters.bNodeIsSelected || DebugDrawParameters.bNodeIsPinned))
+	{
+		DebugDrawRenderSettings.SetDebugDrawSettings(DataflowRenderingInterface);
+
+		if (const FDataflowOutput* Output = FindOutput(&Collection))
+		{
+			const FManagedArrayCollection& OutCollection = Output->GetValue(Context, Collection);
+
+			UE::Dataflow::Convex::DebugDrawProc(DataflowRenderingInterface, OutCollection, bRandomizeColor, ColorRandomSeed);
+		}
+	}
+}
+#endif
+
+/* --------------------------------------------------------------------------------------------------------------------------- */
+
 FCreateNonOverlappingConvexHullsDataflowNode::FCreateNonOverlappingConvexHullsDataflowNode(const UE::Dataflow::FNodeParameters& InParam, FGuid InGuid)
 	: FDataflowNode(InParam, InGuid)
 {
 	RegisterInputConnection(&Collection);
-	RegisterInputConnection(&CanRemoveFraction);
-	RegisterInputConnection(&SimplificationDistanceThreshold);
-	RegisterInputConnection(&CanExceedFraction);
-	RegisterInputConnection(&OverlapRemovalShrinkPercent);
+	RegisterInputConnection(&CanRemoveFraction).SetCanHidePin(true).SetPinIsHidden(true);
+	RegisterInputConnection(&SimplificationDistanceThreshold).SetCanHidePin(true).SetPinIsHidden(true);
+	RegisterInputConnection(&CanExceedFraction).SetCanHidePin(true).SetPinIsHidden(true);
+	RegisterInputConnection(&OverlapRemovalShrinkPercent).SetCanHidePin(true).SetPinIsHidden(true);
 	RegisterOutputConnection(&Collection);
+
+	DebugDrawRenderSettings.RenderType = EDataflowDebugDrawRenderType::Wireframe;
+	DebugDrawRenderSettings.Color = FLinearColor::Green;
+	DebugDrawRenderSettings.LineWidthMultiplier = 5.0;
 }
 
 void FCreateNonOverlappingConvexHullsDataflowNode::Evaluate(UE::Dataflow::FContext& Context, const FDataflowOutput* Out) const
@@ -280,6 +321,30 @@ void FCreateNonOverlappingConvexHullsDataflowNode::Evaluate(UE::Dataflow::FConte
 	}
 }
 
+#if WITH_EDITOR
+bool FCreateNonOverlappingConvexHullsDataflowNode::CanDebugDrawViewMode(const FName& ViewModeName) const
+{
+	return ViewModeName == UE::Dataflow::FDataflowConstruction3DViewMode::Name;
+}
+
+void FCreateNonOverlappingConvexHullsDataflowNode::DebugDraw(UE::Dataflow::FContext& Context, IDataflowDebugDrawInterface& DataflowRenderingInterface, const FDebugDrawParameters& DebugDrawParameters) const
+{
+	using namespace UE::Geometry;
+
+	if ((DebugDrawParameters.bNodeIsSelected || DebugDrawParameters.bNodeIsPinned))
+	{
+		DebugDrawRenderSettings.SetDebugDrawSettings(DataflowRenderingInterface);
+
+		if (const FDataflowOutput* Output = FindOutput(&Collection))
+		{
+			const FManagedArrayCollection& InCollection = Output->GetValue(Context, Collection);
+
+			UE::Dataflow::Convex::DebugDrawProc(DataflowRenderingInterface, InCollection, bRandomizeColor, ColorRandomSeed);
+		}
+	}
+}
+#endif
+
 /* --------------------------------------------------------------------------------------------------------------------------- */
 
 // local helper to convert the dataflow enum
@@ -411,20 +476,7 @@ void FGenerateClusterConvexHullsFromLeafHullsDataflowNode::DebugDraw(UE::Dataflo
 		{
 			const FManagedArrayCollection& InCollection = Output->GetValue(Context, Collection);
 
-			TArray<FDynamicMesh3> HullsMeshes;
-			UE::FractureEngine::Convex::GetConvexHullsAsDynamicMeshes(InCollection, HullsMeshes, false, {});
-
-			int32 Idx = 0;
-			for (const FDynamicMesh3& Mesh : HullsMeshes)
-			{
-				FDynamicMeshDebugDrawMesh DebugdrawMesh(&Mesh);
-				DataflowRenderingInterface.DrawMesh(DebugdrawMesh);
-
-				if (bRandomizeColor)
-				{
-					DataflowRenderingInterface.SetColor(UE::Dataflow::Convex::GetRandomColor(ColorRandomSeed, Idx++));
-				}
-			}
+			UE::Dataflow::Convex::DebugDrawProc(DataflowRenderingInterface, InCollection, bRandomizeColor, ColorRandomSeed);
 		}
 	}
 }
@@ -546,20 +598,7 @@ void FGenerateClusterConvexHullsFromChildrenHullsDataflowNode::DebugDraw(UE::Dat
 		{
 			const FManagedArrayCollection& InCollection = Output->GetValue(Context, Collection);
 
-			TArray<FDynamicMesh3> HullsMeshes;
-			UE::FractureEngine::Convex::GetConvexHullsAsDynamicMeshes(InCollection, HullsMeshes, false, {});
-
-			int32 Idx = 0;
-			for (const FDynamicMesh3& Mesh : HullsMeshes)
-			{
-				FDynamicMeshDebugDrawMesh DebugdrawMesh(&Mesh);
-				DataflowRenderingInterface.DrawMesh(DebugdrawMesh);
-
-				if (bRandomizeColor)
-				{
-					DataflowRenderingInterface.SetColor(UE::Dataflow::Convex::GetRandomColor(ColorRandomSeed, Idx++));
-				}
-			}
+			UE::Dataflow::Convex::DebugDrawProc(DataflowRenderingInterface, InCollection, bRandomizeColor, ColorRandomSeed);
 		}
 	}
 }
@@ -571,17 +610,20 @@ FMergeConvexHullsDataflowNode::FMergeConvexHullsDataflowNode(const UE::Dataflow:
 	: FDataflowNode(InParam, InGuid)
 {
 	RegisterInputConnection(&Collection);
-	RegisterInputConnection(&MaxConvexCount);
-	RegisterInputConnection(&ErrorTolerance);
-	RegisterInputConnection(&OptionalSelectionFilter);
-	RegisterInputConnection(&bProtectNegativeSpace);
-	RegisterInputConnection(&TargetNumSamples);
-	RegisterInputConnection(&MinSampleSpacing);
-	RegisterInputConnection(&NegativeSpaceTolerance);
-	RegisterInputConnection(&MinRadius);
-
+	RegisterInputConnection(&MaxConvexCount).SetCanHidePin(true).SetPinIsHidden(true);
+	RegisterInputConnection(&ErrorTolerance).SetCanHidePin(true).SetPinIsHidden(true);
+	RegisterInputConnection(&OptionalSelectionFilter).SetCanHidePin(true).SetPinIsHidden(true);
+	RegisterInputConnection(&bProtectNegativeSpace).SetCanHidePin(true).SetPinIsHidden(true);
+	RegisterInputConnection(&TargetNumSamples).SetCanHidePin(true).SetPinIsHidden(true);
+	RegisterInputConnection(&MinSampleSpacing).SetCanHidePin(true).SetPinIsHidden(true);
+	RegisterInputConnection(&NegativeSpaceTolerance).SetCanHidePin(true).SetPinIsHidden(true);
+	RegisterInputConnection(&MinRadius).SetCanHidePin(true).SetPinIsHidden(true);
 	RegisterOutputConnection(&Collection);
 	RegisterOutputConnection(&SphereCovering);
+
+	DebugDrawRenderSettings.RenderType = EDataflowDebugDrawRenderType::Wireframe;
+	DebugDrawRenderSettings.Color = FLinearColor::Green;
+	DebugDrawRenderSettings.LineWidthMultiplier = 5.0;
 }
 
 void FClearConvexHullsDataflowNode::Evaluate(UE::Dataflow::FContext& Context, const FDataflowOutput* Out) const
@@ -673,6 +715,31 @@ void FMergeConvexHullsDataflowNode::Evaluate(UE::Dataflow::FContext& Context, co
 	}
 }
 
+#if WITH_EDITOR
+bool FMergeConvexHullsDataflowNode::CanDebugDrawViewMode(const FName& ViewModeName) const
+{
+	return ViewModeName == UE::Dataflow::FDataflowConstruction3DViewMode::Name;
+}
+
+void FMergeConvexHullsDataflowNode::DebugDraw(UE::Dataflow::FContext& Context, IDataflowDebugDrawInterface& DataflowRenderingInterface, const FDebugDrawParameters& DebugDrawParameters) const
+{
+	using namespace UE::Geometry;
+
+	if ((DebugDrawParameters.bNodeIsSelected || DebugDrawParameters.bNodeIsPinned))
+	{
+		DebugDrawRenderSettings.SetDebugDrawSettings(DataflowRenderingInterface);
+
+		if (const FDataflowOutput* Output = FindOutput(&Collection))
+		{
+			const FManagedArrayCollection& InCollection = Output->GetValue(Context, Collection);
+
+			UE::Dataflow::Convex::DebugDrawProc(DataflowRenderingInterface, InCollection, bRandomizeColor, ColorRandomSeed);
+		}
+	}
+}
+#endif
+
+/* --------------------------------------------------------------------------------------------------------------------------- */
 
 FUpdateVolumeAttributesDataflowNode::FUpdateVolumeAttributesDataflowNode(const UE::Dataflow::FNodeParameters& InParam, FGuid InGuid)
 	: FDataflowNode(InParam, InGuid)
@@ -694,6 +761,7 @@ void FUpdateVolumeAttributesDataflowNode::Evaluate(UE::Dataflow::FContext& Conte
 	}
 }
 
+/* --------------------------------------------------------------------------------------------------------------------------- */
 
 FGetConvexHullVolumeDataflowNode::FGetConvexHullVolumeDataflowNode(const UE::Dataflow::FNodeParameters& InParam, FGuid InGuid)
 	: FDataflowNode(InParam, InGuid)
@@ -701,6 +769,10 @@ FGetConvexHullVolumeDataflowNode::FGetConvexHullVolumeDataflowNode(const UE::Dat
 	RegisterInputConnection(&Collection);
 	RegisterInputConnection(&TransformSelection);
 	RegisterOutputConnection(&Volume);
+
+	DebugDrawRenderSettings.RenderType = EDataflowDebugDrawRenderType::Wireframe;
+	DebugDrawRenderSettings.Color = FLinearColor::Green;
+	DebugDrawRenderSettings.LineWidthMultiplier = 5.0;
 }
 
 void FGetConvexHullVolumeDataflowNode::Evaluate(UE::Dataflow::FContext& Context, const FDataflowOutput* Out) const
@@ -787,6 +859,32 @@ void FGetConvexHullVolumeDataflowNode::Evaluate(UE::Dataflow::FContext& Context,
 	}
 }
 
+#if WITH_EDITOR
+bool FGetConvexHullVolumeDataflowNode::CanDebugDrawViewMode(const FName& ViewModeName) const
+{
+	return ViewModeName == UE::Dataflow::FDataflowConstruction3DViewMode::Name;
+}
+
+void FGetConvexHullVolumeDataflowNode::DebugDraw(UE::Dataflow::FContext& Context, IDataflowDebugDrawInterface& DataflowRenderingInterface, const FDebugDrawParameters& DebugDrawParameters) const
+{
+	using namespace UE::Geometry;
+
+	if ((DebugDrawParameters.bNodeIsSelected || DebugDrawParameters.bNodeIsPinned))
+	{
+		DebugDrawRenderSettings.SetDebugDrawSettings(DataflowRenderingInterface);
+
+		if (IsConnected(&Collection))
+		{
+			const FManagedArrayCollection& InCollection = GetValue(Context, &Collection);
+
+			UE::Dataflow::Convex::DebugDrawProc(DataflowRenderingInterface, InCollection, bRandomizeColor, ColorRandomSeed);
+		}
+	}
+}
+#endif
+
+/* --------------------------------------------------------------------------------------------------------------------------- */
+
 void FFixTinyGeoDataflowNode::Evaluate(UE::Dataflow::FContext& Context, const FDataflowOutput* Out) const
 {
 	if (Out->IsA(&Collection))
diff --git a/Engine/Plugins/Experimental/GeometryCollectionPlugin/Source/GeometryCollectionNodes/Public/Dataflow/GeometryCollectionUtilityNodes.h b/Engine/Plugins/Experimental/GeometryCollectionPlugin/Source/GeometryCollectionNodes/Public/Dataflow/GeometryCollectionUtilityNodes.h
--- a/Engine/Plugins/Experimental/GeometryCollectionPlugin/Source/GeometryCollectionNodes/Public/Dataflow/GeometryCollectionUtilityNodes.h
+++ b/Engine/Plugins/Experimental/GeometryCollectionPlugin/Source/GeometryCollectionNodes/Public/Dataflow/GeometryCollectionUtilityNodes.h
@@ -211,7 +211,7 @@ struct FSimplifyConvexHullsDataflowNode : public FDataflowNode
 	GENERATED_USTRUCT_BODY()
 	DATAFLOW_NODE_DEFINE_INTERNAL(FSimplifyConvexHullsDataflowNode, "SimplifyConvexHulls", "GeometryCollection|Utilities", "")
 
-public:
+private:
 	UPROPERTY(meta = (DataflowInput, DataflowOutput))
 	FManagedArrayCollection Collection;
 
@@ -238,9 +238,28 @@ public:
 	UPROPERTY(EditAnywhere, Category = "Convex")
 	bool bUseExistingVertices = false;
 
-	FSimplifyConvexHullsDataflowNode(const UE::Dataflow::FNodeParameters& InParam, FGuid InGuid = FGuid::NewGuid());
+	UPROPERTY(EditAnywhere, Category = "Debug Draw", meta = (ShowOnlyInnerProperties))
+	FDataflowNodeDebugDrawSettings DebugDrawRenderSettings;
+
+	/** Randomize color per convex hull */
+	UPROPERTY(EditAnywhere, Category = "Debug Draw")
+	bool bRandomizeColor = true;
+
+	/** Random seed */
+	UPROPERTY(EditAnywhere, Category = "Debug Draw", meta = (ClampMin = "0"))
+	int32 ColorRandomSeed = 0;
 
 	virtual void Evaluate(UE::Dataflow::FContext& Context, const FDataflowOutput* Out) const override;
+#if WITH_EDITOR
+	virtual bool CanDebugDraw() const override
+	{
+		return true;
+	}
+	virtual bool CanDebugDrawViewMode(const FName& ViewModeName) const override;
+	virtual void DebugDraw(UE::Dataflow::FContext& Context, IDataflowDebugDrawInterface& DataflowRenderingInterface, const FDebugDrawParameters& DebugDrawParameters) const override;
+#endif
+public:
+	FSimplifyConvexHullsDataflowNode(const UE::Dataflow::FNodeParameters& InParam, FGuid InGuid = FGuid::NewGuid());
 };
 
 /**
@@ -254,7 +273,7 @@ struct FCreateNonOverlappingConvexHullsDataflowNode : public FDataflowNode
 	GENERATED_USTRUCT_BODY()
 	DATAFLOW_NODE_DEFINE_INTERNAL(FCreateNonOverlappingConvexHullsDataflowNode, "CreateNonOverlappingConvexHulls", "GeometryCollection|Utilities", "")
 
-public:
+private:
 	UPROPERTY(meta = (DataflowInput, DataflowOutput))
 	FManagedArrayCollection Collection;
 
@@ -278,10 +297,28 @@ public:
 	UPROPERTY(EditAnywhere, Category = AutomaticOverlapRemoval, meta = (DataflowInput, DisplayName = "Max Removal Fraction", ClampMin = 0.01f, ClampMax = 1.f))
 	float CanRemoveFraction = 0.3f;
 
-	FCreateNonOverlappingConvexHullsDataflowNode(const UE::Dataflow::FNodeParameters& InParam, FGuid InGuid = FGuid::NewGuid());
+	UPROPERTY(EditAnywhere, Category = "Debug Draw", meta = (ShowOnlyInnerProperties))
+	FDataflowNodeDebugDrawSettings DebugDrawRenderSettings;
 
-	virtual void Evaluate(UE::Dataflow::FContext& Context, const FDataflowOutput* Out) const override;
+	/** Randomize color per convex hull */
+	UPROPERTY(EditAnywhere, Category = "Debug Draw")
+	bool bRandomizeColor = true;
+
+	/** Random seed */
+	UPROPERTY(EditAnywhere, Category = "Debug Draw", meta = (ClampMin = "0"))
+	int32 ColorRandomSeed = 0;
 
+	virtual void Evaluate(UE::Dataflow::FContext& Context, const FDataflowOutput* Out) const override;
+#if WITH_EDITOR
+	virtual bool CanDebugDraw() const override
+	{
+		return true;
+	}
+	virtual bool CanDebugDrawViewMode(const FName& ViewModeName) const override;
+	virtual void DebugDraw(UE::Dataflow::FContext& Context, IDataflowDebugDrawInterface& DataflowRenderingInterface, const FDebugDrawParameters& DebugDrawParameters) const override;
+#endif
+public:
+	FCreateNonOverlappingConvexHullsDataflowNode(const UE::Dataflow::FNodeParameters& InParam, FGuid InGuid = FGuid::NewGuid());
 };
 
 //~ Simple wrapper class to make the sphere covering data possible to pass via dataflow
@@ -571,7 +608,7 @@ struct FMergeConvexHullsDataflowNode : public FDataflowNode
 	GENERATED_USTRUCT_BODY()
 	DATAFLOW_NODE_DEFINE_INTERNAL(FMergeConvexHullsDataflowNode, "MergeConvexHulls", "GeometryCollection|Utilities", "")
 
-public:
+private:
 	UPROPERTY(meta = (DataflowInput, DataflowOutput))
 	FManagedArrayCollection Collection;
 
@@ -639,10 +676,28 @@ public:
 	UPROPERTY(EditAnywhere, Category = NegativeSpace, meta = (DataflowInput, ClampMin = 0, Units = cm, EditCondition = "bProtectNegativeSpace", EditConditionHides))
 	double MinRadius = 10.0;
 
-	FMergeConvexHullsDataflowNode(const UE::Dataflow::FNodeParameters& InParam, FGuid InGuid = FGuid::NewGuid());
+	UPROPERTY(EditAnywhere, Category = "Debug Draw", meta = (ShowOnlyInnerProperties))
+	FDataflowNodeDebugDrawSettings DebugDrawRenderSettings;
 
-	virtual void Evaluate(UE::Dataflow::FContext& Context, const FDataflowOutput* Out) const override;
+	/** Randomize color per convex hull */
+	UPROPERTY(EditAnywhere, Category = "Debug Draw")
+	bool bRandomizeColor = true;
+
+	/** Random seed */
+	UPROPERTY(EditAnywhere, Category = "Debug Draw", meta = (ClampMin = "0"))
+	int32 ColorRandomSeed = 0;
 
+	virtual void Evaluate(UE::Dataflow::FContext& Context, const FDataflowOutput* Out) const override;
+#if WITH_EDITOR
+	virtual bool CanDebugDraw() const override
+	{
+		return true;
+	}
+	virtual bool CanDebugDrawViewMode(const FName& ViewModeName) const override;
+	virtual void DebugDraw(UE::Dataflow::FContext& Context, IDataflowDebugDrawInterface& DataflowRenderingInterface, const FDebugDrawParameters& DebugDrawParameters) const override;
+#endif
+public:
+	FMergeConvexHullsDataflowNode(const UE::Dataflow::FNodeParameters& InParam, FGuid InGuid = FGuid::NewGuid());
 };
 
 
@@ -675,7 +730,7 @@ struct FGetConvexHullVolumeDataflowNode : public FDataflowNode
 	GENERATED_USTRUCT_BODY()
 	DATAFLOW_NODE_DEFINE_INTERNAL(FGetConvexHullVolumeDataflowNode, "GetConvexHullVolume", "GeometryCollection|Utilities", "")
 
-public:
+private:
 	UPROPERTY(meta = (DataflowInput))
 	FManagedArrayCollection Collection;
 
@@ -695,9 +750,28 @@ public:
 	UPROPERTY(EditAnywhere, Category = Options)
 	bool bVolumeOfUnion = false;
 
-	FGetConvexHullVolumeDataflowNode(const UE::Dataflow::FNodeParameters& InParam, FGuid InGuid = FGuid::NewGuid());
+	UPROPERTY(EditAnywhere, Category = "Debug Draw", meta = (ShowOnlyInnerProperties))
+	FDataflowNodeDebugDrawSettings DebugDrawRenderSettings;
+
+	/** Randomize color per convex hull */
+	UPROPERTY(EditAnywhere, Category = "Debug Draw")
+	bool bRandomizeColor = true;
+
+	/** Random seed */
+	UPROPERTY(EditAnywhere, Category = "Debug Draw", meta = (ClampMin = "0"))
+	int32 ColorRandomSeed = 0;
 
 	virtual void Evaluate(UE::Dataflow::FContext& Context, const FDataflowOutput* Out) const override;
+#if WITH_EDITOR
+	virtual bool CanDebugDraw() const override
+	{
+		return true;
+	}
+	virtual bool CanDebugDrawViewMode(const FName& ViewModeName) const override;
+	virtual void DebugDraw(UE::Dataflow::FContext& Context, IDataflowDebugDrawInterface& DataflowRenderingInterface, const FDebugDrawParameters& DebugDrawParameters) const override;
+#endif
+public:
+	FGetConvexHullVolumeDataflowNode(const UE::Dataflow::FNodeParameters& InParam, FGuid InGuid = FGuid::NewGuid());
 };
 
 /**
