Geometry collection : Add specific select internal faces node as this is a very common node to be used

diff --git a/Engine/Plugins/Experimental/GeometryCollectionPlugin/Source/GeometryCollectionNodes/Private/Dataflow/GeometryCollectionSelectionNodes.cpp b/Engine/Plugins/Experimental/GeometryCollectionPlugin/Source/GeometryCollectionNodes/Private/Dataflow/GeometryCollectionSelectionNodes.cpp
--- a/Engine/Plugins/Experimental/GeometryCollectionPlugin/Source/GeometryCollectionNodes/Private/Dataflow/GeometryCollectionSelectionNodes.cpp
+++ b/Engine/Plugins/Experimental/GeometryCollectionPlugin/Source/GeometryCollectionNodes/Private/Dataflow/GeometryCollectionSelectionNodes.cpp
@@ -69,6 +69,7 @@ namespace UE::Dataflow
 		DATAFLOW_NODE_REGISTER_CREATION_FACTORY(FCollectionVertexSelectionSetOperationDataflowNode);
 		DATAFLOW_NODE_REGISTER_CREATION_FACTORY(FCollectionSelectionByAttrDataflowNode);
 		DATAFLOW_NODE_REGISTER_CREATION_FACTORY(FGeometrySelectionToVertexSelectionDataflowNode);
+		DATAFLOW_NODE_REGISTER_CREATION_FACTORY(FCollectionSelectInternalFacesDataflowNode);
 
 		// generic input nodes
 		DATAFLOW_NODE_REGISTER_CREATION_FACTORY(FCollectionSelectionSetOperationDataflowNode);
@@ -1690,4 +1691,85 @@ void FCollectionSelectionInvertDataflowNode::Evaluate(UE::Dataflow::FContext& Co
 	}
 }
 
-//////////////////////////////////////////////////////////////////////////////////////////////////////////////////////
\ No newline at end of file
+//////////////////////////////////////////////////////////////////////////////////////////////////////////////////////
+
+FCollectionSelectInternalFacesDataflowNode::FCollectionSelectInternalFacesDataflowNode(const UE::Dataflow::FNodeParameters& InParam, FGuid InGuid)
+	: FDataflowNode(InParam, InGuid)
+{
+	RegisterInputConnection(&Collection);
+	RegisterInputConnection(&TransformSelection);
+	RegisterOutputConnection(&FaceSelection);
+	RegisterOutputConnection(&Collection, &Collection);
+}
+
+void FCollectionSelectInternalFacesDataflowNode::Evaluate(UE::Dataflow::FContext& Context, const FDataflowOutput* Out) const
+{
+	static const FName InternalAttributeName(TEXT("Internal"));
+
+	if (Out->IsA(&Collection))
+	{
+		SafeForwardInput(Context, &Collection, &Collection);
+	}
+	else if (Out->IsA(&FaceSelection))
+	{
+		const FManagedArrayCollection& InCollection = GetValue(Context, &Collection);
+		const FDataflowTransformSelection& InTransformSelection = GetValue(Context, &TransformSelection);
+
+		struct FFaceRange
+		{
+			int32 Start;
+			int32 Count;
+		};
+
+		const TManagedArray<int32>* TransformIndexFromGeometry = InCollection.FindAttribute<int32>("TransformIndex", FGeometryCollection::GeometryGroup);
+		const TManagedArray<int32>* FaceStart = InCollection.FindAttribute<int32>("FaceStart", FGeometryCollection::GeometryGroup);
+		const TManagedArray<int32>* FaceCount = InCollection.FindAttribute<int32>("FaceCount", FGeometryCollection::GeometryGroup);
+		const TManagedArray<bool>* InternalFaces = InCollection.FindAttribute<bool>(InternalAttributeName, FGeometryCollection::FacesGroup);
+
+		const int32 TotalNumFaces = InCollection.NumElements(FGeometryCollection::FacesGroup);
+
+		FDataflowFaceSelection OutFaceSelection;
+		OutFaceSelection.Initialize(TotalNumFaces, false);
+
+		if (TransformIndexFromGeometry && FaceStart && FaceCount && InternalFaces)
+		{
+			TArray<FFaceRange> FaceRanges;
+			if (IsConnected(&TransformSelection))
+			{
+				for (int32 GeoIdx = 0; GeoIdx < TransformIndexFromGeometry->Num(); ++GeoIdx)
+				{
+					const int32 TransformIndex = (*TransformIndexFromGeometry)[GeoIdx];
+					if (InTransformSelection.IsSelected(TransformIndex))
+					{
+						FaceRanges.Emplace( 
+							FFaceRange { 
+								.Start=(*FaceStart)[GeoIdx],
+								.Count=(*FaceCount)[GeoIdx] 
+							});
+					}
+				}
+			}
+			else
+			{
+				FaceRanges.Emplace(
+					FFaceRange{ 
+						.Start = 0,
+						.Count = TotalNumFaces
+					});
+			}
+
+			for (const FFaceRange& FaceRange : FaceRanges)
+			{
+				for (int32 Idx = 0; Idx < FaceRange.Count; ++Idx)
+				{
+					const int32 FaceIndex = (FaceRange.Start + Idx);
+					if ((*InternalFaces)[FaceIndex])
+					{
+						OutFaceSelection.SetSelected(FaceRange.Start + Idx);
+					}
+				}
+			}
+		}
+		SetValue(Context, OutFaceSelection, &FaceSelection);
+	}
+}
\ No newline at end of file
diff --git a/Engine/Plugins/Experimental/GeometryCollectionPlugin/Source/GeometryCollectionNodes/Public/Dataflow/GeometryCollectionSelectionNodes.h b/Engine/Plugins/Experimental/GeometryCollectionPlugin/Source/GeometryCollectionNodes/Public/Dataflow/GeometryCollectionSelectionNodes.h
--- a/Engine/Plugins/Experimental/GeometryCollectionPlugin/Source/GeometryCollectionNodes/Public/Dataflow/GeometryCollectionSelectionNodes.h
+++ b/Engine/Plugins/Experimental/GeometryCollectionPlugin/Source/GeometryCollectionNodes/Public/Dataflow/GeometryCollectionSelectionNodes.h
@@ -1609,6 +1609,39 @@ public:
 	FCollectionSelectionInvertDataflowNode(const UE::Dataflow::FNodeParameters& InParam, FGuid InGuid = FGuid::NewGuid());
 };
 
+/**
+ *
+ * Select internal faces
+ *
+ */
+USTRUCT()
+struct FCollectionSelectInternalFacesDataflowNode : public FDataflowNode
+{
+	GENERATED_USTRUCT_BODY()
+	DATAFLOW_NODE_DEFINE_INTERNAL(FCollectionSelectInternalFacesDataflowNode, "CollectionSelectInternalFaces", "GeometryCollection|Selection", "")
+
+private:
+	/** Collection to select the internal faces from */
+	UPROPERTY(meta = (DataflowInput, DataflowOutput, DataflowPassthrough = "Collection", DataflowIntrinsic))
+	FManagedArrayCollection Collection;
+
+	/**
+	* Transform selection to get the internal faces from
+	* if this input is not connected, then all internal faces from the collection will be returned
+	*/
+	UPROPERTY(meta = (DataflowInput))
+	FDataflowTransformSelection TransformSelection;
+
+	/** selection containing Internal faces */
+	UPROPERTY(meta = (DataflowOutput))
+	FDataflowFaceSelection FaceSelection;
+
+	virtual void Evaluate(UE::Dataflow::FContext& Context, const FDataflowOutput* Out) const override;
+
+public:
+	FCollectionSelectInternalFacesDataflowNode(const UE::Dataflow::FNodeParameters& InParam, FGuid InGuid = FGuid::NewGuid());
+};
+
 namespace UE::Dataflow
 {
 	void GeometryCollectionSelectionNodes();
