Dataflow : rename MeshCopyToPoints to ScatterMesh and add a array of mesh output

diff --git a/Engine/Plugins/Experimental/GeometryCollectionPlugin/Source/GeometryCollectionNodes/Private/Dataflow/GeometryCollectionMeshNodes.cpp b/Engine/Plugins/Experimental/GeometryCollectionPlugin/Source/GeometryCollectionNodes/Private/Dataflow/GeometryCollectionMeshNodes.cpp
--- a/Engine/Plugins/Experimental/GeometryCollectionPlugin/Source/GeometryCollectionNodes/Private/Dataflow/GeometryCollectionMeshNodes.cpp
+++ b/Engine/Plugins/Experimental/GeometryCollectionPlugin/Source/GeometryCollectionNodes/Private/Dataflow/GeometryCollectionMeshNodes.cpp
@@ -563,23 +563,24 @@ void FDuplicateMeshUVChannelNode::Evaluate(UE::Dataflow::FContext& Context, cons
 
 void FMeshCopyToPointsDataflowNode::Evaluate(UE::Dataflow::FContext& Context, const FDataflowOutput* Out) const
 {
-	if (Out->IsA<TObjectPtr<UDynamicMesh>>(&Mesh))
+	if (Out->IsA(&Mesh))
 	{
-		if (TObjectPtr<UDynamicMesh> InMeshToCopy = GetValue<TObjectPtr<UDynamicMesh>>(Context, &MeshToCopy))
+		if (TObjectPtr<UDynamicMesh> InMeshToCopy = GetValue(Context, &MeshToCopy))
 		{
+			TObjectPtr<UDynamicMesh> NewMesh = nullptr;
+
 			const UE::Geometry::FDynamicMesh3& InDynMeshToCopy = InMeshToCopy->GetMeshRef();
 
-			const TArray<FVector>& InPoints = GetValue<TArray<FVector>>(Context, &Points);
+			const TArray<FVector>& InPoints = GetValue(Context, &Points);
 
 			if (InPoints.Num() > 0 && InDynMeshToCopy.VertexCount() > 0)
 			{
-				TObjectPtr<UDynamicMesh> NewMesh = NewObject<UDynamicMesh>();
+				NewMesh = NewObject<UDynamicMesh>();
 				NewMesh->Reset();
 
-				UE::Geometry::FDynamicMesh3& ResultDynMesh = NewMesh->GetMeshRef();
-				UE::Geometry::FDynamicMeshEditor MeshEditor(&ResultDynMesh);
+				UE::Geometry::FDynamicMeshEditor MeshEditor(&NewMesh->GetMeshRef());
 
-				for (auto& Point : InPoints)
+				for (const FVector& Point : InPoints)
 				{
 					UE::Geometry::FDynamicMesh3 DynMeshTemp(InDynMeshToCopy);
 					UE::Geometry::FRefCountVector VertexRefCounts = DynMeshTemp.GetVerticesRefCounts();
@@ -596,14 +597,53 @@ void FMeshCopyToPointsDataflowNode::Evaluate(UE::Dataflow::FContext& Context, co
 					UE::Geometry::FMeshIndexMappings IndexMaps;
 					MeshEditor.AppendMesh(&DynMeshTemp, IndexMaps);
 				}
+			}
 
-				SetValue(Context, NewMesh, &Mesh);
+			SetValue(Context, NewMesh, &Mesh);
+			return;
+		}
 
-				return;
+		SetValue(Context, TObjectPtr<UDynamicMesh>(NewObject<UDynamicMesh>()), &Mesh);
+	}
+	else if (Out->IsA(&Meshes))
+	{
+		TArray<TObjectPtr<UDynamicMesh>> OutMeshes;
+		if (TObjectPtr<UDynamicMesh> InMeshToCopy = GetValue(Context, &MeshToCopy))
+		{
+			const UE::Geometry::FDynamicMesh3& InDynMeshToCopy = InMeshToCopy->GetMeshRef();
+
+			const TArray<FVector>& InPoints = GetValue(Context, &Points);
+
+			if (InPoints.Num() > 0 && InDynMeshToCopy.VertexCount() > 0)
+			{
+				for (const FVector& Point : InPoints)
+				{
+					TObjectPtr<UDynamicMesh> NewMesh = NewObject<UDynamicMesh>();
+					NewMesh->Reset();
+					OutMeshes.Add(NewMesh);
+
+					UE::Geometry::FDynamicMeshEditor MeshEditor(&NewMesh->GetMeshRef());
+
+					UE::Geometry::FDynamicMesh3 DynMeshTemp(InDynMeshToCopy);
+					UE::Geometry::FRefCountVector VertexRefCounts = DynMeshTemp.GetVerticesRefCounts();
+
+					UE::Geometry::FRefCountVector::IndexIterator ItVertexID = VertexRefCounts.BeginIndices();
+					const UE::Geometry::FRefCountVector::IndexIterator ItEndVertexID = VertexRefCounts.EndIndices();
+
+					while (ItVertexID != ItEndVertexID)
+					{
+						DynMeshTemp.SetVertex(*ItVertexID, Scale * DynMeshTemp.GetVertex(*ItVertexID) + Point);
+						++ItVertexID;
+					}
+
+					UE::Geometry::FMeshIndexMappings IndexMaps;
+					MeshEditor.AppendMesh(&DynMeshTemp, IndexMaps);				
+				}
 			}
 		}
 
-		SetValue(Context, TObjectPtr<UDynamicMesh>(NewObject<UDynamicMesh>()), &Mesh);
+		SetValue(Context, MoveTemp(OutMeshes), &Meshes);
+		return;
 	}
 }
 
diff --git a/Engine/Plugins/Experimental/GeometryCollectionPlugin/Source/GeometryCollectionNodes/Public/Dataflow/GeometryCollectionMeshNodes.h b/Engine/Plugins/Experimental/GeometryCollectionPlugin/Source/GeometryCollectionNodes/Public/Dataflow/GeometryCollectionMeshNodes.h
--- a/Engine/Plugins/Experimental/GeometryCollectionPlugin/Source/GeometryCollectionNodes/Public/Dataflow/GeometryCollectionMeshNodes.h
+++ b/Engine/Plugins/Experimental/GeometryCollectionPlugin/Source/GeometryCollectionNodes/Public/Dataflow/GeometryCollectionMeshNodes.h
@@ -639,7 +639,7 @@ USTRUCT(meta = (DataflowGeometryCollection))
 struct FMeshCopyToPointsDataflowNode : public FDataflowNode
 {
 	GENERATED_USTRUCT_BODY()
-	DATAFLOW_NODE_DEFINE_INTERNAL(FMeshCopyToPointsDataflowNode, "MeshCopyToPoints", "Mesh|Utilities", "")
+	DATAFLOW_NODE_DEFINE_INTERNAL(FMeshCopyToPointsDataflowNode, "ScatterMeshes", "Mesh|Utilities", "")
 	DATAFLOW_NODE_RENDER_TYPE("SurfaceRender",FName("FDynamicMesh3"), "Mesh")
 
 public:
@@ -648,23 +648,28 @@ public:
 	TArray<FVector> Points;
 
 	/** Mesh to copy onto points */
-	UPROPERTY(meta = (DataflowInput, DataflowIntrinsic))
+	UPROPERTY(meta = (DataflowInput, DataflowIntrinsic, DisplayName = "MeshToScatter"))
 	TObjectPtr<UDynamicMesh> MeshToCopy;
 
-	/** Scale appied to the mesh */
+	/** Scale applied to the mesh */
 	UPROPERTY(EditAnywhere, Category = "Copy");
 	float Scale = 1.f;
 
-	/** Copied meshes */
+	/** merged result mesh */
 	UPROPERTY(meta = (DataflowOutput))
 	TObjectPtr<UDynamicMesh> Mesh;
 
+	/** Result meshes as individual ones */
+	UPROPERTY(meta = (DataflowOutput))
+	TArray<TObjectPtr<UDynamicMesh>> Meshes;
+
 	FMeshCopyToPointsDataflowNode(const UE::Dataflow::FNodeParameters& InParam, FGuid InGuid = FGuid::NewGuid())
 		: FDataflowNode(InParam, InGuid)
 	{
 		RegisterInputConnection(&Points);
 		RegisterInputConnection(&MeshToCopy);
 		RegisterOutputConnection(&Mesh);
+		RegisterOutputConnection(&Meshes);
 	}
 
 	virtual void Evaluate(UE::Dataflow::FContext& Context, const FDataflowOutput* Out) const override;
