Dataflow:
- OptionalSelectionFilter inputs now not hidden by default

diff --git a/Engine/Plugins/Experimental/GeometryCollectionPlugin/Source/GeometryCollectionNodes/Private/Dataflow/GeometryCollectionUtilityNodes.cpp b/Engine/Plugins/Experimental/GeometryCollectionPlugin/Source/GeometryCollectionNodes/Private/Dataflow/GeometryCollectionUtilityNodes.cpp
--- a/Engine/Plugins/Experimental/GeometryCollectionPlugin/Source/GeometryCollectionNodes/Private/Dataflow/GeometryCollectionUtilityNodes.cpp
+++ b/Engine/Plugins/Experimental/GeometryCollectionPlugin/Source/GeometryCollectionNodes/Private/Dataflow/GeometryCollectionUtilityNodes.cpp
@@ -123,7 +123,7 @@ FCreateLeafConvexHullsDataflowNode::FCreateLeafConvexHullsDataflowNode(const UE:
 	: FDataflowNode(InParam, InGuid)
 {
 	RegisterInputConnection(&Collection);
-	RegisterInputConnection(&OptionalSelectionFilter).SetCanHidePin(true).SetPinIsHidden(true);
+	RegisterInputConnection(&OptionalSelectionFilter);
 	RegisterInputConnection(&SimplificationDistanceThreshold).SetCanHidePin(true).SetPinIsHidden(true);
 	RegisterInputConnection(&ConvexDecompositionSettings).SetCanHidePin(true).SetPinIsHidden(true);
 	RegisterOutputConnection(&Collection);
@@ -211,7 +211,7 @@ FSimplifyConvexHullsDataflowNode::FSimplifyConvexHullsDataflowNode(const UE::Dat
 	: FDataflowNode(InParam, InGuid)
 {
 	RegisterInputConnection(&Collection);
-	RegisterInputConnection(&OptionalSelectionFilter).SetCanHidePin(true).SetPinIsHidden(true);
+	RegisterInputConnection(&OptionalSelectionFilter);
 	RegisterInputConnection(&SimplificationAngleThreshold).SetCanHidePin(true).SetPinIsHidden(true);
 	RegisterInputConnection(&SimplificationDistanceThreshold).SetCanHidePin(true).SetPinIsHidden(true);
 	RegisterInputConnection(&MinTargetTriangleCount).SetCanHidePin(true).SetPinIsHidden(true);
@@ -368,7 +368,7 @@ FGenerateClusterConvexHullsFromLeafHullsDataflowNode::FGenerateClusterConvexHull
 	RegisterInputConnection(&Collection);
 	RegisterInputConnection(&ConvexCount).SetCanHidePin(true).SetPinIsHidden(true);
 	RegisterInputConnection(&ErrorTolerance).SetCanHidePin(true).SetPinIsHidden(true);
-	RegisterInputConnection(&OptionalSelectionFilter).SetCanHidePin(true).SetPinIsHidden(true);
+	RegisterInputConnection(&OptionalSelectionFilter);
 	RegisterInputConnection(&bAllowMergingLeafHulls).SetCanHidePin(true).SetPinIsHidden(true);
 	RegisterInputConnection(&bProtectNegativeSpace).SetCanHidePin(true).SetPinIsHidden(true);
 	RegisterInputConnection(&TargetNumSamples).SetCanHidePin(true).SetPinIsHidden(true);
@@ -490,7 +490,7 @@ FGenerateClusterConvexHullsFromChildrenHullsDataflowNode::FGenerateClusterConvex
 	RegisterInputConnection(&Collection);
 	RegisterInputConnection(&ConvexCount).SetCanHidePin(true).SetPinIsHidden(true);
 	RegisterInputConnection(&ErrorTolerance).SetCanHidePin(true).SetPinIsHidden(true);
-	RegisterInputConnection(&OptionalSelectionFilter).SetCanHidePin(true).SetPinIsHidden(true);
+	RegisterInputConnection(&OptionalSelectionFilter);
 	RegisterInputConnection(&bAllowMergingLeafHulls).SetCanHidePin(true).SetPinIsHidden(true);
 	RegisterInputConnection(&bProtectNegativeSpace).SetCanHidePin(true).SetPinIsHidden(true);
 	RegisterInputConnection(&TargetNumSamples).SetCanHidePin(true).SetPinIsHidden(true);
@@ -612,7 +612,7 @@ FMergeConvexHullsDataflowNode::FMergeConvexHullsDataflowNode(const UE::Dataflow:
 	RegisterInputConnection(&Collection);
 	RegisterInputConnection(&MaxConvexCount).SetCanHidePin(true).SetPinIsHidden(true);
 	RegisterInputConnection(&ErrorTolerance).SetCanHidePin(true).SetPinIsHidden(true);
-	RegisterInputConnection(&OptionalSelectionFilter).SetCanHidePin(true).SetPinIsHidden(true);
+	RegisterInputConnection(&OptionalSelectionFilter);
 	RegisterInputConnection(&bProtectNegativeSpace).SetCanHidePin(true).SetPinIsHidden(true);
 	RegisterInputConnection(&TargetNumSamples).SetCanHidePin(true).SetPinIsHidden(true);
 	RegisterInputConnection(&MinSampleSpacing).SetCanHidePin(true).SetPinIsHidden(true);
