Dataflow : Deprecate specific nodes that overlap the get array element and size generic nodes

diff --git a/Engine/Plugins/Experimental/GeometryCollectionPlugin/Source/GeometryCollectionNodes/Private/Dataflow/GeometryCollectionArrayNodes.cpp b/Engine/Plugins/Experimental/GeometryCollectionPlugin/Source/GeometryCollectionNodes/Private/Dataflow/GeometryCollectionArrayNodes.cpp
--- a/Engine/Plugins/Experimental/GeometryCollectionPlugin/Source/GeometryCollectionNodes/Private/Dataflow/GeometryCollectionArrayNodes.cpp
+++ b/Engine/Plugins/Experimental/GeometryCollectionPlugin/Source/GeometryCollectionNodes/Private/Dataflow/GeometryCollectionArrayNodes.cpp
@@ -11,10 +11,7 @@ namespace UE::Dataflow
 
 	void GeometryCollectionArrayNodes()
 	{
-		DATAFLOW_NODE_REGISTER_CREATION_FACTORY(FGetFloatArrayElementDataflowNode);
 		DATAFLOW_NODE_REGISTER_CREATION_FACTORY(FFloatArrayToIntArrayDataflowNode);
-		DATAFLOW_NODE_REGISTER_CREATION_FACTORY(FGetArrayElementDataflowNode);
-		DATAFLOW_NODE_REGISTER_CREATION_FACTORY(FGetNumArrayElementsDataflowNode);
 		DATAFLOW_NODE_REGISTER_CREATION_FACTORY(FBoolArrayToFaceSelectionDataflowNode);
 		DATAFLOW_NODE_REGISTER_CREATION_FACTORY(FFloatArrayToVertexSelectionDataflowNode);
 		DATAFLOW_NODE_REGISTER_CREATION_FACTORY(FFloatArrayNormalizeDataflowNode);
@@ -25,7 +22,11 @@ namespace UE::Dataflow
 		DATAFLOW_NODE_REGISTER_CREATION_FACTORY(FRandomizeFloatArrayDataflowNode);
 		DATAFLOW_NODE_REGISTER_CREATION_FACTORY(FDataflowGetArraySizeNode);
 		DATAFLOW_NODE_REGISTER_CREATION_FACTORY(FDataflowGetArrayElementNode);
-		
+
+		// Deprecated nodes
+		DATAFLOW_NODE_REGISTER_CREATION_FACTORY(FGetFloatArrayElementDataflowNode);
+		DATAFLOW_NODE_REGISTER_CREATION_FACTORY(FGetArrayElementDataflowNode);
+		DATAFLOW_NODE_REGISTER_CREATION_FACTORY(FGetNumArrayElementsDataflowNode);
 	}
 }
 
diff --git a/Engine/Plugins/Experimental/GeometryCollectionPlugin/Source/GeometryCollectionNodes/Private/Dataflow/GeometryCollectionMaterialInterfaceNodes.cpp b/Engine/Plugins/Experimental/GeometryCollectionPlugin/Source/GeometryCollectionNodes/Private/Dataflow/GeometryCollectionMaterialInterfaceNodes.cpp
--- a/Engine/Plugins/Experimental/GeometryCollectionPlugin/Source/GeometryCollectionNodes/Private/Dataflow/GeometryCollectionMaterialInterfaceNodes.cpp
+++ b/Engine/Plugins/Experimental/GeometryCollectionPlugin/Source/GeometryCollectionNodes/Private/Dataflow/GeometryCollectionMaterialInterfaceNodes.cpp
@@ -17,13 +17,15 @@ namespace UE::Dataflow
 	{
 		DATAFLOW_NODE_REGISTER_CREATION_FACTORY(FMakeMaterialInterfaceArrayDataflowNode);
 		DATAFLOW_NODE_REGISTER_CREATION_FACTORY(FGetMaterialInterfaceAssetDataflowNode);
-		DATAFLOW_NODE_REGISTER_CREATION_FACTORY(FGetMaterialInterfaceArraySizeDataflowNode);
-		DATAFLOW_NODE_REGISTER_CREATION_FACTORY(FGetFromMaterialInterfaceArrayDataflowNode);
 		DATAFLOW_NODE_REGISTER_CREATION_FACTORY(FSetIntoMaterialInterfaceArrayDataflowNode);
 		DATAFLOW_NODE_REGISTER_CREATION_FACTORY(FAddToMaterialInterfaceArrayDataflowNode);
 		DATAFLOW_NODE_REGISTER_CREATION_FACTORY(FAssignMaterialInterfaceToCollectionDataflowNode);
 		DATAFLOW_NODE_REGISTER_CREATION_FACTORY(FMaterialInterfaceTextureOverrideDataflowNode);
 
+		// Deprecated nodes
+		DATAFLOW_NODE_REGISTER_CREATION_FACTORY(FGetMaterialInterfaceArraySizeDataflowNode);
+		DATAFLOW_NODE_REGISTER_CREATION_FACTORY(FGetFromMaterialInterfaceArrayDataflowNode);
+
 		DATAFLOW_NODE_REGISTER_GETTER_FOR_ASSET(UMaterialInterface, FGetMaterialInterfaceAssetDataflowNode);
 	}
 }
diff --git a/Engine/Plugins/Experimental/GeometryCollectionPlugin/Source/GeometryCollectionNodes/Public/Dataflow/GeometryCollectionArrayNodes.h b/Engine/Plugins/Experimental/GeometryCollectionPlugin/Source/GeometryCollectionNodes/Public/Dataflow/GeometryCollectionArrayNodes.h
--- a/Engine/Plugins/Experimental/GeometryCollectionPlugin/Source/GeometryCollectionNodes/Public/Dataflow/GeometryCollectionArrayNodes.h
+++ b/Engine/Plugins/Experimental/GeometryCollectionPlugin/Source/GeometryCollectionNodes/Public/Dataflow/GeometryCollectionArrayNodes.h
@@ -21,9 +21,10 @@ class FGeometryCollection;
 /**
  *
  * Returns the specified element from an array
+ * DEPRECATED 5.6 - use the generic GetArrayElement node instead
  *
  */
-USTRUCT()
+USTRUCT(meta = (Deprecated = 5.6))
 struct FGetFloatArrayElementDataflowNode : public FDataflowNode
 {
 	GENERATED_USTRUCT_BODY()
@@ -115,9 +116,10 @@ public:
 /**
  *
  * Returns the specified element from an array
+ * DEPRECATED 5.6 - use the generic GetArrayElement node instead
  *
  */
-USTRUCT()
+USTRUCT(meta = (Deprecated = 5.6))
 struct FGetArrayElementDataflowNode : public FDataflowNode
 {
 	GENERATED_USTRUCT_BODY()
@@ -152,9 +154,10 @@ public:
 /**
  *
  * Returns the number of elements in an array
- *
+ * DEPRECATED 5.6 - use the generic GetArraySize node instead
+ * 
  */
-USTRUCT()
+USTRUCT(meta = (Deprecated = 5.6))
 struct FGetNumArrayElementsDataflowNode : public FDataflowNode
 {
 	GENERATED_USTRUCT_BODY()
@@ -534,7 +537,7 @@ public:
 };
 
 /**
- * Ge size of an array
+ * Get size of an array
  */
 USTRUCT()
 struct FDataflowGetArraySizeNode : public FDataflowNode
diff --git a/Engine/Plugins/Experimental/GeometryCollectionPlugin/Source/GeometryCollectionNodes/Public/Dataflow/GeometryCollectionMaterialInterfaceNodes.h b/Engine/Plugins/Experimental/GeometryCollectionPlugin/Source/GeometryCollectionNodes/Public/Dataflow/GeometryCollectionMaterialInterfaceNodes.h
--- a/Engine/Plugins/Experimental/GeometryCollectionPlugin/Source/GeometryCollectionNodes/Public/Dataflow/GeometryCollectionMaterialInterfaceNodes.h
+++ b/Engine/Plugins/Experimental/GeometryCollectionPlugin/Source/GeometryCollectionNodes/Public/Dataflow/GeometryCollectionMaterialInterfaceNodes.h
@@ -33,8 +33,9 @@ public:
 
 /**
 * Get number of element in an material array 
+* DEPRECATED 5.6 - use the generic GetArraySize node instead
 */
-USTRUCT()
+USTRUCT(meta = (Deprecated = 5.6))
 struct FGetMaterialInterfaceArraySizeDataflowNode : public FDataflowNode
 {
 	GENERATED_USTRUCT_BODY()
@@ -81,8 +82,9 @@ public:
 /**
  * Get an element from a material array
  * (if the index does not match the range of the array, null is returned)
+ * DEPRECATED 5.6 - use the generic GetArrayElement node instead
  */
-USTRUCT()
+USTRUCT(meta = (Deprecated = 5.6))
 struct FGetFromMaterialInterfaceArrayDataflowNode : public FDataflowNode
 {
 	GENERATED_USTRUCT_BODY()
