Dataflow : Add transform and append points nodes

diff --git a/Engine/Plugins/Experimental/GeometryCollectionPlugin/Source/GeometryCollectionNodes/Private/Dataflow/GeometryCollectionFracturingNodes.cpp b/Engine/Plugins/Experimental/GeometryCollectionPlugin/Source/GeometryCollectionNodes/Private/Dataflow/GeometryCollectionFracturingNodes.cpp
--- a/Engine/Plugins/Experimental/GeometryCollectionPlugin/Source/GeometryCollectionNodes/Private/Dataflow/GeometryCollectionFracturingNodes.cpp
+++ b/Engine/Plugins/Experimental/GeometryCollectionPlugin/Source/GeometryCollectionNodes/Private/Dataflow/GeometryCollectionFracturingNodes.cpp
@@ -63,6 +63,9 @@ namespace UE::Dataflow
 		DATAFLOW_NODE_REGISTER_CREATION_FACTORY(FMeshCutterDataflowNode);
 		DATAFLOW_NODE_REGISTER_CREATION_FACTORY(FUniformFractureDataflowNode);
 		DATAFLOW_NODE_REGISTER_CREATION_FACTORY(FVisualizeFractureDataflowNode);
+		DATAFLOW_NODE_REGISTER_CREATION_FACTORY(FTransformPointsDataflowNode)
+		DATAFLOW_NODE_REGISTER_CREATION_FACTORY(FAppendPointsDataflowNode)
+
 //		Commented out until we decide how to make generic data setter nodes
 //		DATAFLOW_NODE_REGISTER_CREATION_FACTORY(FSetFloatAttributeDataflowNode);
 	}
@@ -326,6 +329,54 @@ void FGridScatterPointsDataflowNode::Evaluate(UE::Dataflow::FContext& Context, c
 	}
 }
 
+////////////////////////////////////////////////////////////////////////////////////////////
+
+FTransformPointsDataflowNode::FTransformPointsDataflowNode(const UE::Dataflow::FNodeParameters& InParam, FGuid InGuid)
+	: FDataflowNode(InParam, InGuid)
+{
+	RegisterInputConnection(&Points);
+	RegisterInputConnection(&Transform);
+	RegisterOutputConnection(&Points, &Points);
+}
+
+void FTransformPointsDataflowNode::Evaluate(UE::Dataflow::FContext& Context, const FDataflowOutput* Out) const
+{
+	if (Out->IsA(&Points))
+	{
+		const FTransform& InTransform = GetValue(Context, &Transform);
+		TArray<FVector> OutPoints = GetValue(Context, &Points);
+		for (FVector& Point : OutPoints)
+		{
+			Point = InTransform.TransformPosition(Point);
+		}
+		SetValue(Context, MoveTemp(OutPoints), &Points);
+	}
+}
+
+////////////////////////////////////////////////////////////////////////////////////////////
+
+FAppendPointsDataflowNode::FAppendPointsDataflowNode(const UE::Dataflow::FNodeParameters& InParam, FGuid InGuid)
+	: FDataflowNode(InParam, InGuid)
+{
+	RegisterInputConnection(&PointsA);
+	RegisterInputConnection(&PointsB);
+	RegisterOutputConnection(&Points, &PointsA);
+}
+
+void FAppendPointsDataflowNode::Evaluate(UE::Dataflow::FContext& Context, const FDataflowOutput* Out) const
+{
+	if (Out->IsA(&Points))
+	{
+		const TArray<FVector>& InPointsA = GetValue(Context, &PointsA);
+		const TArray<FVector>& InPointsB = GetValue(Context, &PointsB);
+		TArray<FVector> OutPoints;
+		OutPoints.Append(InPointsA);
+		OutPoints.Append(InPointsB);
+		SetValue(Context, MoveTemp(OutPoints), &Points);
+	}
+}
+
+////////////////////////////////////////////////////////////////////////////////////////////
 void FVoronoiFractureDataflowNode::Evaluate(UE::Dataflow::FContext& Context, const FDataflowOutput* Out) const
 {
 	if (Out->IsA<FManagedArrayCollection>(&Collection))
diff --git a/Engine/Plugins/Experimental/GeometryCollectionPlugin/Source/GeometryCollectionNodes/Public/Dataflow/GeometryCollectionFracturingNodes.h b/Engine/Plugins/Experimental/GeometryCollectionPlugin/Source/GeometryCollectionNodes/Public/Dataflow/GeometryCollectionFracturingNodes.h
--- a/Engine/Plugins/Experimental/GeometryCollectionPlugin/Source/GeometryCollectionNodes/Public/Dataflow/GeometryCollectionFracturingNodes.h
+++ b/Engine/Plugins/Experimental/GeometryCollectionPlugin/Source/GeometryCollectionNodes/Public/Dataflow/GeometryCollectionFracturingNodes.h
@@ -308,6 +308,55 @@ public:
 	virtual void Evaluate(UE::Dataflow::FContext& Context, const FDataflowOutput* Out) const override;
 };
 
+/**
+* Transform an array of points
+*/
+USTRUCT(meta = (DataflowGeometryCollection))
+struct FTransformPointsDataflowNode : public FDataflowNode
+{
+	GENERATED_USTRUCT_BODY()
+	DATAFLOW_NODE_DEFINE_INTERNAL(FTransformPointsDataflowNode, "TransformPoints", "Generators|Point", "")
+	DATAFLOW_NODE_RENDER_TYPE("PointsRender", FName("TArray<FVector>"), "Points")
+
+private:
+	UPROPERTY(meta = (DataflowOutput, DataflowInput, DataflowPassthrough = Points, DataflowIntrinsic))
+	TArray<FVector> Points;
+
+	UPROPERTY(EditAnywhere, Category = "Transform", meta = (DataflowInput));
+	FTransform Transform = FTransform::Identity;
+
+	virtual void Evaluate(UE::Dataflow::FContext& Context, const FDataflowOutput* Out) const override;
+
+public:
+	FTransformPointsDataflowNode(const UE::Dataflow::FNodeParameters& InParam, FGuid InGuid = FGuid::NewGuid());
+};
+
+/** 
+* Combine two arrays of points into one
+*/
+USTRUCT(meta = (DataflowGeometryCollection))
+struct FAppendPointsDataflowNode : public FDataflowNode
+{
+	GENERATED_USTRUCT_BODY()
+	DATAFLOW_NODE_DEFINE_INTERNAL(FAppendPointsDataflowNode, "AppendPoints", "Generators|Point", "")
+	DATAFLOW_NODE_RENDER_TYPE("PointsRender", FName("TArray<FVector>"), "Points")
+
+private:
+	UPROPERTY(meta = (DataflowInput, DataflowIntrinsic))
+	TArray<FVector> PointsA;
+
+	UPROPERTY(meta = (DataflowInput, DataflowIntrinsic))
+	TArray<FVector> PointsB;
+
+	UPROPERTY(meta = (DataflowOutput, DataflowPassthrough = PointsA))
+	TArray<FVector> Points;
+
+	virtual void Evaluate(UE::Dataflow::FContext& Context, const FDataflowOutput* Out) const override;
+
+public:
+	FAppendPointsDataflowNode(const UE::Dataflow::FNodeParameters& InParam, FGuid InGuid = FGuid::NewGuid());
+};
+
 /**
  *
  * Voronoi fracture
