New general purpose dataflow nodes to support Cloth Mesh Resizing:

* FDataflowForceDependencyNode: this creates a pull dependency between two inputs. Useful in ForEachLoops to ensure intermediate outputs are cached.
* FDataflowMakeManagedArrayCollectionArrayNode: Create/modify a TArray<FManagedArrayCollection> by appending an element.

diff --git a/Engine/Plugins/Experimental/GeometryCollectionPlugin/Source/GeometryCollectionNodes/Private/Dataflow/GeometryCollectionArrayNodes.cpp b/Engine/Plugins/Experimental/GeometryCollectionPlugin/Source/GeometryCollectionNodes/Private/Dataflow/GeometryCollectionArrayNodes.cpp
--- a/Engine/Plugins/Experimental/GeometryCollectionPlugin/Source/GeometryCollectionNodes/Private/Dataflow/GeometryCollectionArrayNodes.cpp
+++ b/Engine/Plugins/Experimental/GeometryCollectionPlugin/Source/GeometryCollectionNodes/Private/Dataflow/GeometryCollectionArrayNodes.cpp
@@ -22,6 +22,8 @@ namespace UE::Dataflow
 		DATAFLOW_NODE_REGISTER_CREATION_FACTORY(FRandomizeFloatArrayDataflowNode);
 		DATAFLOW_NODE_REGISTER_CREATION_FACTORY(FDataflowGetArraySizeNode);
 		DATAFLOW_NODE_REGISTER_CREATION_FACTORY(FDataflowGetArrayElementNode);
+		DATAFLOW_NODE_REGISTER_CREATION_FACTORY(FDataflowMakeManagedArrayCollectionArrayNode);
+		
 
 		// Deprecated nodes
 		DATAFLOW_NODE_REGISTER_CREATION_FACTORY(FGetFloatArrayElementDataflowNode);
@@ -673,4 +675,32 @@ void FDataflowGetArrayElementNode::Evaluate(UE::Dataflow::FContext& Context, con
 		const int32 InIndex = GetValue(Context, &Index);
 		SetArrayElementFromInput(Context, &Array, InIndex, &Element);
 	}
+}
+
+///////////////////////////////////////////////////////////////////////////////////////////////
+
+FDataflowMakeManagedArrayCollectionArrayNode::FDataflowMakeManagedArrayCollectionArrayNode(const UE::Dataflow::FNodeParameters& InParam, FGuid InGuid)
+	: FDataflowNode(InParam, InGuid)
+{
+	static FName ArrayElementDependencyGroup("ArrayElement");
+
+	RegisterInputConnection(&Array);
+	RegisterInputConnection(&Element);
+	RegisterOutputConnection(&Array, &Array);
+}
+
+void FDataflowMakeManagedArrayCollectionArrayNode::Evaluate(UE::Dataflow::FContext& Context, const FDataflowOutput* Out) const
+{
+	if (Out->IsA(&Array))
+	{
+		if (IsConnected(&Element))
+		{
+			TArray<FManagedArrayCollection> OutArray = GetValue(Context, &Array);
+			OutArray.Emplace(GetValue(Context, &Element));
+			SetValue(Context, MoveTemp(OutArray), &Array);
+			return;
+		}
+
+		SafeForwardInput(Context, &Array, &Array);
+	}
 }
\ No newline at end of file
diff --git a/Engine/Plugins/Experimental/GeometryCollectionPlugin/Source/GeometryCollectionNodes/Public/Dataflow/GeometryCollectionArrayNodes.h b/Engine/Plugins/Experimental/GeometryCollectionPlugin/Source/GeometryCollectionNodes/Public/Dataflow/GeometryCollectionArrayNodes.h
--- a/Engine/Plugins/Experimental/GeometryCollectionPlugin/Source/GeometryCollectionNodes/Public/Dataflow/GeometryCollectionArrayNodes.h
+++ b/Engine/Plugins/Experimental/GeometryCollectionPlugin/Source/GeometryCollectionNodes/Public/Dataflow/GeometryCollectionArrayNodes.h
@@ -590,6 +590,32 @@ public:
 	FDataflowGetArrayElementNode(const UE::Dataflow::FNodeParameters& InParam, FGuid InGuid = FGuid::NewGuid());
 };
 
+
+
+/**
+ * Append an element to an array of ManagedArrayCollections.
+ */
+USTRUCT()
+struct FDataflowMakeManagedArrayCollectionArrayNode : public FDataflowNode
+{
+	GENERATED_USTRUCT_BODY()
+	DATAFLOW_NODE_DEFINE_INTERNAL(FDataflowMakeManagedArrayCollectionArrayNode, "MakeManagedArrayCollectionArray", "Utilities|Array", "Make Managed Array Collection")
+
+private:
+	/** Array to append to. If no input connection, a new array will be created */
+	UPROPERTY(meta = (DataflowInput, DataflowOutput, DataflowPassthrough = "Array", DataflowIntrinsic))
+	TArray<FManagedArrayCollection> Array;
+
+	/** The element to append.*/
+	UPROPERTY(meta = (DataflowInput))
+	FManagedArrayCollection Element;
+
+	virtual void Evaluate(UE::Dataflow::FContext& Context, const FDataflowOutput* Out) const override;
+
+public:
+	FDataflowMakeManagedArrayCollectionArrayNode(const UE::Dataflow::FNodeParameters& InParam, FGuid InGuid = FGuid::NewGuid());
+};
+
 namespace UE::Dataflow
 {
 	void GeometryCollectionArrayNodes();
diff --git a/Engine/Source/Runtime/Experimental/Dataflow/Core/Private/Dataflow/DataflowCoreNodes.cpp b/Engine/Source/Runtime/Experimental/Dataflow/Core/Private/Dataflow/DataflowCoreNodes.cpp
--- a/Engine/Source/Runtime/Experimental/Dataflow/Core/Private/Dataflow/DataflowCoreNodes.cpp
+++ b/Engine/Source/Runtime/Experimental/Dataflow/Core/Private/Dataflow/DataflowCoreNodes.cpp
@@ -12,6 +12,7 @@ namespace UE::Dataflow
 		DATAFLOW_NODE_REGISTER_CREATION_FACTORY(FDataflowBranchNode);
 		DATAFLOW_NODE_REGISTER_CREATION_FACTORY(FDataflowSelectNode);
 		DATAFLOW_NODE_REGISTER_CREATION_FACTORY(FDataflowPrintNode);
+		DATAFLOW_NODE_REGISTER_CREATION_FACTORY(FDataflowForceDependencyNode);
 	}
 }
 
@@ -198,4 +199,21 @@ void FDataflowPrintNode::Evaluate(UE::Dataflow::FContext& Context, const FDatafl
 {
 	const FString InValue = "Result=" + GetValue(Context, &Value);
 	Context.Info(InValue, this, nullptr);
+}
+
+FDataflowForceDependencyNode::FDataflowForceDependencyNode(const UE::Dataflow::FNodeParameters& InParam, FGuid InGuid)
+	: FDataflowNode(InParam, InGuid)
+{
+	RegisterInputConnection(&DependentValue);
+	RegisterInputConnection(&Value);
+	RegisterOutputConnection(&Value, &Value);
+}
+
+void FDataflowForceDependencyNode::Evaluate(UE::Dataflow::FContext& Context, const FDataflowOutput* Out) const
+{
+	if (Out->IsA(&Value))
+	{
+		FindInput(&DependentValue)->PullValue(Context);
+		ForwardInput(Context, &Value, &Value);
+	}
 }
\ No newline at end of file
diff --git a/Engine/Source/Runtime/Experimental/Dataflow/Core/Public/Dataflow/DataflowCoreNodes.h b/Engine/Source/Runtime/Experimental/Dataflow/Core/Public/Dataflow/DataflowCoreNodes.h
--- a/Engine/Source/Runtime/Experimental/Dataflow/Core/Public/Dataflow/DataflowCoreNodes.h
+++ b/Engine/Source/Runtime/Experimental/Dataflow/Core/Public/Dataflow/DataflowCoreNodes.h
@@ -111,6 +111,32 @@ public:
 	FDataflowStringConvertibleTypes Value;
 };
 
+/**
+* Force an evaluation dependency between two values
+*/
+USTRUCT()
+struct FDataflowForceDependencyNode : public FDataflowNode
+{
+	GENERATED_USTRUCT_BODY()
+
+	DATAFLOW_NODE_DEFINE_INTERNAL(FDataflowForceDependencyNode, "ForceDependency", "FlowControl", "")
+
+public:
+	FDataflowForceDependencyNode(const UE::Dataflow::FNodeParameters& Param, FGuid InGuid = FGuid::NewGuid());
+	virtual void Evaluate(UE::Dataflow::FContext& Context, const FDataflowOutput* Out) const;
+
+
+public:
+
+	/* Evaluating Value will force an evaluation of DependentValue*/
+	UPROPERTY(meta = (DataflowInput, DataflowOutput, DataflowPassthrough = "Value"))
+	FDataflowAnyType Value;
+
+	/* Evaluating Value will force an evaluation of DependentValue*/
+	UPROPERTY(meta = (DataflowInput))
+	FDataflowAnyType DependentValue;
+};
+
 namespace UE::Dataflow
 {
 	void RegisterCoreNodes();
