Dataflow:
- Fixed FMakeTransformDataflowNode_v2 node issues

diff --git a/Engine/Plugins/Experimental/GeometryCollectionPlugin/Source/GeometryCollectionNodes/Private/Dataflow/GeometryCollectionMakeNodes.cpp b/Engine/Plugins/Experimental/GeometryCollectionPlugin/Source/GeometryCollectionNodes/Private/Dataflow/GeometryCollectionMakeNodes.cpp
--- a/Engine/Plugins/Experimental/GeometryCollectionPlugin/Source/GeometryCollectionNodes/Private/Dataflow/GeometryCollectionMakeNodes.cpp
+++ b/Engine/Plugins/Experimental/GeometryCollectionPlugin/Source/GeometryCollectionNodes/Private/Dataflow/GeometryCollectionMakeNodes.cpp
@@ -67,14 +67,11 @@ namespace UE::Dataflow
 		DATAFLOW_NODE_REGISTER_CREATION_FACTORY(FMakeLiteralIntDataflowNode_v2);
 		DATAFLOW_NODE_REGISTER_CREATION_FACTORY(FMakeLiteralBoolDataflowNode_v2);
 		DATAFLOW_NODE_REGISTER_CREATION_FACTORY(FMakeLiteralVectorDataflowNode);
-		DATAFLOW_NODE_REGISTER_CREATION_FACTORY(FMakeTransformDataflowNode_v2);
 		DATAFLOW_NODE_REGISTER_CREATION_FACTORY(FMakeQuaternionDataflowNode);
 		DATAFLOW_NODE_REGISTER_CREATION_FACTORY(FMakeFloatArrayDataflowNode);
 		DATAFLOW_NODE_REGISTER_CREATION_FACTORY(FMakeCollectionDataflowNode);
 		DATAFLOW_NODE_REGISTER_CREATION_FACTORY(FMakeRotatorDataflowNode);
 		DATAFLOW_NODE_REGISTER_CREATION_FACTORY(FBreakTransformDataflowNode);
-		// Deprecated
-		DATAFLOW_NODE_REGISTER_CREATION_FACTORY(FMakeTransformDataflowNode);
 		DATAFLOW_NODE_REGISTER_CREATION_FACTORY(FMakeLiteralStringDataflowNode);
 		DATAFLOW_NODE_REGISTER_CREATION_FACTORY(FMakeLiteralFloatDataflowNode);
 		DATAFLOW_NODE_REGISTER_CREATION_FACTORY(FMakeLiteralIntDataflowNode);
@@ -88,6 +85,10 @@ namespace UE::Dataflow
 		DATAFLOW_NODE_REGISTER_CREATION_FACTORY(FMakeStairMeshDataflowNode);
 		DATAFLOW_NODE_REGISTER_CREATION_FACTORY(FMakeRectangleMeshDataflowNode);
 		DATAFLOW_NODE_REGISTER_CREATION_FACTORY(FMakeTorusMeshDataflowNode);
+		DATAFLOW_NODE_REGISTER_CREATION_FACTORY(FMakeTransformDataflowNode);
+
+		// Deprecated
+		DATAFLOW_NODE_REGISTER_CREATION_FACTORY(FMakeTransformDataflowNode_v2);
 	}
 }
 
diff --git a/Engine/Plugins/Experimental/GeometryCollectionPlugin/Source/GeometryCollectionNodes/Public/Dataflow/GeometryCollectionMakeNodes.h b/Engine/Plugins/Experimental/GeometryCollectionPlugin/Source/GeometryCollectionNodes/Public/Dataflow/GeometryCollectionMakeNodes.h
--- a/Engine/Plugins/Experimental/GeometryCollectionPlugin/Source/GeometryCollectionNodes/Public/Dataflow/GeometryCollectionMakeNodes.h
+++ b/Engine/Plugins/Experimental/GeometryCollectionPlugin/Source/GeometryCollectionNodes/Public/Dataflow/GeometryCollectionMakeNodes.h
@@ -406,29 +406,35 @@ public:
 
 /**
  * Make an FTransform
- * Deprecated (5.6)
- * Use FMakeTransformDataflowNode_v2 instead
+ * Note: Originaly this version was depricated and replaced with FMakeTransformDataflowNode_v2 but when AnyRotationType was
+ * introduced with the ConvertAnyRotation node FMakeTransformDataflowNode_v2 became obsolete and this version became the current version again
  */
-USTRUCT(meta = (Deprecated = "5.6"))
+USTRUCT()
 struct FMakeTransformDataflowNode : public FDataflowNode
 {
 	GENERATED_USTRUCT_BODY()
 	DATAFLOW_NODE_DEFINE_INTERNAL(FMakeTransformDataflowNode, "MakeTransform", "Generators|Transform", "")
 
-public:
-
+private:
+	/** Translation */
 	UPROPERTY(EditAnywhere, Category = "Transform", meta = (DataflowInput, DisplayName = "Translation"));
 	FVector InTranslation = FVector(0, 0, 0);
 
+	/** Rotation as Euler */
 	UPROPERTY(EditAnywhere, Category = "Transform", meta = (DataflowInput, DisplayName = "Rotation"));
 	FVector InRotation = FVector(0, 0, 0);
 
+	/** Scale */
 	UPROPERTY(EditAnywhere, Category = "Transform", meta = (DataflowInput, DisplayName = "Scale"));
 	FVector InScale = FVector(1, 1, 1);
 
+	/** Result transform */
 	UPROPERTY(meta = (DataflowOutput, DisplayName = "Transform"));
 	FTransform OutTransform = FTransform::Identity;
 
+	virtual void Evaluate(UE::Dataflow::FContext& Context, const FDataflowOutput* Out) const override;
+
+public:
 	FMakeTransformDataflowNode(const UE::Dataflow::FNodeParameters& InParam, FGuid InGuid = FGuid::NewGuid())
 		: FDataflowNode(InParam, InGuid)
 	{
@@ -437,16 +443,15 @@ public:
 		RegisterInputConnection(&InScale);
 		RegisterOutputConnection(&OutTransform);
 	}
-
-	virtual void Evaluate(UE::Dataflow::FContext& Context, const FDataflowOutput* Out) const override;
 };
 
 /**
 *
 * Make a FTransform
-*
+* Deprecated (5.6)
+* Use FMakeTransformDataflowNode instead
 */
-USTRUCT()
+USTRUCT(meta = (Deprecated = "5.6"))
 struct FMakeTransformDataflowNode_v2 : public FDataflowNode
 {
 	GENERATED_USTRUCT_BODY()
