Dataflow : expose translate / rotate / scale to Transform collection

diff --git a/Engine/Plugins/Experimental/GeometryCollectionPlugin/Source/GeometryCollectionNodes/Private/Dataflow/GeometryCollectionNodes.cpp b/Engine/Plugins/Experimental/GeometryCollectionPlugin/Source/GeometryCollectionNodes/Private/Dataflow/GeometryCollectionNodes.cpp
--- a/Engine/Plugins/Experimental/GeometryCollectionPlugin/Source/GeometryCollectionNodes/Private/Dataflow/GeometryCollectionNodes.cpp
+++ b/Engine/Plugins/Experimental/GeometryCollectionPlugin/Source/GeometryCollectionNodes/Private/Dataflow/GeometryCollectionNodes.cpp
@@ -478,10 +478,15 @@ void FTransformCollectionDataflowNode::Evaluate(UE::Dataflow::FContext& Context,
 	{
 		FManagedArrayCollection InCollection = GetValue<FManagedArrayCollection>(Context, &Collection);
 
-		FTransform NewTransform = GeometryCollection::Facades::FCollectionTransformFacade::BuildTransform(Translate,
+		const FVector InTranslate = GetValue(Context, &Translate);
+		const FVector InRotate = GetValue(Context, &Rotate);
+		const FVector InScale = GetValue(Context, &Scale);
+
+		FTransform NewTransform = GeometryCollection::Facades::FCollectionTransformFacade::BuildTransform(
+			InTranslate,
 			(uint8)RotationOrder,
-			Rotate,
-			Scale,
+			InRotate,
+			InScale,
 			UniformScale,
 			RotatePivot,
 			ScalePivot,
diff --git a/Engine/Plugins/Experimental/GeometryCollectionPlugin/Source/GeometryCollectionNodes/Public/Dataflow/GeometryCollectionNodes.h b/Engine/Plugins/Experimental/GeometryCollectionPlugin/Source/GeometryCollectionNodes/Public/Dataflow/GeometryCollectionNodes.h
--- a/Engine/Plugins/Experimental/GeometryCollectionPlugin/Source/GeometryCollectionNodes/Public/Dataflow/GeometryCollectionNodes.h
+++ b/Engine/Plugins/Experimental/GeometryCollectionPlugin/Source/GeometryCollectionNodes/Public/Dataflow/GeometryCollectionNodes.h
@@ -688,7 +688,7 @@ public:
 	FManagedArrayCollection Collection;
 
 	/** Translation */
-	UPROPERTY(EditAnywhere, Category = "Transform");
+	UPROPERTY(EditAnywhere, Category = "Transform", meta = (DataflowInput));
 	FVector Translate = FVector(0.0);
 
 	/** Rotation Order */
@@ -696,11 +696,11 @@ public:
 	ERotationOrderEnum RotationOrder = ERotationOrderEnum::Dataflow_RotationOrder_XYZ;
 
 	/** Rotation */
-	UPROPERTY(EditAnywhere, Category = "Transform");
+	UPROPERTY(EditAnywhere, Category = "Transform", meta = (DataflowInput));
 	FVector Rotate = FVector(0.0);
 
 	/** Scale */
-	UPROPERTY(EditAnywhere, Category = "Transform", meta = (UIMin = 0.001f));
+	UPROPERTY(EditAnywhere, Category = "Transform", meta = (UIMin = 0.001f), meta = (DataflowInput));
 	FVector Scale = FVector(1.0);
 
 	/** Shear */
@@ -727,6 +727,9 @@ public:
 		: FDataflowNode(InParam, InGuid)
 	{
 		RegisterInputConnection(&Collection);
+		RegisterInputConnection(&Translate);
+		RegisterInputConnection(&Rotate);
+		RegisterInputConnection(&Scale);
 		RegisterOutputConnection(&Collection, &Collection);
 	}
 
