Dataflow : use selection information to display convex related node debug draw

diff --git a/Engine/Plugins/Experimental/GeometryCollectionPlugin/Source/GeometryCollectionNodes/Private/Dataflow/GeometryCollectionUtilityNodes.cpp b/Engine/Plugins/Experimental/GeometryCollectionPlugin/Source/GeometryCollectionNodes/Private/Dataflow/GeometryCollectionUtilityNodes.cpp
--- a/Engine/Plugins/Experimental/GeometryCollectionPlugin/Source/GeometryCollectionNodes/Private/Dataflow/GeometryCollectionUtilityNodes.cpp
+++ b/Engine/Plugins/Experimental/GeometryCollectionPlugin/Source/GeometryCollectionNodes/Private/Dataflow/GeometryCollectionUtilityNodes.cpp
@@ -98,12 +98,14 @@ namespace UE::Dataflow::Convex
 		return FLinearColor(FColor(R, G, B, 255));
 	}
 
-	static void DebugDrawProc(IDataflowDebugDrawInterface& DataflowRenderingInterface, const FManagedArrayCollection& InCollection, const bool bRandomizeColor, const int32 ColorRandomSeed)
+	static void DebugDrawProc(IDataflowDebugDrawInterface& DataflowRenderingInterface, const FManagedArrayCollection& InCollection, const bool bRandomizeColor, const int32 ColorRandomSeed, const FDataflowTransformSelection& Selection)
 	{
 		using namespace UE::Geometry;
 
 		TArray<FDynamicMesh3> HullsMeshes;
-		UE::FractureEngine::Convex::GetConvexHullsAsDynamicMeshes(InCollection, HullsMeshes, false, {});
+		const bool bRestrictToSelection = (Selection.Num() > 0);
+	
+		UE::FractureEngine::Convex::GetConvexHullsAsDynamicMeshes(InCollection, HullsMeshes, bRestrictToSelection, Selection.AsArray());
 
 		int32 Idx = 0;
 		for (const FDynamicMesh3& Mesh : HullsMeshes)
@@ -198,8 +200,8 @@ void FCreateLeafConvexHullsDataflowNode::DebugDraw(UE::Dataflow::FContext& Conte
 		if (const FDataflowOutput* Output = FindOutput(&Collection))
 		{
 			const FManagedArrayCollection& InCollection = Output->GetValue(Context, Collection);
-
-			UE::Dataflow::Convex::DebugDrawProc(DataflowRenderingInterface, InCollection, bRandomizeColor, ColorRandomSeed);
+			const FDataflowTransformSelection& InOptionalSelectionFilter = GetValue(Context, &OptionalSelectionFilter);
+			UE::Dataflow::Convex::DebugDrawProc(DataflowRenderingInterface, InCollection, bRandomizeColor, ColorRandomSeed, InOptionalSelectionFilter);
 		}
 	}
 }
@@ -272,8 +274,9 @@ void FSimplifyConvexHullsDataflowNode::DebugDraw(UE::Dataflow::FContext& Context
 		if (const FDataflowOutput* Output = FindOutput(&Collection))
 		{
 			const FManagedArrayCollection& OutCollection = Output->GetValue(Context, Collection);
+			const FDataflowTransformSelection& InOptionalSelectionFilter = GetValue(Context, &OptionalSelectionFilter);
 
-			UE::Dataflow::Convex::DebugDrawProc(DataflowRenderingInterface, OutCollection, bRandomizeColor, ColorRandomSeed);
+			UE::Dataflow::Convex::DebugDrawProc(DataflowRenderingInterface, OutCollection, bRandomizeColor, ColorRandomSeed, InOptionalSelectionFilter);
 		}
 	}
 }
@@ -338,8 +341,9 @@ void FCreateNonOverlappingConvexHullsDataflowNode::DebugDraw(UE::Dataflow::FCont
 		if (const FDataflowOutput* Output = FindOutput(&Collection))
 		{
 			const FManagedArrayCollection& InCollection = Output->GetValue(Context, Collection);
+			FDataflowTransformSelection EmptySelection;
 
-			UE::Dataflow::Convex::DebugDrawProc(DataflowRenderingInterface, InCollection, bRandomizeColor, ColorRandomSeed);
+			UE::Dataflow::Convex::DebugDrawProc(DataflowRenderingInterface, InCollection, bRandomizeColor, ColorRandomSeed, EmptySelection);
 		}
 	}
 }
@@ -475,8 +479,9 @@ void FGenerateClusterConvexHullsFromLeafHullsDataflowNode::DebugDraw(UE::Dataflo
 		if (const FDataflowOutput* Output = FindOutput(&Collection))
 		{
 			const FManagedArrayCollection& InCollection = Output->GetValue(Context, Collection);
+			const FDataflowTransformSelection& InOptionalSelectionFilter = GetValue<FDataflowTransformSelection>(Context, &OptionalSelectionFilter);
 
-			UE::Dataflow::Convex::DebugDrawProc(DataflowRenderingInterface, InCollection, bRandomizeColor, ColorRandomSeed);
+			UE::Dataflow::Convex::DebugDrawProc(DataflowRenderingInterface, InCollection, bRandomizeColor, ColorRandomSeed, InOptionalSelectionFilter);
 		}
 	}
 }
@@ -597,8 +602,9 @@ void FGenerateClusterConvexHullsFromChildrenHullsDataflowNode::DebugDraw(UE::Dat
 		if (const FDataflowOutput* Output = FindOutput(&Collection))
 		{
 			const FManagedArrayCollection& InCollection = Output->GetValue(Context, Collection);
+			const FDataflowTransformSelection& InOptionalSelectionFilter = GetValue(Context, &OptionalSelectionFilter);
 
-			UE::Dataflow::Convex::DebugDrawProc(DataflowRenderingInterface, InCollection, bRandomizeColor, ColorRandomSeed);
+			UE::Dataflow::Convex::DebugDrawProc(DataflowRenderingInterface, InCollection, bRandomizeColor, ColorRandomSeed, InOptionalSelectionFilter);
 		}
 	}
 }
@@ -732,8 +738,9 @@ void FMergeConvexHullsDataflowNode::DebugDraw(UE::Dataflow::FContext& Context, I
 		if (const FDataflowOutput* Output = FindOutput(&Collection))
 		{
 			const FManagedArrayCollection& InCollection = Output->GetValue(Context, Collection);
+			const FDataflowTransformSelection& InOptionalSelectionFilter = GetValue<FDataflowTransformSelection>(Context, &OptionalSelectionFilter);
 
-			UE::Dataflow::Convex::DebugDrawProc(DataflowRenderingInterface, InCollection, bRandomizeColor, ColorRandomSeed);
+			UE::Dataflow::Convex::DebugDrawProc(DataflowRenderingInterface, InCollection, bRandomizeColor, ColorRandomSeed, InOptionalSelectionFilter);
 		}
 	}
 }
@@ -876,8 +883,9 @@ void FGetConvexHullVolumeDataflowNode::DebugDraw(UE::Dataflow::FContext& Context
 		if (IsConnected(&Collection))
 		{
 			const FManagedArrayCollection& InCollection = GetValue(Context, &Collection);
+			const FDataflowTransformSelection& InSelection = GetValue(Context, &TransformSelection);
 
-			UE::Dataflow::Convex::DebugDrawProc(DataflowRenderingInterface, InCollection, bRandomizeColor, ColorRandomSeed);
+			UE::Dataflow::Convex::DebugDrawProc(DataflowRenderingInterface, InCollection, bRandomizeColor, ColorRandomSeed, InSelection);
 		}
 	}
 }
