remove stub for convex hulls dataflownode (accidentally added)

diff --git a/Engine/Plugins/Experimental/GeometryCollectionPlugin/Source/GeometryCollectionNodes/Public/Dataflow/GeometryCollectionUtilityNodes.h b/Engine/Plugins/Experimental/GeometryCollectionPlugin/Source/GeometryCollectionNodes/Public/Dataflow/GeometryCollectionUtilityNodes.h
--- a/Engine/Plugins/Experimental/GeometryCollectionPlugin/Source/GeometryCollectionNodes/Public/Dataflow/GeometryCollectionUtilityNodes.h
+++ b/Engine/Plugins/Experimental/GeometryCollectionPlugin/Source/GeometryCollectionNodes/Public/Dataflow/GeometryCollectionUtilityNodes.h
@@ -605,26 +605,6 @@ public:
 };
 
 
-/** Copy convex hulls from a reference collection */
-USTRUCT(meta = (DataflowGeometryCollection))
-struct FCopyConvexHullsAcrossCollectionBonesDataflowNode : public FDataflowNode
-{
-	GENERATED_USTRUCT_BODY()
-	DATAFLOW_NODE_DEFINE_INTERNAL(FCopyConvexHullsAcrossCollectionBonesDataflowNode, "CopyConvexHullsAcrossCollectionBones", "GeometryCollection|Utilities", "")
-
-private:
-	UPROPERTY(meta = (DataflowInput))
-	FManagedArrayCollection SourceCollection;
-
-	UPROPERTY(meta = (DataflowOutput))
-	FManagedArrayCollection TargetCollection;
-
-	UPROPERTY(meta = (DataflowInput))
-	FDataflowTransformSelection SourceTransformSelection;
-
-	UPROPERTY(meta = (DataflowInput))
-	FDataflowTransformSelection TargetTransformSelection;
-};
 
 
 /** Merge convex hulls on transforms with multiple hulls */
