Dataflow : Add SetDynamicState node

diff --git a/Engine/Plugins/Experimental/GeometryCollectionPlugin/Source/GeometryCollectionNodes/Private/Dataflow/GeometryCollectionNodes.cpp b/Engine/Plugins/Experimental/GeometryCollectionPlugin/Source/GeometryCollectionNodes/Private/Dataflow/GeometryCollectionNodes.cpp
--- a/Engine/Plugins/Experimental/GeometryCollectionPlugin/Source/GeometryCollectionNodes/Private/Dataflow/GeometryCollectionNodes.cpp
+++ b/Engine/Plugins/Experimental/GeometryCollectionPlugin/Source/GeometryCollectionNodes/Private/Dataflow/GeometryCollectionNodes.cpp
@@ -69,6 +69,7 @@ namespace UE::Dataflow
 		DATAFLOW_NODE_REGISTER_CREATION_FACTORY(FGetSchemaDataflowNode);
 		DATAFLOW_NODE_REGISTER_CREATION_FACTORY(FRemoveOnBreakDataflowNode);
 		DATAFLOW_NODE_REGISTER_CREATION_FACTORY(FSetAnchorStateDataflowNode);
+		DATAFLOW_NODE_REGISTER_CREATION_FACTORY(FSetDynamicStateDataflowNode);
 		DATAFLOW_NODE_REGISTER_CREATION_FACTORY(FProximityDataflowNode);
 		DATAFLOW_NODE_REGISTER_CREATION_FACTORY(FCollectionSetPivotDataflowNode);
 		DATAFLOW_NODE_REGISTER_CREATION_FACTORY(FAddCustomCollectionAttributeDataflowNode);
@@ -768,45 +769,92 @@ void FRemoveOnBreakDataflowNode::Evaluate(UE::Dataflow::FContext& Context, const
 	}
 }
 
+//////////////////////////////////////////////////////////////////////////////////////
+
+FSetAnchorStateDataflowNode::FSetAnchorStateDataflowNode(const UE::Dataflow::FNodeParameters& InParam, FGuid InGuid )
+	: FDataflowNode(InParam, InGuid)
+{
+	RegisterInputConnection(&Collection);
+	RegisterInputConnection(&TransformSelection);
+	RegisterOutputConnection(&Collection, &Collection);
+}
+
 void FSetAnchorStateDataflowNode::Evaluate(UE::Dataflow::FContext& Context, const FDataflowOutput* Out) const
 {
-	if (Out->IsA<FManagedArrayCollection>(&Collection))
+	if (Out->IsA(&Collection))
 	{
-		const FManagedArrayCollection& InCollection = GetValue<FManagedArrayCollection>(Context, &Collection);
+		FManagedArrayCollection OutCollection = GetValue(Context, &Collection);
 		if (IsConnected(&Collection))
 		{
-			FDataflowTransformSelection InTransformSelection = GetValue<FDataflowTransformSelection>(Context, &TransformSelection);
+			FDataflowTransformSelection InTransformSelection = GetValue(Context, &TransformSelection);
 
-			if (TUniquePtr<FGeometryCollection> GeomCollection = TUniquePtr<FGeometryCollection>(InCollection.NewCopy<FGeometryCollection>()))
+			Chaos::Facades::FCollectionAnchoringFacade AnchoringFacade(OutCollection);
+			if (!AnchoringFacade.HasAnchoredAttribute())
 			{
-				Chaos::Facades::FCollectionAnchoringFacade AnchoringFacade(*GeomCollection);
-				if (!AnchoringFacade.HasAnchoredAttribute())
-				{
-					AnchoringFacade.AddAnchoredAttribute();
-				}
-
-				bool bAnchored = (AnchorState == EAnchorStateEnum::Dataflow_AnchorState_Anchored) ? true : false;
-				TArray<int32> BoneIndices;
-				InTransformSelection.AsArray(BoneIndices);
-				AnchoringFacade.SetAnchored(BoneIndices, bAnchored);
+				AnchoringFacade.AddAnchoredAttribute();
+			}
 
-				if (bSetNotSelectedBonesToOppositeState)
-				{
-					InTransformSelection.Invert();
-					InTransformSelection.AsArray(BoneIndices);
-					AnchoringFacade.SetAnchored(BoneIndices, !bAnchored);
-				}
+			const bool bAnchored = (AnchorState == EAnchorStateEnum::Dataflow_AnchorState_Anchored) ? true : false;
+			TArray<int32> BoneIndices;
+			InTransformSelection.AsArray(BoneIndices);
+			AnchoringFacade.SetAnchored(BoneIndices, bAnchored);
 
-				SetValue(Context, static_cast<const FManagedArrayCollection&>(*GeomCollection), &Collection);
+			if (bSetNotSelectedBonesToOppositeState)
+			{
+				InTransformSelection.Invert();
+				InTransformSelection.AsArray(BoneIndices);
+				AnchoringFacade.SetAnchored(BoneIndices, !bAnchored);
 			}
 		}
-		else
+		SetValue(Context, OutCollection, &Collection);
+	}
+}
+
+//////////////////////////////////////////////////////////////////////////////////////
+
+FSetDynamicStateDataflowNode::FSetDynamicStateDataflowNode(const UE::Dataflow::FNodeParameters& InParam, FGuid InGuid)
+	: FDataflowNode(InParam, InGuid)
+{
+	RegisterInputConnection(&Collection);
+	RegisterInputConnection(&TransformSelection);
+	RegisterOutputConnection(&Collection, &Collection);
+}
+
+void FSetDynamicStateDataflowNode::Evaluate(UE::Dataflow::FContext& Context, const FDataflowOutput* Out) const
+{
+	if (Out->IsA(&Collection))
+	{
+		FManagedArrayCollection OutCollection = GetValue(Context, &Collection);
+		if (IsConnected(&Collection))
 		{
-			SetValue(Context, InCollection, &Collection);
+			const FDataflowTransformSelection& InTransformSelection = GetValue(Context, &TransformSelection);
+
+			const TArray<int32> BoneIndices = InTransformSelection.AsArray();
+
+			Chaos::Facades::FCollectionAnchoringFacade AnchoringFacade(OutCollection);
+			if (AnchoringFacade.HasInitialDynamicStateAttribute())
+			{
+				const Chaos::EObjectStateType ObjectState = [this]()
+					{
+						switch (DynamicState)
+						{
+						case EDataflowGeometryCollectionDynamicState::None:		return Chaos::EObjectStateType::Uninitialized;
+						case EDataflowGeometryCollectionDynamicState::Dynamic:	return Chaos::EObjectStateType::Dynamic;
+						case EDataflowGeometryCollectionDynamicState::Kinematic:return Chaos::EObjectStateType::Kinematic;
+						case EDataflowGeometryCollectionDynamicState::Static:	return Chaos::EObjectStateType::Static;
+						}
+						return Chaos::EObjectStateType::Dynamic;
+					}();
+					
+
+				AnchoringFacade.SetInitialDynamicState(BoneIndices, ObjectState);
+			}
 		}
+		SetValue(Context, OutCollection, &Collection);
 	}
 }
 
+//////////////////////////////////////////////////////////////////////////////////////
 
 void FProximityDataflowNode::Evaluate(UE::Dataflow::FContext& Context, const FDataflowOutput* Out) const
 {
diff --git a/Engine/Plugins/Experimental/GeometryCollectionPlugin/Source/GeometryCollectionNodes/Public/Dataflow/GeometryCollectionNodes.h b/Engine/Plugins/Experimental/GeometryCollectionPlugin/Source/GeometryCollectionNodes/Public/Dataflow/GeometryCollectionNodes.h
--- a/Engine/Plugins/Experimental/GeometryCollectionPlugin/Source/GeometryCollectionNodes/Public/Dataflow/GeometryCollectionNodes.h
+++ b/Engine/Plugins/Experimental/GeometryCollectionPlugin/Source/GeometryCollectionNodes/Public/Dataflow/GeometryCollectionNodes.h
@@ -1122,7 +1122,7 @@ struct FSetAnchorStateDataflowNode : public FDataflowNode
 	GENERATED_USTRUCT_BODY()
 	DATAFLOW_NODE_DEFINE_INTERNAL(FSetAnchorStateDataflowNode, "SetAnchorState", "GeometryCollection|Utilities", "")
 
-public:
+private:
 	/** What anchor state to set on selected bones */
 	UPROPERTY(EditAnywhere, Category = "Anchoring");
 	EAnchorStateEnum AnchorState = EAnchorStateEnum::Dataflow_AnchorState_Anchored;
@@ -1132,23 +1132,57 @@ public:
 	bool bSetNotSelectedBonesToOppositeState = false;
 
 	/** GeometryCollection to set anchor state on */
-	UPROPERTY(meta = (DataflowInput, DataflowOutput, DataflowPassthrough = "Collection", DataflowIntrinsic, DataflowCollectionGroups="CustomOne,CustomTwo"))
+	UPROPERTY(meta = (DataflowInput, DataflowOutput, DataflowPassthrough = "Collection", DataflowIntrinsic))
 	FManagedArrayCollection Collection;
 
 	/** Bone selection for setting the state on */
 	UPROPERTY(meta = (DataflowInput, DisplayName = "TransformSelection", DataflowIntrinsic))
 	FDataflowTransformSelection TransformSelection;
 
-	FSetAnchorStateDataflowNode(const UE::Dataflow::FNodeParameters& InParam, FGuid InGuid = FGuid::NewGuid())
-		: FDataflowNode(InParam, InGuid)
-	{
-		RegisterInputConnection(&Collection);
-		RegisterInputConnection(&TransformSelection);
-		RegisterOutputConnection(&Collection, &Collection);
-	}
+	virtual void Evaluate(UE::Dataflow::FContext& Context, const FDataflowOutput* Out) const override;
+
+public:
+	FSetAnchorStateDataflowNode(const UE::Dataflow::FNodeParameters& InParam, FGuid InGuid = FGuid::NewGuid());
+};
+
+UENUM(BlueprintType)
+enum class EDataflowGeometryCollectionDynamicState : uint8
+{
+	None = 0		UMETA(DisplayName = "None"),
+	Dynamic  = 1	UMETA(DisplayName = "Dynamic"),
+	Kinematic =  2	UMETA(DisplayName = "Kinematic"),
+	Static = 3		UMETA(DisplayName = "Static")
+};
+
+/**
+ *
+ * Sets the dynamic state on the selected bones in a Collection
+ *
+ */
+USTRUCT(meta = (DataflowGeometryCollection))
+struct FSetDynamicStateDataflowNode : public FDataflowNode
+{
+	GENERATED_USTRUCT_BODY()
+	DATAFLOW_NODE_DEFINE_INTERNAL(FSetDynamicStateDataflowNode, "SetDynamicState", "GeometryCollection|Utilities", "")
+
+private:
+	/** Dynamic state to set on selected bones */
+	UPROPERTY(EditAnywhere, Category = "State");
+	EDataflowGeometryCollectionDynamicState DynamicState = EDataflowGeometryCollectionDynamicState::Kinematic;
+
+	/** GeometryCollection to set anchor state on */
+	UPROPERTY(meta = (DataflowInput, DataflowOutput, DataflowPassthrough = "Collection", DataflowIntrinsic))
+	FManagedArrayCollection Collection;
+
+	/** Bone selection for setting the state on */
+	UPROPERTY(meta = (DataflowInput, DisplayName = "TransformSelection", DataflowIntrinsic))
+	FDataflowTransformSelection TransformSelection;
 
 	virtual void Evaluate(UE::Dataflow::FContext& Context, const FDataflowOutput* Out) const override;
 
+public:
+	FSetDynamicStateDataflowNode(const UE::Dataflow::FNodeParameters& InParam, FGuid InGuid = FGuid::NewGuid());
+
 };
 
 
diff --git a/Engine/Source/Runtime/Experimental/Chaos/Private/GeometryCollection/Facades/CollectionAnchoringFacade.cpp b/Engine/Source/Runtime/Experimental/Chaos/Private/GeometryCollection/Facades/CollectionAnchoringFacade.cpp
--- a/Engine/Source/Runtime/Experimental/Chaos/Private/GeometryCollection/Facades/CollectionAnchoringFacade.cpp
+++ b/Engine/Source/Runtime/Experimental/Chaos/Private/GeometryCollection/Facades/CollectionAnchoringFacade.cpp
@@ -70,7 +70,10 @@ namespace Chaos::Facades
 			TManagedArray<int32>& InitialDynamicState = InitialDynamicStateAttribute.Modify();
 			for (const int32 TransformIndex: TransformIndices)
 			{
-				InitialDynamicState[TransformIndex] = static_cast<int32>(State);
+				if (InitialDynamicState.IsValidIndex(TransformIndex))
+				{
+					InitialDynamicState[TransformIndex] = static_cast<int32>(State);
+				}
 			}
 		}
 	}
