Dataflow : Fix the selection by level nodes to make sure the level attribute exists in the collection

diff --git a/Engine/Plugins/Experimental/GeometryCollectionPlugin/Source/GeometryCollectionNodes/Private/Dataflow/GeometryCollectionSelectionNodes.cpp b/Engine/Plugins/Experimental/GeometryCollectionPlugin/Source/GeometryCollectionNodes/Private/Dataflow/GeometryCollectionSelectionNodes.cpp
--- a/Engine/Plugins/Experimental/GeometryCollectionPlugin/Source/GeometryCollectionNodes/Private/Dataflow/GeometryCollectionSelectionNodes.cpp
+++ b/Engine/Plugins/Experimental/GeometryCollectionPlugin/Source/GeometryCollectionNodes/Private/Dataflow/GeometryCollectionSelectionNodes.cpp
@@ -28,6 +28,7 @@
 #include "FractureEngineClustering.h"
 #include "FractureEngineSelection.h"
 #include "GeometryCollection/Facades/CollectionTransformSelectionFacade.h"
+#include "GeometryCollection/Facades/CollectionHierarchyFacade.h"
 #include "Internationalization/Regex.h"
 
 #include UE_INLINE_GENERATED_CPP_BY_NAME(GeometryCollectionSelectionNodes)
@@ -657,46 +658,55 @@ void FCollectionTransformSelectionSiblingsDataflowNode::Evaluate(UE::Dataflow::F
 
 void FCollectionTransformSelectionLevelDataflowNode::Evaluate(UE::Dataflow::FContext& Context, const FDataflowOutput* Out) const
 {
-	if (Out->IsA<FDataflowTransformSelection>(&TransformSelection))
+	if (Out->IsA(&TransformSelection) || Out->IsA(&Collection))
 	{
-		const FManagedArrayCollection& InCollection = GetValue<FManagedArrayCollection>(Context, &Collection);
-		FDataflowTransformSelection InTransformSelection = GetValue<FDataflowTransformSelection>(Context, &TransformSelection);
+		FManagedArrayCollection OutCollection = GetValue(Context, &Collection);
+		FDataflowTransformSelection InTransformSelection = GetValue(Context, &TransformSelection);
 
-		GeometryCollection::Facades::FCollectionTransformSelectionFacade TransformSelectionFacade(InCollection);
+		// make sure there's a level attribute
+		Chaos::Facades::FCollectionHierarchyFacade HierarchyFacade(OutCollection);
+		if (!HierarchyFacade.HasLevelAttribute())
+		{
+			HierarchyFacade.GenerateLevelAttribute();
+		}
+
+		GeometryCollection::Facades::FCollectionTransformSelectionFacade TransformSelectionFacade(OutCollection);
 		TArray<int32> SelectionArr = InTransformSelection.AsArray();
 
 		TransformSelectionFacade.SelectLevel(SelectionArr);
 		InTransformSelection.SetFromArray(SelectionArr);
 
 		SetValue(Context, MoveTemp(InTransformSelection), &TransformSelection);
-	}
-	else if (Out->IsA<FManagedArrayCollection>(&Collection))
-	{
-		SafeForwardInput(Context, &Collection, &Collection);
+		SetValue(Context, MoveTemp(OutCollection), &Collection);
 	}
 }
 
 
 void FCollectionTransformSelectionTargetLevelDataflowNode::Evaluate(UE::Dataflow::FContext& Context, const FDataflowOutput* Out) const
 {
-	if (Out->IsA(&TransformSelection))
+	if (Out->IsA(&TransformSelection) || Out->IsA(&Collection))
 	{
-		const FManagedArrayCollection& InCollection = GetValue(Context, &Collection);
-		GeometryCollection::Facades::FCollectionTransformSelectionFacade TransformSelectionFacade(InCollection);
+		FManagedArrayCollection OutCollection = GetValue(Context, &Collection);
+
+		// make sure there's a level attribute
+		Chaos::Facades::FCollectionHierarchyFacade HierarchyFacade(OutCollection);
+		if (!HierarchyFacade.HasLevelAttribute())
+		{
+			HierarchyFacade.GenerateLevelAttribute();
+		}
+
+		GeometryCollection::Facades::FCollectionTransformSelectionFacade TransformSelectionFacade(OutCollection);
 
-		int InTargetLevel = GetValue(Context, &TargetLevel);
+		const int32 InTargetLevel = GetValue(Context, &TargetLevel);
 
-		TArray<int32> AllAtLevel = TransformSelectionFacade.GetBonesExactlyAtLevel(InTargetLevel, bSkipEmbedded);
+		const TArray<int32> AllAtLevel = TransformSelectionFacade.GetBonesExactlyAtLevel(InTargetLevel, bSkipEmbedded);
 
 		FDataflowTransformSelection NewTransformSelection;
-		NewTransformSelection.Initialize(InCollection.NumElements(FGeometryCollection::TransformGroup), false);
+		NewTransformSelection.Initialize(OutCollection.NumElements(FGeometryCollection::TransformGroup), false);
 		NewTransformSelection.SetFromArray(AllAtLevel);
 
 		SetValue(Context, MoveTemp(NewTransformSelection), &TransformSelection);
-	}
-	else if (Out->IsA(&Collection))
-	{
-		SafeForwardInput(Context, &Collection, &Collection);
+		SetValue(Context, MoveTemp(OutCollection), &Collection);
 	}
 }
 
diff --git a/Engine/Plugins/Experimental/GeometryCollectionPlugin/Source/GeometryCollectionNodes/Public/Dataflow/GeometryCollectionSelectionNodes.h b/Engine/Plugins/Experimental/GeometryCollectionPlugin/Source/GeometryCollectionNodes/Public/Dataflow/GeometryCollectionSelectionNodes.h
--- a/Engine/Plugins/Experimental/GeometryCollectionPlugin/Source/GeometryCollectionNodes/Public/Dataflow/GeometryCollectionSelectionNodes.h
+++ b/Engine/Plugins/Experimental/GeometryCollectionPlugin/Source/GeometryCollectionNodes/Public/Dataflow/GeometryCollectionSelectionNodes.h
@@ -421,8 +421,8 @@ public:
 	FCollectionTransformSelectionParentDataflowNode(const UE::Dataflow::FNodeParameters& InParam, FGuid InGuid = FGuid::NewGuid())
 		: FDataflowNode(InParam, InGuid)
 	{
-		RegisterInputConnection(&TransformSelection);
 		RegisterInputConnection(&Collection);
+		RegisterInputConnection(&TransformSelection);
 		RegisterOutputConnection(&Collection, &Collection);
 		RegisterOutputConnection(&TransformSelection, &TransformSelection);
 	}
@@ -498,8 +498,8 @@ public:
 	FCollectionTransformSelectionChildrenDataflowNode(const UE::Dataflow::FNodeParameters& InParam, FGuid InGuid = FGuid::NewGuid())
 		: FDataflowNode(InParam, InGuid)
 	{
-		RegisterInputConnection(&TransformSelection);
 		RegisterInputConnection(&Collection);
+		RegisterInputConnection(&TransformSelection);
 		RegisterOutputConnection(&Collection, &Collection);
 		RegisterOutputConnection(&TransformSelection, &TransformSelection);
 	}
@@ -532,8 +532,8 @@ public:
 	FCollectionTransformSelectionSiblingsDataflowNode(const UE::Dataflow::FNodeParameters& InParam, FGuid InGuid = FGuid::NewGuid())
 		: FDataflowNode(InParam, InGuid)
 	{
-		RegisterInputConnection(&TransformSelection);
 		RegisterInputConnection(&Collection);
+		RegisterInputConnection(&TransformSelection);
 		RegisterOutputConnection(&Collection, &Collection);
 		RegisterOutputConnection(&TransformSelection, &TransformSelection);
 	}
@@ -552,7 +552,7 @@ USTRUCT(meta = (DataflowGeometryCollection))
 struct FCollectionTransformSelectionLevelDataflowNode : public FDataflowNode
 {
 	GENERATED_USTRUCT_BODY()
-	DATAFLOW_NODE_DEFINE_INTERNAL(FCollectionTransformSelectionLevelDataflowNode, "CollectionTransformSelectLevels", "GeometryCollection|Selection|Transform", "")
+	DATAFLOW_NODE_DEFINE_INTERNAL(FCollectionTransformSelectionLevelDataflowNode, "CollectionTransformSelectSameLevel", "GeometryCollection|Selection|Transform", "")
 
 public:
 	/** Array of the selected bone indices */
@@ -566,8 +566,8 @@ public:
 	FCollectionTransformSelectionLevelDataflowNode(const UE::Dataflow::FNodeParameters& InParam, FGuid InGuid = FGuid::NewGuid())
 		: FDataflowNode(InParam, InGuid)
 	{
-		RegisterInputConnection(&TransformSelection);
 		RegisterInputConnection(&Collection);
+		RegisterInputConnection(&TransformSelection);
 		RegisterOutputConnection(&Collection, &Collection);
 		RegisterOutputConnection(&TransformSelection, &TransformSelection);
 	}
@@ -647,8 +647,8 @@ public:
 	FCollectionTransformSelectionContactDataflowNode(const UE::Dataflow::FNodeParameters& InParam, FGuid InGuid = FGuid::NewGuid())
 		: FDataflowNode(InParam, InGuid)
 	{
-		RegisterInputConnection(&TransformSelection);
 		RegisterInputConnection(&Collection);
+		RegisterInputConnection(&TransformSelection);
 		RegisterOutputConnection(&Collection, &Collection);
 		RegisterOutputConnection(&TransformSelection, &TransformSelection);
 	}
