fix dataflow mesh cutter not using mesh normals and UVs

diff --git a/Engine/Plugins/Experimental/GeometryCollectionPlugin/Source/GeometryCollectionNodes/Private/Dataflow/GeometryCollectionFracturingNodes.cpp b/Engine/Plugins/Experimental/GeometryCollectionPlugin/Source/GeometryCollectionNodes/Private/Dataflow/GeometryCollectionFracturingNodes.cpp
--- a/Engine/Plugins/Experimental/GeometryCollectionPlugin/Source/GeometryCollectionNodes/Private/Dataflow/GeometryCollectionFracturingNodes.cpp
+++ b/Engine/Plugins/Experimental/GeometryCollectionPlugin/Source/GeometryCollectionNodes/Private/Dataflow/GeometryCollectionFracturingNodes.cpp
@@ -1044,8 +1044,8 @@ void FMeshCutterDataflowNode::Evaluate(UE::Dataflow::FContext& Context, const FD
 
 					if (MeshDescription->Vertices().Num() > 0)
 					{
-						FMeshDescriptionToDynamicMesh ConverterToDynamicMesh;
-						ConverterToDynamicMesh.Convert(MeshDescription, LocalMesh);
+						int32 NumUVLayers = GeometryCollection::UV::GetNumUVLayers(InCollection);
+						LocalMesh = ConvertMeshDescriptionToCuttingDynamicMesh(MeshDescription, NumUVLayers);
 						UseMeshes.Add(&LocalMesh);
 					}
 
diff --git a/Engine/Plugins/Experimental/PlanarCutPlugin/Source/PlanarCut/Private/GeometryMeshConversion.cpp b/Engine/Plugins/Experimental/PlanarCutPlugin/Source/PlanarCut/Private/GeometryMeshConversion.cpp
--- a/Engine/Plugins/Experimental/PlanarCutPlugin/Source/PlanarCut/Private/GeometryMeshConversion.cpp
+++ b/Engine/Plugins/Experimental/PlanarCutPlugin/Source/PlanarCut/Private/GeometryMeshConversion.cpp
@@ -1615,6 +1615,8 @@ FCellMeshes::FCellMeshes(int32 NumUVLayersIn, const FDynamicMesh3& SingleCutter,
 	if (!AugmentedDynamicMesh::IsAugmented(CellMeshes[0].AugMesh))
 	{
 		AugmentedDynamicMesh::Augment(CellMeshes[0].AugMesh, NumUVLayers);
+		// transfer overlay attributes to 'augmented attribute' format
+		AugmentedDynamicMesh::SplitOverlayAttributesToPerVertex(CellMeshes[0].AugMesh);
 	}
 
 	// Make sure color is unset
