add missing path in dataflow to go from single dynamic mesh to an array (incl. options to split islands)

diff --git a/Engine/Plugins/Experimental/GeometryCollectionPlugin/Source/GeometryCollectionNodes/Private/Dataflow/GeometryCollectionMeshNodes.cpp b/Engine/Plugins/Experimental/GeometryCollectionPlugin/Source/GeometryCollectionNodes/Private/Dataflow/GeometryCollectionMeshNodes.cpp
--- a/Engine/Plugins/Experimental/GeometryCollectionPlugin/Source/GeometryCollectionNodes/Private/Dataflow/GeometryCollectionMeshNodes.cpp
+++ b/Engine/Plugins/Experimental/GeometryCollectionPlugin/Source/GeometryCollectionNodes/Private/Dataflow/GeometryCollectionMeshNodes.cpp
@@ -55,6 +55,7 @@ namespace UE::Dataflow
 		DATAFLOW_NODE_REGISTER_CREATION_FACTORY(FMakeDataflowMeshDataflowNode);
 		DATAFLOW_NODE_REGISTER_CREATION_FACTORY(FDuplicateMeshUVChannelNode);
 		DATAFLOW_NODE_REGISTER_CREATION_FACTORY(FSplitDataflowMeshDataflowNode);
+		DATAFLOW_NODE_REGISTER_CREATION_FACTORY(FSplitMeshIslandsDataflowNode);
 		DATAFLOW_NODE_REGISTER_CREATION_FACTORY(FMeshCopyToPointsDataflowNode);
 		DATAFLOW_NODE_REGISTER_CREATION_FACTORY(FGetMeshDataDataflowNode);
 		DATAFLOW_NODE_REGISTER_CREATION_FACTORY(FApplyMeshProcessorToMeshDataflowNode);
@@ -476,6 +477,55 @@ void FMakeDataflowMeshDataflowNode::Evaluate(UE::Dataflow::FContext& Context, co
 	}
 }
 
+FSplitMeshIslandsDataflowNode::FSplitMeshIslandsDataflowNode(const UE::Dataflow::FNodeParameters& InParam, FGuid InGuid)
+	: FDataflowNode(InParam, InGuid)
+{
+	RegisterInputConnection(&Mesh);
+	RegisterOutputConnection(&Meshes);
+}
+
+void FSplitMeshIslandsDataflowNode::Evaluate(UE::Dataflow::FContext& Context, const FDataflowOutput* Out) const
+{
+	using namespace UE::Geometry;
+
+	if (Out->IsA(&Meshes))
+	{
+		TArray<TObjectPtr<UDynamicMesh>> OutMeshes;
+		if (const TObjectPtr<UDynamicMesh> InMesh = GetValue(Context, &Mesh))
+		{
+			if (SplitMethod == EDataflowMeshSplitIslandsMethod::NoSplit)
+			{
+				OutMeshes.Add(InMesh);
+			}
+			else
+			{
+				InMesh->ProcessMesh([&OutMeshes, this](const FDynamicMesh3& ToSplit)
+				{
+					TArray<FDynamicMesh3> SplitMeshes;
+					FVertexConnectedComponents Components(ToSplit.MaxVertexID());
+					Components.ConnectTriangles(ToSplit);
+					if (SplitMethod == EDataflowMeshSplitIslandsMethod::ByVertexOverlap)
+					{
+						Components.ConnectCloseVertices(ToSplit, ConnectVerticesThreshold, 2);
+					}
+					FDynamicMeshEditor::SplitMesh(&ToSplit, SplitMeshes, [&Components, &ToSplit](int32 TID)
+					{
+						return Components.GetComponent(ToSplit.GetTriangle(TID).A);
+					});
+					OutMeshes.SetNum(SplitMeshes.Num());
+					for (int32 Idx = 0; Idx < SplitMeshes.Num(); ++Idx)
+					{
+						OutMeshes[Idx] = NewObject<UDynamicMesh>();
+						OutMeshes[Idx]->SetMesh(MoveTemp(SplitMeshes[Idx]));
+					}
+				});
+			}
+		}
+		SetValue(Context, OutMeshes, &Meshes);
+	}
+}
+
+
 FSplitDataflowMeshDataflowNode::FSplitDataflowMeshDataflowNode(const UE::Dataflow::FNodeParameters& InParam, FGuid InGuid)
 	: FDataflowNode(InParam, InGuid)
 {
diff --git a/Engine/Plugins/Experimental/GeometryCollectionPlugin/Source/GeometryCollectionNodes/Public/Dataflow/GeometryCollectionMeshNodes.h b/Engine/Plugins/Experimental/GeometryCollectionPlugin/Source/GeometryCollectionNodes/Public/Dataflow/GeometryCollectionMeshNodes.h
--- a/Engine/Plugins/Experimental/GeometryCollectionPlugin/Source/GeometryCollectionNodes/Public/Dataflow/GeometryCollectionMeshNodes.h
+++ b/Engine/Plugins/Experimental/GeometryCollectionPlugin/Source/GeometryCollectionNodes/Public/Dataflow/GeometryCollectionMeshNodes.h
@@ -593,6 +593,50 @@ private:
 };
 
 
+UENUM(BlueprintType)
+enum class EDataflowMeshSplitIslandsMethod : uint8
+{
+	NoSplit,
+	ByMeshTopology,
+	ByVertexOverlap
+};
+
+
+/**
+ *
+ * Split a mesh into a connected islands
+ *
+ */
+USTRUCT(meta = (DataflowGeometryCollection))
+struct FSplitMeshIslandsDataflowNode : public FDataflowNode
+{
+	GENERATED_USTRUCT_BODY()
+	DATAFLOW_NODE_DEFINE_INTERNAL(FSplitMeshIslandsDataflowNode, "SplitMeshIslands", "Mesh|Utilities", "")
+
+public:
+
+	FSplitMeshIslandsDataflowNode(const UE::Dataflow::FNodeParameters& InParam, FGuid InGuid = FGuid::NewGuid());
+
+private:
+
+	/** Mesh input */
+	UPROPERTY(meta = (DataflowInput, DataflowIntrinsic))
+	TObjectPtr<UDynamicMesh> Mesh;
+
+	/** Meshes output */
+	UPROPERTY(meta = (DataflowOutput))
+	TArray<TObjectPtr<UDynamicMesh>> Meshes;
+
+	/** Whether to consider coincident vertices as connected even if the topology does not connect them */
+	UPROPERTY(EditAnywhere, Category = "General");
+	EDataflowMeshSplitIslandsMethod SplitMethod = EDataflowMeshSplitIslandsMethod::ByMeshTopology;
+
+	/** Vertices closer than this distance are considered to be overlapping */
+	UPROPERTY(EditAnywhere, Category = "General", meta = (ClampMin = 0.0, EditCondition = "SplitMethod == EDataflowMeshSplitIslandsMethod::ByVertexOverlap"));
+	float ConnectVerticesThreshold = 0.001;
+
+	virtual void Evaluate(UE::Dataflow::FContext& Context, const FDataflowOutput* Out) const override;
+};
 
 /**
  *
