fix "append collections" dataflow node not updating transform indices correctly

diff --git a/Engine/Plugins/Experimental/GeometryCollectionPlugin/Source/GeometryCollectionNodes/Private/Dataflow/GeometryCollectionNodes.cpp b/Engine/Plugins/Experimental/GeometryCollectionPlugin/Source/GeometryCollectionNodes/Private/Dataflow/GeometryCollectionNodes.cpp
--- a/Engine/Plugins/Experimental/GeometryCollectionPlugin/Source/GeometryCollectionNodes/Private/Dataflow/GeometryCollectionNodes.cpp
+++ b/Engine/Plugins/Experimental/GeometryCollectionPlugin/Source/GeometryCollectionNodes/Private/Dataflow/GeometryCollectionNodes.cpp
@@ -130,6 +130,49 @@ void FAppendCollectionAssetsDataflowNode::Evaluate(UE::Dataflow::FContext& Conte
 			GeometryGroupGuidsLocal1 = GuidArray1->GetConstArray();
 		}
 		InCollection1.Append(InCollection2);
+
+		//Manually update indices in TransformToGeometryIndex, Parent and Children attributes, since they do not have group dependencies set to automatically manage this
+		// TODO: Can we set up dependencies s.t. these indices are updated automatically, and then remove this manual fixup?
+		{
+			const int32 GeometryOffset = InCollection2.NumElements(FGeometryCollection::GeometryGroup);
+			const int32 OtherSize = InCollection2.NumElements(FGeometryCollection::TransformGroup);
+			const int32 Size = InCollection1.NumElements(FGeometryCollection::TransformGroup);
+			if (TManagedArray<int32>* TransformToGeometryIndex = InCollection1.ModifyAttributeTyped<int32>("TransformToGeometryIndex", FGeometryCollection::TransformGroup))
+			{
+				for (int32 Idx = OtherSize; Idx < Size; ++Idx)
+				{
+					if ((*TransformToGeometryIndex)[Idx] != INDEX_NONE)
+					{
+						(*TransformToGeometryIndex)[Idx] += GeometryOffset;
+					}
+				}
+			}
+			if (TManagedArray<int32>* Parent = InCollection1.ModifyAttributeTyped<int32>("Parent", FGeometryCollection::TransformGroup))
+			{
+				for (int32 Idx = OtherSize; Idx < Size; ++Idx)
+				{
+					if ((*Parent)[Idx] != INDEX_NONE)
+					{
+						(*Parent)[Idx] += OtherSize;
+					}
+				}
+			}
+			if (TManagedArray<TSet<int32>>* Children = InCollection1.ModifyAttributeTyped<TSet<int32>>("Children", FGeometryCollection::TransformGroup))
+			{
+				for (int32 Idx = OtherSize; Idx < Size; ++Idx)
+				{
+					for (int32& Child : (*Children)[Idx])
+					{
+						if (Child != INDEX_NONE)
+						{
+							Child += OtherSize;
+						}
+					}
+				}
+			}
+		}
+
+
 		SetValue(Context, MoveTemp(InCollection1), &Collection1);
 		if (const TManagedArray<FString>* GuidArray2 = InCollection2.FindAttribute<FString>("Guid", FGeometryCollection::GeometryGroup))
 		{
