Dataflow:
- Added sphere covering debug draw to convex nodes

diff --git a/Engine/Plugins/Experimental/GeometryCollectionPlugin/Source/GeometryCollectionNodes/Private/Dataflow/GeometryCollectionUtilityNodes.cpp b/Engine/Plugins/Experimental/GeometryCollectionPlugin/Source/GeometryCollectionNodes/Private/Dataflow/GeometryCollectionUtilityNodes.cpp
--- a/Engine/Plugins/Experimental/GeometryCollectionPlugin/Source/GeometryCollectionNodes/Private/Dataflow/GeometryCollectionUtilityNodes.cpp
+++ b/Engine/Plugins/Experimental/GeometryCollectionPlugin/Source/GeometryCollectionNodes/Private/Dataflow/GeometryCollectionUtilityNodes.cpp
@@ -119,6 +119,50 @@ namespace UE::Dataflow::Convex
 			}
 		}
 	}
+
+	static void SphereCoveringDebugDrawProc(IDataflowDebugDrawInterface& DataflowRenderingInterface, const FDataflowSphereCovering& OutSpheres, const FDataflowNodeSphereCoveringDebugDrawSettings& SphereCoveringDebugDrawRenderSettings)
+	{
+		constexpr int32 CMaxNumberOfSpheres = 500;
+		using namespace UE::Geometry;
+
+		const int32 NumSpheres = OutSpheres.Spheres.Num();
+		if (NumSpheres > 0)
+		{
+			DataflowRenderingInterface.SetLineWidth(SphereCoveringDebugDrawRenderSettings.LineWidthMultiplier);
+			if (SphereCoveringDebugDrawRenderSettings.RenderType == EDataflowDebugDrawRenderType::Shaded)
+			{
+				DataflowRenderingInterface.SetShaded(true);
+				DataflowRenderingInterface.SetTranslucent(SphereCoveringDebugDrawRenderSettings.bTranslucent);
+				DataflowRenderingInterface.SetWireframe(true);
+			}
+			else
+			{
+				DataflowRenderingInterface.SetShaded(false);
+				DataflowRenderingInterface.SetWireframe(true);
+			}
+			DataflowRenderingInterface.SetWorldPriority();
+			DataflowRenderingInterface.SetColor(SphereCoveringDebugDrawRenderSettings.Color);
+
+			int32 N = 1;
+			if (NumSpheres > CMaxNumberOfSpheres)
+			{
+				N = NumSpheres / CMaxNumberOfSpheres;
+			}
+
+			for (int32 Idx = 0; Idx < NumSpheres; ++Idx)
+			{
+				if (Idx % N == 0)
+				{
+					if (SphereCoveringDebugDrawRenderSettings.bRandomizeColor)
+					{
+						DataflowRenderingInterface.SetColor(UE::Dataflow::Convex::GetRandomColor(SphereCoveringDebugDrawRenderSettings.ColorRandomSeed + 7, Idx));
+					}
+
+					DataflowRenderingInterface.DrawSphere(OutSpheres.Spheres.GetCenter(Idx), OutSpheres.Spheres.GetRadius(Idx));
+				}
+			}
+		}
+	}
 }
 
 FCreateLeafConvexHullsDataflowNode::FCreateLeafConvexHullsDataflowNode(const UE::Dataflow::FNodeParameters& InParam, FGuid InGuid)
@@ -214,6 +258,13 @@ void FCreateLeafConvexHullsDataflowNode::DebugDraw(UE::Dataflow::FContext& Conte
 			const FDataflowTransformSelection& InOptionalSelectionFilter = GetValue(Context, &OptionalSelectionFilter);
 			UE::Dataflow::Convex::DebugDrawProc(DataflowRenderingInterface, InCollection, bRandomizeColor, ColorRandomSeed, InOptionalSelectionFilter);
 		}
+
+		if (const FDataflowOutput* SphereCoveringOutput = FindOutput(&SphereCovering))
+		{
+			const FDataflowSphereCovering& OutSpheres = SphereCoveringOutput->GetValue(Context, SphereCovering);
+
+			UE::Dataflow::Convex::SphereCoveringDebugDrawProc(DataflowRenderingInterface, OutSpheres, SphereCoveringDebugDrawRenderSettings);
+		}
 	}
 }
 #endif
@@ -494,6 +545,13 @@ void FGenerateClusterConvexHullsFromLeafHullsDataflowNode::DebugDraw(UE::Dataflo
 
 			UE::Dataflow::Convex::DebugDrawProc(DataflowRenderingInterface, InCollection, bRandomizeColor, ColorRandomSeed, InOptionalSelectionFilter);
 		}
+
+		if (const FDataflowOutput* SphereCoveringOutput = FindOutput(&SphereCovering))
+		{
+			const FDataflowSphereCovering& OutSpheres = SphereCoveringOutput->GetValue(Context, SphereCovering);
+
+			UE::Dataflow::Convex::SphereCoveringDebugDrawProc(DataflowRenderingInterface, OutSpheres, SphereCoveringDebugDrawRenderSettings);
+		}
 	}
 }
 #endif
@@ -617,6 +675,13 @@ void FGenerateClusterConvexHullsFromChildrenHullsDataflowNode::DebugDraw(UE::Dat
 
 			UE::Dataflow::Convex::DebugDrawProc(DataflowRenderingInterface, InCollection, bRandomizeColor, ColorRandomSeed, InOptionalSelectionFilter);
 		}
+
+		if (const FDataflowOutput* SphereCoveringOutput = FindOutput(&SphereCovering))
+		{
+			const FDataflowSphereCovering& OutSpheres = SphereCoveringOutput->GetValue(Context, SphereCovering);
+
+			UE::Dataflow::Convex::SphereCoveringDebugDrawProc(DataflowRenderingInterface, OutSpheres, SphereCoveringDebugDrawRenderSettings);
+		}
 	}
 }
 #endif
@@ -753,6 +818,13 @@ void FMergeConvexHullsDataflowNode::DebugDraw(UE::Dataflow::FContext& Context, I
 
 			UE::Dataflow::Convex::DebugDrawProc(DataflowRenderingInterface, InCollection, bRandomizeColor, ColorRandomSeed, InOptionalSelectionFilter);
 		}
+
+		if (const FDataflowOutput* SphereCoveringOutput = FindOutput(&SphereCovering))
+		{
+			const FDataflowSphereCovering& OutSpheres = SphereCoveringOutput->GetValue(Context, SphereCovering);
+
+			UE::Dataflow::Convex::SphereCoveringDebugDrawProc(DataflowRenderingInterface, OutSpheres, SphereCoveringDebugDrawRenderSettings);
+		}
 	}
 }
 #endif
diff --git a/Engine/Plugins/Experimental/GeometryCollectionPlugin/Source/GeometryCollectionNodes/Public/Dataflow/GeometryCollectionUtilityNodes.h b/Engine/Plugins/Experimental/GeometryCollectionPlugin/Source/GeometryCollectionNodes/Public/Dataflow/GeometryCollectionUtilityNodes.h
--- a/Engine/Plugins/Experimental/GeometryCollectionPlugin/Source/GeometryCollectionNodes/Public/Dataflow/GeometryCollectionUtilityNodes.h
+++ b/Engine/Plugins/Experimental/GeometryCollectionPlugin/Source/GeometryCollectionNodes/Public/Dataflow/GeometryCollectionUtilityNodes.h
@@ -203,9 +203,12 @@ private:
 	bool bRandomizeColor = true;
 
 	/** Random seed */
-	UPROPERTY(EditAnywhere, Category = "Debug Draw", meta = (ClampMin = "0"))
+	UPROPERTY(EditAnywhere, Category = "Debug Draw", meta = (ClampMin = "0", EditCondition = "bRandomizeColor==true", EditConditionHides))
 	int32 ColorRandomSeed = 0;
 
+	UPROPERTY(EditAnywhere, Category = "Debug Draw - Sphere Covering", meta = (ShowOnlyInnerProperties))
+	FDataflowNodeSphereCoveringDebugDrawSettings SphereCoveringDebugDrawRenderSettings;
+
 	virtual void Evaluate(UE::Dataflow::FContext& Context, const FDataflowOutput* Out) const override;
 #if WITH_EDITOR
 	virtual bool CanDebugDraw() const override
@@ -261,7 +264,7 @@ private:
 	bool bRandomizeColor = true;
 
 	/** Random seed */
-	UPROPERTY(EditAnywhere, Category = "Debug Draw", meta = (ClampMin = "0"))
+	UPROPERTY(EditAnywhere, Category = "Debug Draw", meta = (ClampMin = "0" , EditCondition = "bRandomizeColor==true", EditConditionHides))
 	int32 ColorRandomSeed = 0;
 
 	virtual void Evaluate(UE::Dataflow::FContext& Context, const FDataflowOutput* Out) const override;
@@ -320,7 +323,7 @@ private:
 	bool bRandomizeColor = true;
 
 	/** Random seed */
-	UPROPERTY(EditAnywhere, Category = "Debug Draw", meta = (ClampMin = "0"))
+	UPROPERTY(EditAnywhere, Category = "Debug Draw", meta = (ClampMin = "0", EditCondition = "bRandomizeColor==true", EditConditionHides))
 	int32 ColorRandomSeed = 0;
 
 	virtual void Evaluate(UE::Dataflow::FContext& Context, const FDataflowOutput* Out) const override;
@@ -446,9 +449,12 @@ private:
 	bool bRandomizeColor = true;
 
 	/** Random seed */
-	UPROPERTY(EditAnywhere, Category = "Debug Draw", meta = (ClampMin = "0"))
+	UPROPERTY(EditAnywhere, Category = "Debug Draw", meta = (ClampMin = "0", EditCondition = "bRandomizeColor==true", EditConditionHides))
 	int32 ColorRandomSeed = 0;
 
+	UPROPERTY(EditAnywhere, Category = "Debug Draw - Sphere Covering", meta = (ShowOnlyInnerProperties))
+	FDataflowNodeSphereCoveringDebugDrawSettings SphereCoveringDebugDrawRenderSettings;
+
 	virtual void Evaluate(UE::Dataflow::FContext& Context, const FDataflowOutput* Out) const override;
 #if WITH_EDITOR
 	virtual bool CanDebugDraw() const override
@@ -557,9 +563,12 @@ private:
 	bool bRandomizeColor = true;
 
 	/** Random seed */
-	UPROPERTY(EditAnywhere, Category = "Debug Draw", meta = (ClampMin = "0"))
+	UPROPERTY(EditAnywhere, Category = "Debug Draw", meta = (ClampMin = "0", EditCondition = "bRandomizeColor==true", EditConditionHides))
 	int32 ColorRandomSeed = 0;
 
+	UPROPERTY(EditAnywhere, Category = "Debug Draw - Sphere Covering", meta = (ShowOnlyInnerProperties))
+	FDataflowNodeSphereCoveringDebugDrawSettings SphereCoveringDebugDrawRenderSettings;
+
 	virtual void Evaluate(UE::Dataflow::FContext& Context, const FDataflowOutput* Out) const override;
 #if WITH_EDITOR
 	virtual bool CanDebugDraw() const override
@@ -690,9 +699,12 @@ private:
 	bool bRandomizeColor = true;
 
 	/** Random seed */
-	UPROPERTY(EditAnywhere, Category = "Debug Draw", meta = (ClampMin = "0"))
+	UPROPERTY(EditAnywhere, Category = "Debug Draw", meta = (ClampMin = "0", EditCondition = "bRandomizeColor==true", EditConditionHides))
 	int32 ColorRandomSeed = 0;
 
+	UPROPERTY(EditAnywhere, Category = "Debug Draw - Sphere Covering", meta = (ShowOnlyInnerProperties))
+	FDataflowNodeSphereCoveringDebugDrawSettings SphereCoveringDebugDrawRenderSettings;
+
 	virtual void Evaluate(UE::Dataflow::FContext& Context, const FDataflowOutput* Out) const override;
 #if WITH_EDITOR
 	virtual bool CanDebugDraw() const override
@@ -764,7 +776,7 @@ private:
 	bool bRandomizeColor = true;
 
 	/** Random seed */
-	UPROPERTY(EditAnywhere, Category = "Debug Draw", meta = (ClampMin = "0"))
+	UPROPERTY(EditAnywhere, Category = "Debug Draw", meta = (ClampMin = "0", EditCondition = "bRandomizeColor==true", EditConditionHides))
 	int32 ColorRandomSeed = 0;
 
 	virtual void Evaluate(UE::Dataflow::FContext& Context, const FDataflowOutput* Out) const override;
diff --git a/Engine/Source/Runtime/Experimental/Dataflow/Engine/Public/Dataflow/DataflowDebugDraw.h b/Engine/Source/Runtime/Experimental/Dataflow/Engine/Public/Dataflow/DataflowDebugDraw.h
--- a/Engine/Source/Runtime/Experimental/Dataflow/Engine/Public/Dataflow/DataflowDebugDraw.h
+++ b/Engine/Source/Runtime/Experimental/Dataflow/Engine/Public/Dataflow/DataflowDebugDraw.h
@@ -92,18 +92,51 @@ struct FDataflowNodeDebugDrawSettings
 	UPROPERTY(EditAnywhere, Category = "Debug Draw")
 	EDataflowDebugDrawRenderType RenderType = EDataflowDebugDrawRenderType::Wireframe;
 
-	UPROPERTY(EditAnywhere, Category = "Debug Draw", meta = (EditCondition = "RenderType==EDataflowDebugDrawRenderType::Shaded"));
+	UPROPERTY(EditAnywhere, Category = "Debug Draw", meta = (EditCondition = "RenderType==EDataflowDebugDrawRenderType::Shaded", EditConditionHides));
 	bool bTranslucent = true;
 
 	UPROPERTY(EditAnywhere, Category = "Debug Draw");
 	FLinearColor Color = FLinearColor::Gray;
 
-	UPROPERTY(EditAnywhere, Category = "Debug Draw", meta = (UIMin = "0.1", UIMax = "10.0", ClampMin = "0.1", ClampMax = "10.0"));
+	UPROPERTY(EditAnywhere, Category = "Debug Draw", meta = (ClampMin = "0.1", ClampMax = "10.0"));
 	float LineWidthMultiplier = 1.f;
 
 	DATAFLOWENGINE_API void SetDebugDrawSettings(IDataflowDebugDrawInterface& DataflowRenderingInterface) const;
 };
 
+/**
+ * SphereCovering DebugDraw basic common settings
+ */
+USTRUCT()
+struct FDataflowNodeSphereCoveringDebugDrawSettings
+{
+	GENERATED_USTRUCT_BODY()
+
+	/** Display sphere covering */
+	UPROPERTY(EditAnywhere, Category = "Debug Draw - Sphere Covering")
+	bool bDisplaySphereCovering = false;
+
+	UPROPERTY(EditAnywhere, Category = "Debug Draw - Sphere Covering")
+	EDataflowDebugDrawRenderType RenderType = EDataflowDebugDrawRenderType::Wireframe;
+
+	UPROPERTY(EditAnywhere, Category = "Debug Draw - Sphere Covering", meta = (EditCondition = "RenderType==EDataflowDebugDrawRenderType::Shaded", EditConditionHides));
+	bool bTranslucent = true;
+
+	UPROPERTY(EditAnywhere, Category = "Debug Draw - Sphere Covering");
+	FLinearColor Color = FLinearColor::Red;
+
+	UPROPERTY(EditAnywhere, Category = "Debug Draw - Sphere Covering", meta = (ClampMin = "0.1", ClampMax = "10.0"));
+	float LineWidthMultiplier = .25f;
+
+	/** Randomize color per sphere */
+	UPROPERTY(EditAnywhere, Category = "Debug Draw - Sphere Covering")
+	bool bRandomizeColor = false;
+
+	/** Random seed */
+	UPROPERTY(EditAnywhere, Category = "Debug Draw - Sphere Covering", meta = (ClampMin = "0", EditCondition = "bRandomizeColor==true", EditConditionHides))
+	int32 ColorRandomSeed = 0;
+};
+
 static void SetWireframeRender(IDataflowDebugDrawInterface& DataflowRenderingInterface)
 {
 	DataflowRenderingInterface.SetShaded(false);
