Dataflow:
- Added optional TransformSelection input to FTransformCollectionDataflowNode
- Added debug draw to display connectivity to Proximity and AutoCluster Dataflow nodes

diff --git a/Engine/Plugins/Experimental/GeometryCollectionPlugin/Source/GeometryCollectionNodes/Private/Dataflow/GeometryCollectionClusteringNodes.cpp b/Engine/Plugins/Experimental/GeometryCollectionPlugin/Source/GeometryCollectionNodes/Private/Dataflow/GeometryCollectionClusteringNodes.cpp
--- a/Engine/Plugins/Experimental/GeometryCollectionPlugin/Source/GeometryCollectionNodes/Private/Dataflow/GeometryCollectionClusteringNodes.cpp
+++ b/Engine/Plugins/Experimental/GeometryCollectionPlugin/Source/GeometryCollectionNodes/Private/Dataflow/GeometryCollectionClusteringNodes.cpp
@@ -30,6 +30,12 @@
 #include "GeometryCollection/GeometryCollectionProximityUtility.h"
 #include "FractureEngineClustering.h"
 #include "FractureEngineSelection.h"
+#include "Dataflow/DataflowDebugDrawInterface.h"
+#include "Dataflow/DataflowDebugDraw.h"
+#if WITH_EDITOR
+#include "Dataflow/DataflowRenderingViewMode.h"
+#endif
+#include "Dataflow/DataflowEngineUtil.h"
 
 #include UE_INLINE_GENERATED_CPP_BY_NAME(GeometryCollectionClusteringNodes)
 
@@ -84,6 +90,8 @@ FAutoClusterDataflowNode::FAutoClusterDataflowNode(const UE::Dataflow::FNodePara
 	RegisterOutputConnection(&Collection, &Collection);
 }
 
+/* ---------------------------------------------------------------------------------------------------------------------------------*/
+
 void FAutoClusterDataflowNode::Evaluate(UE::Dataflow::FContext& Context, const FDataflowOutput* Out) const
 {
 	if (Out->IsA<FManagedArrayCollection>(&Collection))
@@ -129,6 +137,9 @@ void FAutoClusterDataflowNode::Evaluate(UE::Dataflow::FContext& Context, const F
 					InEnforceSiteParameters,
 					InGridX, InGridY, InGridZ, InMinimumClusterSize, InKMeansIterations,
 					bInPreferConvexity, InConcavityTolerance);
+
+				FGeometryCollectionProximityUtility ProximityUtility(GeomCollection.Get());
+				ProximityUtility.UpdateProximity();
 			}
 			else
 			{
@@ -139,6 +150,63 @@ void FAutoClusterDataflowNode::Evaluate(UE::Dataflow::FContext& Context, const F
 	}
 }
 
+#if WITH_EDITOR
+bool FAutoClusterDataflowNode::CanDebugDrawViewMode(const FName& ViewModeName) const
+{
+	return ViewModeName == UE::Dataflow::FDataflowConstruction3DViewMode::Name;
+}
+
+void FAutoClusterDataflowNode::DebugDraw(UE::Dataflow::FContext& Context, IDataflowDebugDrawInterface& DataflowRenderingInterface, const FDebugDrawParameters& DebugDrawParameters) const
+{
+	using namespace UE::Geometry;
+
+	if ((DebugDrawParameters.bNodeIsSelected || DebugDrawParameters.bNodeIsPinned))
+	{
+		if (const FDataflowOutput* Output = FindOutput(&Collection))
+		{
+			const FManagedArrayCollection& OutCollection = Output->GetValue(Context, Collection);
+
+			DataflowRenderingInterface.SetLineWidth(LineWidthMultiplier);
+			DataflowRenderingInterface.SetWireframe(true);
+			DataflowRenderingInterface.SetWorldPriority();
+			DataflowRenderingInterface.SetColor(Color);
+
+			if (OutCollection.HasAttribute("Proximity", FGeometryCollection::GeometryGroup) &&
+				OutCollection.HasAttribute("BoundingBox", FGeometryCollection::GeometryGroup))
+			{
+				const TManagedArray<TSet<int32>>& Proximity = OutCollection.GetAttribute<TSet<int32>>("Proximity", FGeometryCollection::GeometryGroup);
+				const TManagedArray<FBox>& BoundingBox = OutCollection.GetAttribute<FBox>("BoundingBox", FGeometryCollection::GeometryGroup);
+
+				const int32 NumGeometry = OutCollection.NumElements(FGeometryCollection::GeometryGroup);
+				for (int32 Idx = 0; Idx < NumGeometry; ++Idx)
+				{
+					if (Proximity[Idx].Num() > 0)
+					{
+						const FVector Center = BoundingBox[Idx].GetCenter();
+
+						if (bRandomizeColor)
+						{
+							DataflowRenderingInterface.SetColor(UE::Dataflow::Color::GetRandomColor(ColorRandomSeed + 17, Idx));
+						}
+
+						DataflowRenderingInterface.DrawSphere(Center, 1.f);
+
+						for (int32 IdxOther : Proximity[Idx])
+						{
+							const FVector CenterOther = BoundingBox[IdxOther].GetCenter();
+
+							DataflowRenderingInterface.DrawSphere(CenterOther, 1.f);
+							DataflowRenderingInterface.DrawLine(Center, CenterOther);
+						}
+					}
+				}
+			}
+		}
+	}
+}
+#endif
+
+/* ---------------------------------------------------------------------------------------------------------------------------------*/
 
 void FClusterFlattenDataflowNode::Evaluate(UE::Dataflow::FContext& Context, const FDataflowOutput* Out) const
 {
diff --git a/Engine/Plugins/Experimental/GeometryCollectionPlugin/Source/GeometryCollectionNodes/Private/Dataflow/GeometryCollectionNodes.cpp b/Engine/Plugins/Experimental/GeometryCollectionPlugin/Source/GeometryCollectionNodes/Private/Dataflow/GeometryCollectionNodes.cpp
--- a/Engine/Plugins/Experimental/GeometryCollectionPlugin/Source/GeometryCollectionNodes/Private/Dataflow/GeometryCollectionNodes.cpp
+++ b/Engine/Plugins/Experimental/GeometryCollectionPlugin/Source/GeometryCollectionNodes/Private/Dataflow/GeometryCollectionNodes.cpp
@@ -38,6 +38,11 @@
 #include "GeometryCollection/Facades/CollectionHierarchyFacade.h"
 #include "DynamicMesh/MeshTransforms.h"
 #include "DynamicMesh/DynamicMesh3.h"
+#include "Dataflow/DataflowDebugDrawInterface.h"
+#include "Dataflow/DataflowDebugDraw.h"
+#if WITH_EDITOR
+#include "Dataflow/DataflowRenderingViewMode.h"
+#endif
 
 #include UE_INLINE_GENERATED_CPP_BY_NAME(GeometryCollectionNodes)
 
@@ -537,7 +542,17 @@ void FTransformCollectionDataflowNode::Evaluate(UE::Dataflow::FContext& Context,
 			bInvertTransformation);
 
 		GeometryCollection::Facades::FCollectionTransformFacade TransformFacade(InCollection);
-		TransformFacade.Transform(NewTransform);
+
+		if (!IsConnected(&TransformSelection))
+		{
+			TransformFacade.Transform(NewTransform);
+		}
+		else
+		{
+			const FDataflowTransformSelection& InTransformSelection = GetValue(Context, &TransformSelection);
+
+			TransformFacade.Transform(NewTransform, InTransformSelection.AsArray());
+		}
 
 		SetValue(Context, MoveTemp(InCollection), &Collection);
 	}
@@ -899,6 +914,8 @@ void FSetDynamicStateDataflowNode::Evaluate(UE::Dataflow::FContext& Context, con
 
 //////////////////////////////////////////////////////////////////////////////////////
 
+/* ---------------------------------------------------------------------------------------------------------------------------------*/
+
 void FProximityDataflowNode::Evaluate(UE::Dataflow::FContext& Context, const FDataflowOutput* Out) const
 {
 	if (Out->IsA<FManagedArrayCollection>(&Collection))
@@ -935,6 +952,63 @@ void FProximityDataflowNode::Evaluate(UE::Dataflow::FContext& Context, const FDa
 	}
 }
 
+#if WITH_EDITOR
+bool FProximityDataflowNode::CanDebugDrawViewMode(const FName& ViewModeName) const
+{
+	return ViewModeName == UE::Dataflow::FDataflowConstruction3DViewMode::Name;
+}
+
+void FProximityDataflowNode::DebugDraw(UE::Dataflow::FContext& Context, IDataflowDebugDrawInterface& DataflowRenderingInterface, const FDebugDrawParameters& DebugDrawParameters) const
+{
+	using namespace UE::Geometry;
+
+	if ((DebugDrawParameters.bNodeIsSelected || DebugDrawParameters.bNodeIsPinned))
+	{
+		if (const FDataflowOutput* Output = FindOutput(&Collection))
+		{
+			const FManagedArrayCollection& OutCollection = Output->GetValue(Context, Collection);
+
+			DataflowRenderingInterface.SetLineWidth(LineWidthMultiplier);
+			DataflowRenderingInterface.SetWireframe(true);
+			DataflowRenderingInterface.SetWorldPriority();
+			DataflowRenderingInterface.SetColor(Color);
+
+			if (OutCollection.HasAttribute("Proximity", FGeometryCollection::GeometryGroup) &&
+				OutCollection.HasAttribute("BoundingBox", FGeometryCollection::GeometryGroup))
+			{
+				const TManagedArray<TSet<int32>>& Proximity = OutCollection.GetAttribute<TSet<int32>>("Proximity", FGeometryCollection::GeometryGroup);
+				const TManagedArray<FBox>& BoundingBox = OutCollection.GetAttribute<FBox>("BoundingBox", FGeometryCollection::GeometryGroup);
+
+				const int32 NumGeometry = OutCollection.NumElements(FGeometryCollection::GeometryGroup);
+				for (int32 Idx = 0; Idx < NumGeometry; ++Idx)
+				{
+					if (Proximity[Idx].Num() > 0)
+					{
+						const FVector Center = BoundingBox[Idx].GetCenter();
+
+						if (bRandomizeColor)
+						{
+							DataflowRenderingInterface.SetColor(UE::Dataflow::Color::GetRandomColor(ColorRandomSeed + 17, Idx));
+						}
+
+						DataflowRenderingInterface.DrawSphere(Center, 1.f);
+
+						for (int32 IdxOther : Proximity[Idx])
+						{
+							const FVector CenterOther = BoundingBox[IdxOther].GetCenter();
+
+							DataflowRenderingInterface.DrawSphere(CenterOther, 1.f);
+							DataflowRenderingInterface.DrawLine(Center, CenterOther);
+						}
+					}
+				}
+			}
+		}
+	}
+}
+#endif
+
+/* ---------------------------------------------------------------------------------------------------------------------------------*/
 
 void FCollectionSetPivotDataflowNode::Evaluate(UE::Dataflow::FContext& Context, const FDataflowOutput* Out) const
 {
diff --git a/Engine/Plugins/Experimental/GeometryCollectionPlugin/Source/GeometryCollectionNodes/Public/Dataflow/GeometryCollectionClusteringNodes.h b/Engine/Plugins/Experimental/GeometryCollectionPlugin/Source/GeometryCollectionNodes/Public/Dataflow/GeometryCollectionClusteringNodes.h
--- a/Engine/Plugins/Experimental/GeometryCollectionPlugin/Source/GeometryCollectionNodes/Public/Dataflow/GeometryCollectionClusteringNodes.h
+++ b/Engine/Plugins/Experimental/GeometryCollectionPlugin/Source/GeometryCollectionNodes/Public/Dataflow/GeometryCollectionClusteringNodes.h
@@ -40,8 +40,9 @@ struct FAutoClusterDataflowNode : public FDataflowNode
 {
 	GENERATED_USTRUCT_BODY()
 	DATAFLOW_NODE_DEFINE_INTERNAL(FAutoClusterDataflowNode, "AutoCluster", "GeometryCollection|Cluster", "")
+	DATAFLOW_NODE_RENDER_TYPE("SurfaceRender", FGeometryCollection::StaticType(), "Collection")
 
-public:
+private:
 	/** How to choose the size of the clusters to create */
 	UPROPERTY(EditAnywhere, Category = ClusterSize);
 	EClusterSizeMethodEnum ClusterSizeMethod = EClusterSizeMethodEnum::Dataflow_ClusterSizeMethod_ByNumber;
@@ -103,12 +104,35 @@ public:
 	FManagedArrayCollection Collection;
 
 	/** Bone selection for the clustering */
-	UPROPERTY(meta = (DataflowInput, DisplayName = "TransformSelection"))
+	UPROPERTY(meta = (DataflowInput, DisplayName = "TransformSelection", DataflowIntrinsic))
 	FDataflowTransformSelection TransformSelection;
 
-	FAutoClusterDataflowNode(const UE::Dataflow::FNodeParameters& InParam, FGuid InGuid = FGuid::NewGuid());
+	UPROPERTY(EditAnywhere, Category = "Debug Draw");
+	FLinearColor Color = FLinearColor::Yellow;
+
+	UPROPERTY(EditAnywhere, Category = "Debug Draw", meta = (ClampMin = "0.1", ClampMax = "10.0"));
+	float LineWidthMultiplier = 2.f;
+
+	/** Randomize color per connection */
+	UPROPERTY(EditAnywhere, Category = "Debug Draw")
+	bool bRandomizeColor = false;
+
+	/** Random seed */
+	UPROPERTY(EditAnywhere, Category = "Debug Draw", meta = (ClampMin = "0", EditCondition = "bRandomizeColor==true", EditConditionHides))
+	int32 ColorRandomSeed = 0;
 
 	virtual void Evaluate(UE::Dataflow::FContext& Context, const FDataflowOutput* Out) const override;
+
+#if WITH_EDITOR
+	virtual bool CanDebugDraw() const override
+	{
+		return true;
+	}
+	virtual bool CanDebugDrawViewMode(const FName& ViewModeName) const override;
+	virtual void DebugDraw(UE::Dataflow::FContext& Context, IDataflowDebugDrawInterface& DataflowRenderingInterface, const FDebugDrawParameters& DebugDrawParameters) const override;
+#endif
+public:
+	FAutoClusterDataflowNode(const UE::Dataflow::FNodeParameters& InParam, FGuid InGuid = FGuid::NewGuid());
 };
 
 
diff --git a/Engine/Plugins/Experimental/GeometryCollectionPlugin/Source/GeometryCollectionNodes/Public/Dataflow/GeometryCollectionNodes.h b/Engine/Plugins/Experimental/GeometryCollectionPlugin/Source/GeometryCollectionNodes/Public/Dataflow/GeometryCollectionNodes.h
--- a/Engine/Plugins/Experimental/GeometryCollectionPlugin/Source/GeometryCollectionNodes/Public/Dataflow/GeometryCollectionNodes.h
+++ b/Engine/Plugins/Experimental/GeometryCollectionPlugin/Source/GeometryCollectionNodes/Public/Dataflow/GeometryCollectionNodes.h
@@ -687,6 +687,10 @@ public:
 	UPROPERTY(meta = (DataflowInput, DataflowOutput, DataflowPassthrough = "Collection", DataflowIntrinsic))
 	FManagedArrayCollection Collection;
 
+	/** Transform selection for transforming */
+	UPROPERTY(meta = (DataflowInput))
+	FDataflowTransformSelection TransformSelection;
+
 	/** Translation */
 	UPROPERTY(EditAnywhere, Category = "Transform", meta = (DataflowInput));
 	FVector Translate = FVector(0.0);
@@ -727,6 +731,7 @@ public:
 		: FDataflowNode(InParam, InGuid)
 	{
 		RegisterInputConnection(&Collection);
+		RegisterInputConnection(&TransformSelection);
 		RegisterInputConnection(&Translate);
 		RegisterInputConnection(&Rotate);
 		RegisterInputConnection(&Scale);
@@ -1235,8 +1240,9 @@ struct FProximityDataflowNode : public FDataflowNode
 {
 	GENERATED_USTRUCT_BODY()
 	DATAFLOW_NODE_DEFINE_INTERNAL(FProximityDataflowNode, "Proximity", "GeometryCollection|Utilities", "")
+	DATAFLOW_NODE_RENDER_TYPE("SurfaceRender", FGeometryCollection::StaticType(), "Collection")
 
-public:
+private:
 	/** Which method to use to decide whether a given piece of geometry is in proximity with another */
 	UPROPERTY(EditAnywhere, Category = "Proximity");
 	EProximityMethodEnum ProximityMethod = EProximityMethodEnum::Dataflow_ProximityMethod_Precise;
@@ -1271,6 +1277,30 @@ public:
 	UPROPERTY(meta = (DataflowInput, DataflowOutput, DataflowPassthrough = "Collection", DataflowIntrinsic))
 	FManagedArrayCollection Collection;
 
+	UPROPERTY(EditAnywhere, Category = "Debug Draw");
+	FLinearColor Color = FLinearColor::Yellow;
+
+	UPROPERTY(EditAnywhere, Category = "Debug Draw", meta = (ClampMin = "0.1", ClampMax = "10.0"));
+	float LineWidthMultiplier = 2.f;
+
+	/** Randomize color per connection */
+	UPROPERTY(EditAnywhere, Category = "Debug Draw")
+	bool bRandomizeColor = false;
+
+	/** Random seed */
+	UPROPERTY(EditAnywhere, Category = "Debug Draw", meta = (ClampMin = "0", EditCondition = "bRandomizeColor==true", EditConditionHides))
+	int32 ColorRandomSeed = 0;
+
+	virtual void Evaluate(UE::Dataflow::FContext& Context, const FDataflowOutput* Out) const override;
+#if WITH_EDITOR
+	virtual bool CanDebugDraw() const override
+	{
+		return true;
+	}
+	virtual bool CanDebugDrawViewMode(const FName& ViewModeName) const override;
+	virtual void DebugDraw(UE::Dataflow::FContext& Context, IDataflowDebugDrawInterface& DataflowRenderingInterface, const FDebugDrawParameters& DebugDrawParameters) const override;
+#endif
+public:
 	FProximityDataflowNode(const UE::Dataflow::FNodeParameters& InParam, FGuid InGuid = FGuid::NewGuid())
 		: FDataflowNode(InParam, InGuid)
 	{
@@ -1279,9 +1309,6 @@ public:
 		RegisterInputConnection(&ContactThreshold);
 		RegisterOutputConnection(&Collection, &Collection);
 	}
-
-	virtual void Evaluate(UE::Dataflow::FContext& Context, const FDataflowOutput* Out) const override;
-
 };
 
 
diff --git a/Engine/Source/Runtime/Experimental/Chaos/Private/GeometryCollection/Facades/CollectionTransformFacade.cpp b/Engine/Source/Runtime/Experimental/Chaos/Private/GeometryCollection/Facades/CollectionTransformFacade.cpp
--- a/Engine/Source/Runtime/Experimental/Chaos/Private/GeometryCollection/Facades/CollectionTransformFacade.cpp
+++ b/Engine/Source/Runtime/Experimental/Chaos/Private/GeometryCollection/Facades/CollectionTransformFacade.cpp
@@ -112,7 +112,18 @@ namespace GeometryCollection::Facades
 
 	void FCollectionTransformFacade::Transform(const FTransform& InTransform, const TArray<int32>& InSelection)
 	{
-		ensureAlwaysMsgf(false, TEXT("NOT YET IMPLEMENTED"));
+		if (TransformAttribute.IsValid())
+		{
+			TManagedArray<FTransform3f>& Transforms = TransformAttribute.Modify();
+
+			for (int32 Idx : InSelection)
+			{
+				if (Transforms.IsValidIndex(Idx))
+				{
+					Transforms[Idx] = Transforms[Idx] * FTransform3f(InTransform);
+				}
+			}
+		}
 	}
 
 	/**
diff --git a/Engine/Source/Runtime/Experimental/Dataflow/Engine/Private/Dataflow/DataflowEngineUtil.cpp b/Engine/Source/Runtime/Experimental/Dataflow/Engine/Private/Dataflow/DataflowEngineUtil.cpp
--- a/Engine/Source/Runtime/Experimental/Dataflow/Engine/Private/Dataflow/DataflowEngineUtil.cpp
+++ b/Engine/Source/Runtime/Experimental/Dataflow/Engine/Private/Dataflow/DataflowEngineUtil.cpp
@@ -43,5 +43,19 @@ namespace UE::Dataflow
 			}
 		}
 	}
+
+	namespace Color
+	{
+		FLinearColor GetRandomColor(const int32 RandomSeed, int32 Idx)
+		{
+			FRandomStream RandomStream(RandomSeed * 7 + Idx * 41);
+
+			const uint8 R = static_cast<uint8>(RandomStream.FRandRange(128, 255));
+			const uint8 G = static_cast<uint8>(RandomStream.FRandRange(128, 255));
+			const uint8 B = static_cast<uint8>(RandomStream.FRandRange(128, 255));
+
+			return FLinearColor(FColor(R, G, B, 255));
+		}
+	}
 }
 
diff --git a/Engine/Source/Runtime/Experimental/Dataflow/Engine/Public/Dataflow/DataflowEngineUtil.h b/Engine/Source/Runtime/Experimental/Dataflow/Engine/Public/Dataflow/DataflowEngineUtil.h
--- a/Engine/Source/Runtime/Experimental/Dataflow/Engine/Public/Dataflow/DataflowEngineUtil.h
+++ b/Engine/Source/Runtime/Experimental/Dataflow/Engine/Public/Dataflow/DataflowEngineUtil.h
@@ -74,4 +74,9 @@ namespace UE::Dataflow
 	{
 		void DATAFLOWENGINE_API GlobalTransforms(const FReferenceSkeleton&, TArray<FTransform>& GlobalTransforms);
 	}
+
+	namespace Color
+	{
+		DATAFLOWENGINE_API FLinearColor GetRandomColor(const int32 RandomSeed, int32 Idx);
+	}
 }
\ No newline at end of file
