add dataflow support for copying root collisions across collections

diff --git a/Engine/Plugins/Experimental/GeometryCollectionPlugin/Source/GeometryCollectionNodes/Private/Dataflow/GeometryCollectionUtilityNodes.cpp b/Engine/Plugins/Experimental/GeometryCollectionPlugin/Source/GeometryCollectionNodes/Private/Dataflow/GeometryCollectionUtilityNodes.cpp
--- a/Engine/Plugins/Experimental/GeometryCollectionPlugin/Source/GeometryCollectionNodes/Private/Dataflow/GeometryCollectionUtilityNodes.cpp
+++ b/Engine/Plugins/Experimental/GeometryCollectionPlugin/Source/GeometryCollectionNodes/Private/Dataflow/GeometryCollectionUtilityNodes.cpp
@@ -36,6 +36,7 @@ namespace UE::Dataflow
 		DATAFLOW_NODE_REGISTER_CREATION_FACTORY(FGenerateClusterConvexHullsFromLeafHullsDataflowNode);
 		DATAFLOW_NODE_REGISTER_CREATION_FACTORY(FGenerateClusterConvexHullsFromChildrenHullsDataflowNode);
 		DATAFLOW_NODE_REGISTER_CREATION_FACTORY(FClearConvexHullsDataflowNode);
+		DATAFLOW_NODE_REGISTER_CREATION_FACTORY(FCopyConvexHullsFromRootDataflowNode);
 		DATAFLOW_NODE_REGISTER_CREATION_FACTORY(FMergeConvexHullsDataflowNode);
 		DATAFLOW_NODE_REGISTER_CREATION_FACTORY(FUpdateVolumeAttributesDataflowNode);
 		DATAFLOW_NODE_REGISTER_CREATION_FACTORY(FGetConvexHullVolumeDataflowNode);
@@ -708,6 +709,80 @@ FMergeConvexHullsDataflowNode::FMergeConvexHullsDataflowNode(const UE::Dataflow:
 	DebugDrawRenderSettings.LineWidthMultiplier = 5.0;
 }
 
+#if WITH_EDITOR
+bool FCopyConvexHullsFromRootDataflowNode::CanDebugDrawViewMode(const FName& ViewModeName) const
+{
+	return ViewModeName == UE::Dataflow::FDataflowConstruction3DViewMode::Name;
+}
+
+void FCopyConvexHullsFromRootDataflowNode::DebugDraw(UE::Dataflow::FContext& Context, IDataflowDebugDrawInterface& DataflowRenderingInterface, const FDebugDrawParameters& DebugDrawParameters) const
+{
+	using namespace UE::Geometry;
+
+	if ((DebugDrawParameters.bNodeIsSelected || DebugDrawParameters.bNodeIsPinned))
+	{
+		DebugDrawRenderSettings.SetDebugDrawSettings(DataflowRenderingInterface);
+
+		if (const FDataflowOutput* Output = FindOutput(&Collection))
+		{
+			const FManagedArrayCollection& InCollection = Output->GetValue(Context, Collection);
+			FDataflowTransformSelection EmptySelection;
+
+			UE::Dataflow::Convex::DebugDrawProc(DataflowRenderingInterface, InCollection, bRandomizeColor, ColorRandomSeed, EmptySelection);
+		}
+	}
+}
+#endif
+
+
+FCopyConvexHullsFromRootDataflowNode::FCopyConvexHullsFromRootDataflowNode(const UE::Dataflow::FNodeParameters& InParam, FGuid InGuid)
+	: Super(InParam, InGuid)
+{
+	RegisterInputConnection(&Collection);
+	RegisterInputConnection(&FromCollection);
+
+	RegisterOutputConnection(&Collection, &Collection);
+	
+	DebugDrawRenderSettings.RenderType = EDataflowDebugDrawRenderType::Wireframe;
+	DebugDrawRenderSettings.Color = FLinearColor::Green;
+	DebugDrawRenderSettings.LineWidthMultiplier = 5.0;
+}
+
+void FCopyConvexHullsFromRootDataflowNode::Evaluate(UE::Dataflow::FContext& Context, const FDataflowOutput* Out) const
+{
+	if (Out->IsA(&Collection))
+	{
+		FManagedArrayCollection InCollection = GetValue(Context, &Collection);
+
+		if (IsConnected(&Collection) && IsConnected(&FromCollection))
+		{
+			const FManagedArrayCollection& InFromCollection = GetValue(Context, &FromCollection);
+
+			if (FGeometryCollectionConvexUtility::HasConvexHullData(&InFromCollection))
+			{
+				GeometryCollection::Facades::FCollectionTransformSelectionFacade ToTransformSelectionFacade(InCollection);
+				const TArray<int32> ToRoots = ToTransformSelectionFacade.SelectRootBones();
+				GeometryCollection::Facades::FCollectionTransformSelectionFacade FromTransformSelectionFacade(InFromCollection);
+				const TArray<int32> FromRoots = FromTransformSelectionFacade.SelectRootBones();
+				if (ToRoots.Num() != FromRoots.Num())
+				{
+					UE_LOG(LogChaosDataflow, Warning, TEXT("Failed to copy root collision across collections with different number of root nodes (%d vs %d)"), ToRoots.Num(), FromRoots.Num());
+				}
+				else
+				{
+					FGeometryCollectionConvexUtility::CopyConvexHulls(InCollection, ToRoots, InFromCollection, FromRoots);
+				}
+			}
+			else
+			{
+				UE_LOG(LogChaosDataflow, Warning, TEXT("Failed to copy root collision: FromCollection has no convex hull data"));
+			}
+		}
+		
+		SetValue(Context, MoveTemp(InCollection), &Collection);
+	}
+}
+
 void FClearConvexHullsDataflowNode::Evaluate(UE::Dataflow::FContext& Context, const FDataflowOutput* Out) const
 {
 	if (Out->IsA(&Collection))
diff --git a/Engine/Plugins/Experimental/GeometryCollectionPlugin/Source/GeometryCollectionNodes/Public/Dataflow/GeometryCollectionUtilityNodes.h b/Engine/Plugins/Experimental/GeometryCollectionPlugin/Source/GeometryCollectionNodes/Public/Dataflow/GeometryCollectionUtilityNodes.h
--- a/Engine/Plugins/Experimental/GeometryCollectionPlugin/Source/GeometryCollectionNodes/Public/Dataflow/GeometryCollectionUtilityNodes.h
+++ b/Engine/Plugins/Experimental/GeometryCollectionPlugin/Source/GeometryCollectionNodes/Public/Dataflow/GeometryCollectionUtilityNodes.h
@@ -614,6 +614,46 @@ public:
 };
 
 
+/** Copy convex hulls from given transforms on a source collection to transforms on the target collection */
+USTRUCT(meta = (DataflowGeometryCollection))
+struct FCopyConvexHullsFromRootDataflowNode : public FDataflowNode
+{
+	GENERATED_USTRUCT_BODY()
+	DATAFLOW_NODE_DEFINE_INTERNAL(FCopyConvexHullsFromRootDataflowNode, "CopyConvexHullsFromRoot", "GeometryCollection|Utilities", "")
+	DATAFLOW_NODE_RENDER_TYPE("SurfaceRender", FGeometryCollection::StaticType(), "Collection")
+
+public:
+	FCopyConvexHullsFromRootDataflowNode(const UE::Dataflow::FNodeParameters& InParam, FGuid InGuid = FGuid::NewGuid());
+
+private:
+	UPROPERTY(meta = (DataflowInput, DataflowOutput, DataflowPassthrough = "Collection", DataflowIntrinsic))
+	FManagedArrayCollection Collection;
+	
+	UPROPERTY(meta = (DataflowInput, DataflowIntrinsic))
+	FManagedArrayCollection FromCollection;
+
+	UPROPERTY(EditAnywhere, Category = "Debug Draw", meta = (ShowOnlyInnerProperties))
+	FDataflowNodeDebugDrawSettings DebugDrawRenderSettings;
+
+	/** Randomize color per convex hull */
+	UPROPERTY(EditAnywhere, Category = "Debug Draw")
+	bool bRandomizeColor = true;
+
+	/** Random seed */
+	UPROPERTY(EditAnywhere, Category = "Debug Draw", meta = (ClampMin = "0"))
+	int32 ColorRandomSeed = 0;
+
+	virtual void Evaluate(UE::Dataflow::FContext& Context, const FDataflowOutput* Out) const override;
+#if WITH_EDITOR
+	virtual bool CanDebugDraw() const override
+	{
+		return true;
+	}
+	virtual bool CanDebugDrawViewMode(const FName& ViewModeName) const override;
+	virtual void DebugDraw(UE::Dataflow::FContext& Context, IDataflowDebugDrawInterface& DataflowRenderingInterface, const FDebugDrawParameters& DebugDrawParameters) const override;
+#endif
+
+};
 
 
 /** Merge convex hulls on transforms with multiple hulls */
diff --git a/Engine/Source/Runtime/Experimental/Chaos/Private/GeometryCollection/GeometryCollectionConvexUtility.cpp b/Engine/Source/Runtime/Experimental/Chaos/Private/GeometryCollection/GeometryCollectionConvexUtility.cpp
--- a/Engine/Source/Runtime/Experimental/Chaos/Private/GeometryCollection/GeometryCollectionConvexUtility.cpp
+++ b/Engine/Source/Runtime/Experimental/Chaos/Private/GeometryCollection/GeometryCollectionConvexUtility.cpp
@@ -2537,6 +2537,83 @@ Chaos::FConvexPtr FGeometryCollectionConvexUtility::GetConvexHull(const FGeometr
 }
 
 
+bool FGeometryCollectionConvexUtility::CopyConvexHulls(FManagedArrayCollection& Collection, const TArray<int32>& TransformIndices,
+	const FManagedArrayCollection& SourceCollection, const TArray<int32>& SourceTransformIndices)
+{
+	static FName ConvexGroupName = FGeometryCollection::ConvexGroup;
+	static FName ConvexHullAttributeName = FGeometryCollection::ConvexHullAttribute;
+	static FName TransformToConvexIndicesName("TransformToConvexIndices");
+
+	if (TransformIndices.Num() != SourceTransformIndices.Num())
+	{
+		UE_LOG(LogChaos, Warning,
+			TEXT("Attempted to copy convex hulls failed because the source and target transform indices array did not have the same length (%d vs %d)"), 
+			TransformIndices.Num(), SourceTransformIndices.Num());
+		return false;
+	}
+
+	CreateConvexHullAttributesIfNeeded(Collection);
+
+	TManagedArrayAccessor<Chaos::FConvexPtr> ToConvexHullAttribute(Collection, ConvexHullAttributeName, ConvexGroupName);
+	TManagedArrayAccessor<Chaos::FConvexPtr> FromConvexHullAttribute(SourceCollection, ConvexHullAttributeName, ConvexGroupName);
+	TManagedArrayAccessor<TSet<int32>> ToTransformToConvexIndicesAttribute(Collection, TransformToConvexIndicesName, FTransformCollection::TransformGroup);
+	TManagedArrayAccessor<TSet<int32>> FromTransformToConvexIndicesAttribute(SourceCollection, TransformToConvexIndicesName, FTransformCollection::TransformGroup);
+
+	if (!FromConvexHullAttribute.IsValid() || !FromTransformToConvexIndicesAttribute.IsValid())
+	{
+		UE_LOG(LogChaos, Warning,
+			TEXT("Failed to copy convex hulls: Hull attributes not found on source collection"));
+		return false;
+	}
+
+	TManagedArray<TSet<int32>>& ToTransformToConvexIndices = ToTransformToConvexIndicesAttribute.Modify();
+	TManagedArray<Chaos::FConvexPtr>& ToConvexHull = ToConvexHullAttribute.Modify();
+	
+	TSet<int32> ConvexToRemove;
+	for (int32 ToTransferIdx = 0; ToTransferIdx < TransformIndices.Num(); ++ToTransferIdx)
+	{
+		int32 FromTransformIdx = SourceTransformIndices[ToTransferIdx];
+		int32 ToTransformIdx = TransformIndices[ToTransferIdx];
+		if (!ToTransformToConvexIndices.IsValidIndex(ToTransformIdx))
+		{
+			UE_LOG(LogChaos, Warning,
+				TEXT("Invalid transform index: %d"),
+				ToTransformIdx);
+			continue;
+		}
+		if (!FromTransformToConvexIndicesAttribute.Get().IsValidIndex(FromTransformIdx))
+		{
+			UE_LOG(LogChaos, Warning,
+				TEXT("Invalid source transform index: %d"),
+				FromTransformIdx);
+			continue;
+		}
+		const TSet<int32>& FromIndices = FromTransformToConvexIndicesAttribute[FromTransformIdx];
+		ConvexToRemove.Append(ToTransformToConvexIndices[ToTransformIdx]);
+		TSet<int32>& ToConvexIndices = ToTransformToConvexIndices[ToTransformIdx];
+		ToConvexIndices.Reset();
+		
+		int32 NewConvexIndexStart = Collection.AddElements(FromIndices.Num(), FGeometryCollection::ConvexGroup);
+		for (int32 FromConvexIdx : FromIndices)
+		{
+			int32 ToConvexIdx = NewConvexIndexStart++;
+			ToConvexIndices.Add(ToConvexIdx);
+			ToConvexHull[ToConvexIdx] = FromConvexHullAttribute[FromConvexIdx];
+		}
+	}
+	TArray<int32> ToRemoveArr = ConvexToRemove.Array();
+	if (!ToRemoveArr.IsEmpty())
+	{
+		ToRemoveArr.Sort();
+		FManagedArrayCollection::FProcessingParameters Params;
+		Params.bDoValidation = false; // for perf reasons
+		Collection.RemoveElements(FGeometryCollection::ConvexGroup, ToRemoveArr, Params);
+	}
+
+	return true;
+}
+
+
 void FGeometryCollectionConvexUtility::RemoveConvexHulls(FManagedArrayCollection* GeometryCollection, const TArray<int32>& TransformsToClearHullsFrom)
 {
 	if (GeometryCollection->HasGroup(FGeometryCollection::ConvexGroup) && GeometryCollection->HasAttribute("TransformToConvexIndices", FTransformCollection::TransformGroup))
diff --git a/Engine/Source/Runtime/Experimental/Chaos/Public/GeometryCollection/GeometryCollectionConvexUtility.h b/Engine/Source/Runtime/Experimental/Chaos/Public/GeometryCollection/GeometryCollectionConvexUtility.h
--- a/Engine/Source/Runtime/Experimental/Chaos/Public/GeometryCollection/GeometryCollectionConvexUtility.h
+++ b/Engine/Source/Runtime/Experimental/Chaos/Public/GeometryCollection/GeometryCollectionConvexUtility.h
@@ -3,6 +3,7 @@
 #pragma once
 
 #include "CoreMinimal.h"
+#include "GeometryCollection/ManagedArrayAccessor.h"
 #include "GeometryCollection/ManagedArrayCollection.h"
 #include "Chaos/ImplicitFwd.h"
 #include "CompGeom/ConvexDecomposition3.h"
@@ -274,6 +275,16 @@ public:
 		return nullptr;
 	}
 
+	/** 
+	 * Copy convex hulls on the specified transforms from a source collection to the collection 
+	 * @param GeometryCollection The collection to copy convex hulls to
+	 * @param TransformIndices The transforms on the GeometryCollection to update w/ copied convex hulls. Must be same length as SourceTransformIndices array.
+	 * @param SourceCollection The collection to copy convex hulls from
+	 * @param SourceTransformIndices The transforms on the SourceCollection to copy from. Must be same length as TransformIndices array.
+	 */
+	static CHAOS_API bool CopyConvexHulls(FManagedArrayCollection& Collection, const TArray<int32>& TransformIndices,
+		const FManagedArrayCollection& SourceCollection, const TArray<int32>& SourceTransformIndices);
+
 	/** Delete the convex hulls pointed at by the transform indices provided. */
 	static CHAOS_API void RemoveConvexHulls(FManagedArrayCollection* GeometryCollection, const TArray<int32>& TransformsToClearHullsFrom);
 
