ChaosFlesh: supported setting dynamics (in addition to kinematics) on vertex selection.

diff --git a/Engine/Plugins/Experimental/ChaosFlesh/Source/ChaosFleshNodes/Public/Dataflow/ChaosFleshVertexConstraintNode.h b/Engine/Plugins/Experimental/ChaosFlesh/Source/ChaosFleshNodes/Public/Dataflow/ChaosFleshVertexConstraintNode.h
--- a/Engine/Plugins/Experimental/ChaosFlesh/Source/ChaosFleshNodes/Public/Dataflow/ChaosFleshVertexConstraintNode.h
+++ b/Engine/Plugins/Experimental/ChaosFlesh/Source/ChaosFleshNodes/Public/Dataflow/ChaosFleshVertexConstraintNode.h
@@ -7,8 +7,10 @@
 
 #include "ChaosFleshVertexConstraintNode.generated.h"
 
-// @todo(deprecate), rename to FSetVerticesKinematicDataflowNode
-USTRUCT(meta = (DataflowFlesh))
+/**
+* Use SetKinematicVertexSelection to set/reset kinematic vertices
+*/
+USTRUCT(meta = (Deprecated = "5.6"))
 struct FSetVerticesKinematicDataflowNode : public FDataflowNode
 {
 	GENERATED_USTRUCT_BODY()
diff --git a/Engine/Plugins/Experimental/GeometryCollectionPlugin/Source/GeometryCollectionNodes/Private/Dataflow/GeometryCollectionTransferVertexAttributeNode.cpp b/Engine/Plugins/Experimental/GeometryCollectionPlugin/Source/GeometryCollectionNodes/Private/Dataflow/GeometryCollectionTransferVertexAttributeNode.cpp
--- a/Engine/Plugins/Experimental/GeometryCollectionPlugin/Source/GeometryCollectionNodes/Private/Dataflow/GeometryCollectionTransferVertexAttributeNode.cpp
+++ b/Engine/Plugins/Experimental/GeometryCollectionPlugin/Source/GeometryCollectionNodes/Private/Dataflow/GeometryCollectionTransferVertexAttributeNode.cpp
@@ -731,7 +731,8 @@ void FGeometryCollectionSetKinematicVertexSelectionNode::Evaluate(UE::Dataflow::
 		{
 			FDataflowVertexSelection VertexSelectionIn = GetValue<FDataflowVertexSelection>(Context, &VertexSelection);
 			GeometryCollection::Facades::FVertexBoneWeightsFacade VertexBoneWeightsFacade(InCollection);
-			VertexBoneWeightsFacade.SetVertexArrayKinematic(VertexSelectionIn.AsArray());
+			const bool bSetKinematic = (KinematicValue == ESetKinematicVertexSelectionKinematicValue::SetKinematic);
+			VertexBoneWeightsFacade.SetVertexArrayKinematic(VertexSelectionIn.AsArray(), bSetKinematic);
 		}
 		SetValue(Context, MoveTemp(InCollection), &Collection);
 	}
diff --git a/Engine/Plugins/Experimental/GeometryCollectionPlugin/Source/GeometryCollectionNodes/Public/Dataflow/GeometryCollectionTransferVertexAttributeNode.h b/Engine/Plugins/Experimental/GeometryCollectionPlugin/Source/GeometryCollectionNodes/Public/Dataflow/GeometryCollectionTransferVertexAttributeNode.h
--- a/Engine/Plugins/Experimental/GeometryCollectionPlugin/Source/GeometryCollectionNodes/Public/Dataflow/GeometryCollectionTransferVertexAttributeNode.h
+++ b/Engine/Plugins/Experimental/GeometryCollectionPlugin/Source/GeometryCollectionNodes/Public/Dataflow/GeometryCollectionTransferVertexAttributeNode.h
@@ -240,6 +240,16 @@ private:
 	virtual void Evaluate(UE::Dataflow::FContext& Context, const FDataflowOutput* Out) const override;
 };
 
+UENUM(BlueprintType)
+enum class ESetKinematicVertexSelectionKinematicValue : uint8
+{
+	/** Set vertices to be kinematic */
+	SetKinematic     UMETA(DisplayName = "Kinematic"),
+
+	/** Set vertices to be dynamic */
+	SetNonKinematic  UMETA(DisplayName = "Non-Kinematic")
+};
+
 /**
  * Set VertexSelection to be kinematic. Note that kinematic particles need skin weights.
  */
@@ -259,6 +269,9 @@ public:
 	UPROPERTY(meta = (DataflowInput, DisplayName = "VertexSelection"))
 	FDataflowVertexSelection VertexSelection;
 
+	UPROPERTY(EditAnywhere, Category = "Method", meta = (DisplayName = "VertexSelection"))
+	ESetKinematicVertexSelectionKinematicValue KinematicValue = ESetKinematicVertexSelectionKinematicValue::SetKinematic;
+
 	FGeometryCollectionSetKinematicVertexSelectionNode(const UE::Dataflow::FNodeParameters& InParam, FGuid InGuid = FGuid::NewGuid())
 		: FDataflowNode(InParam, InGuid)
 	{
diff --git a/Engine/Source/Runtime/Experimental/Chaos/Private/GeometryCollection/Facades/CollectionVertexBoneWeightsFacade.cpp b/Engine/Source/Runtime/Experimental/Chaos/Private/GeometryCollection/Facades/CollectionVertexBoneWeightsFacade.cpp
--- a/Engine/Source/Runtime/Experimental/Chaos/Private/GeometryCollection/Facades/CollectionVertexBoneWeightsFacade.cpp
+++ b/Engine/Source/Runtime/Experimental/Chaos/Private/GeometryCollection/Facades/CollectionVertexBoneWeightsFacade.cpp
@@ -123,15 +123,15 @@ namespace GeometryCollection::Facades
 		}
 	}
 
-	void FVertexBoneWeightsFacade::SetVertexKinematic(int32 VertexIndex)
+	void FVertexBoneWeightsFacade::SetVertexKinematic(int32 VertexIndex, bool Value)
 	{
 		if (KinematicAttribute.IsValid() && KinematicAttribute.IsValidIndex(VertexIndex))
 		{
-			KinematicAttribute.ModifyAt(VertexIndex, true);
+			KinematicAttribute.ModifyAt(VertexIndex, Value);
 		}
 	}
 
-	void FVertexBoneWeightsFacade::SetVertexArrayKinematic(const TArray<int32>& VertexIndices)
+	void FVertexBoneWeightsFacade::SetVertexArrayKinematic(const TArray<int32>& VertexIndices, bool Value)
 	{
 		if (KinematicAttribute.IsValid())
 		{
@@ -140,7 +140,7 @@ namespace GeometryCollection::Facades
 			{
 				if (KinematicArray.IsValidIndex(VertexIndex))
 				{
-					KinematicArray[VertexIndex] = true;
+					KinematicArray[VertexIndex] = Value;
 				}
 			}
 		}
diff --git a/Engine/Source/Runtime/Experimental/Chaos/Public/GeometryCollection/Facades/CollectionVertexBoneWeightsFacade.h b/Engine/Source/Runtime/Experimental/Chaos/Public/GeometryCollection/Facades/CollectionVertexBoneWeightsFacade.h
--- a/Engine/Source/Runtime/Experimental/Chaos/Public/GeometryCollection/Facades/CollectionVertexBoneWeightsFacade.h
+++ b/Engine/Source/Runtime/Experimental/Chaos/Public/GeometryCollection/Facades/CollectionVertexBoneWeightsFacade.h
@@ -59,9 +59,11 @@ namespace GeometryCollection::Facades
 		/** Modify bone weight based on the kinematic bindings. */
 		CHAOS_API void ModifyBoneWeight(int32 VertexIndex, TArray<int32> VertexBoneIndex, TArray<float> VertexBoneWeight);
 
-		/** Set vertex to be kinematic */
-		CHAOS_API void SetVertexKinematic(int32 VertexIndex);
-		CHAOS_API void SetVertexArrayKinematic(const TArray<int32>& VertexIndices);
+		/** Set vertex to be kinematic/dynamic */
+		CHAOS_API void SetVertexKinematic(int32 VertexIndex, bool Value = true);
+
+		/** Set vertex to be kinematic/dynamic */
+		CHAOS_API void SetVertexArrayKinematic(const TArray<int32>& VertexIndices, bool Value = true);
 
 		/** Return the vertex bone indices from the collection. Null if not initialized.  */
 		const TManagedArray< TArray<int32> >* FindBoneIndices()  const { return BoneIndexAttribute.Find(); }
