change default line width multiplier for convex hull vis from 5 to 2

diff --git a/Engine/Plugins/Experimental/GeometryCollectionPlugin/Source/GeometryCollectionNodes/Private/Dataflow/GeometryCollectionUtilityNodes.cpp b/Engine/Plugins/Experimental/GeometryCollectionPlugin/Source/GeometryCollectionNodes/Private/Dataflow/GeometryCollectionUtilityNodes.cpp
--- a/Engine/Plugins/Experimental/GeometryCollectionPlugin/Source/GeometryCollectionNodes/Private/Dataflow/GeometryCollectionUtilityNodes.cpp
+++ b/Engine/Plugins/Experimental/GeometryCollectionPlugin/Source/GeometryCollectionNodes/Private/Dataflow/GeometryCollectionUtilityNodes.cpp
@@ -178,7 +178,7 @@ FCreateLeafConvexHullsDataflowNode::FCreateLeafConvexHullsDataflowNode(const UE:
 
 	DebugDrawRenderSettings.RenderType = EDataflowDebugDrawRenderType::Wireframe;
 	DebugDrawRenderSettings.Color = FLinearColor::Green;
-	DebugDrawRenderSettings.LineWidthMultiplier = 5.0;
+	DebugDrawRenderSettings.LineWidthMultiplier = 2.0;
 }
 
 void FCreateLeafConvexHullsDataflowNode::Evaluate(UE::Dataflow::FContext& Context, const FDataflowOutput* Out) const
@@ -284,7 +284,7 @@ FSimplifyConvexHullsDataflowNode::FSimplifyConvexHullsDataflowNode(const UE::Dat
 
 	DebugDrawRenderSettings.RenderType = EDataflowDebugDrawRenderType::Wireframe;
 	DebugDrawRenderSettings.Color = FLinearColor::Green;
-	DebugDrawRenderSettings.LineWidthMultiplier = 5.0;
+	DebugDrawRenderSettings.LineWidthMultiplier = 2.0;
 }
 
 void FSimplifyConvexHullsDataflowNode::Evaluate(UE::Dataflow::FContext& Context, const FDataflowOutput* Out) const
@@ -359,7 +359,7 @@ FCreateNonOverlappingConvexHullsDataflowNode::FCreateNonOverlappingConvexHullsDa
 
 	DebugDrawRenderSettings.RenderType = EDataflowDebugDrawRenderType::Wireframe;
 	DebugDrawRenderSettings.Color = FLinearColor::Green;
-	DebugDrawRenderSettings.LineWidthMultiplier = 5.0;
+	DebugDrawRenderSettings.LineWidthMultiplier = 2.0;
 }
 
 void FCreateNonOverlappingConvexHullsDataflowNode::Evaluate(UE::Dataflow::FContext& Context, const FDataflowOutput* Out) const
@@ -448,7 +448,7 @@ FGenerateClusterConvexHullsFromLeafHullsDataflowNode::FGenerateClusterConvexHull
 
 	DebugDrawRenderSettings.RenderType = EDataflowDebugDrawRenderType::Wireframe;
 	DebugDrawRenderSettings.Color = FLinearColor::Green;
-	DebugDrawRenderSettings.LineWidthMultiplier = 5.0;
+	DebugDrawRenderSettings.LineWidthMultiplier = 2.0;
 }
 
 void FGenerateClusterConvexHullsFromLeafHullsDataflowNode::Evaluate(UE::Dataflow::FContext& Context, const FDataflowOutput* Out) const
@@ -578,7 +578,7 @@ FGenerateClusterConvexHullsFromChildrenHullsDataflowNode::FGenerateClusterConvex
 
 	DebugDrawRenderSettings.RenderType = EDataflowDebugDrawRenderType::Wireframe;
 	DebugDrawRenderSettings.Color = FLinearColor::Green;
-	DebugDrawRenderSettings.LineWidthMultiplier = 5.0;
+	DebugDrawRenderSettings.LineWidthMultiplier = 2.0;
 }
 
 void FGenerateClusterConvexHullsFromChildrenHullsDataflowNode::Evaluate(UE::Dataflow::FContext& Context, const FDataflowOutput* Out) const
@@ -706,7 +706,7 @@ FMergeConvexHullsDataflowNode::FMergeConvexHullsDataflowNode(const UE::Dataflow:
 
 	DebugDrawRenderSettings.RenderType = EDataflowDebugDrawRenderType::Wireframe;
 	DebugDrawRenderSettings.Color = FLinearColor::Green;
-	DebugDrawRenderSettings.LineWidthMultiplier = 5.0;
+	DebugDrawRenderSettings.LineWidthMultiplier = 2.0;
 }
 
 #if WITH_EDITOR
@@ -745,7 +745,7 @@ FCopyConvexHullsFromRootDataflowNode::FCopyConvexHullsFromRootDataflowNode(const
 	
 	DebugDrawRenderSettings.RenderType = EDataflowDebugDrawRenderType::Wireframe;
 	DebugDrawRenderSettings.Color = FLinearColor::Green;
-	DebugDrawRenderSettings.LineWidthMultiplier = 5.0;
+	DebugDrawRenderSettings.LineWidthMultiplier = 2.0;
 }
 
 void FCopyConvexHullsFromRootDataflowNode::Evaluate(UE::Dataflow::FContext& Context, const FDataflowOutput* Out) const
@@ -937,7 +937,7 @@ FGetConvexHullVolumeDataflowNode::FGetConvexHullVolumeDataflowNode(const UE::Dat
 
 	DebugDrawRenderSettings.RenderType = EDataflowDebugDrawRenderType::Wireframe;
 	DebugDrawRenderSettings.Color = FLinearColor::Green;
-	DebugDrawRenderSettings.LineWidthMultiplier = 5.0;
+	DebugDrawRenderSettings.LineWidthMultiplier = 2.0;
 }
 
 void FGetConvexHullVolumeDataflowNode::Evaluate(UE::Dataflow::FContext& Context, const FDataflowOutput* Out) const
