add point filtering w/ mesh dataflow node

diff --git a/Engine/Plugins/Experimental/GeometryCollectionPlugin/Source/GeometryCollectionNodes/Private/Dataflow/GeometryCollectionSamplingNodes.cpp b/Engine/Plugins/Experimental/GeometryCollectionPlugin/Source/GeometryCollectionNodes/Private/Dataflow/GeometryCollectionSamplingNodes.cpp
--- a/Engine/Plugins/Experimental/GeometryCollectionPlugin/Source/GeometryCollectionNodes/Private/Dataflow/GeometryCollectionSamplingNodes.cpp
+++ b/Engine/Plugins/Experimental/GeometryCollectionPlugin/Source/GeometryCollectionNodes/Private/Dataflow/GeometryCollectionSamplingNodes.cpp
@@ -3,6 +3,7 @@
 #include "Dataflow/GeometryCollectionSamplingNodes.h"
 
 #include "UDynamicMesh.h"
+#include "Spatial/FastWinding.h"
 
 #include UE_INLINE_GENERATED_CPP_BY_NAME(GeometryCollectionSamplingNodes)
 
@@ -14,6 +15,75 @@ namespace UE::Dataflow
 		DATAFLOW_NODE_REGISTER_CREATION_FACTORY(FUniformPointSamplingDataflowNode);
 		DATAFLOW_NODE_REGISTER_CREATION_FACTORY(FNonUniformPointSamplingDataflowNode);
 		DATAFLOW_NODE_REGISTER_CREATION_FACTORY(FVertexWeightedPointSamplingDataflowNode);
+		DATAFLOW_NODE_REGISTER_CREATION_FACTORY(FFilterPointSetWithMeshDataflowNode);
+	}
+}
+
+FFilterPointSetWithMeshDataflowNode::FFilterPointSetWithMeshDataflowNode(const UE::Dataflow::FNodeParameters& InParam, FGuid InGuid)
+	: FDataflowNode(InParam, InGuid)
+{
+	RegisterInputConnection(&TargetMesh);
+	RegisterInputConnection(&SamplePoints);
+	RegisterInputConnection(&bKeepInside).SetCanHidePin(true).SetPinIsHidden(true);
+	RegisterInputConnection(&WindingThreshold).SetCanHidePin(true).SetPinIsHidden(true);
+
+	RegisterOutputConnection(&SamplePoints);
+}
+
+void FFilterPointSetWithMeshDataflowNode::Evaluate(UE::Dataflow::FContext& Context, const FDataflowOutput* Out) const
+{
+	if (Out->IsA(&SamplePoints))
+	{
+		TArray<FVector> InSamplePoints = GetValue(Context, &SamplePoints);
+		double InWindingThreshold = (double)GetValue(Context, &WindingThreshold);
+		bool bInKeepInside = GetValue(Context, &bKeepInside);
+		if (TObjectPtr<const UDynamicMesh> InTargetMesh = GetValue(Context, &TargetMesh))
+		{
+			using namespace UE::Geometry;
+			InTargetMesh->ProcessMesh([&InSamplePoints, InWindingThreshold, bInKeepInside](const FDynamicMesh3& Mesh)
+				{
+					//  set up AABBTree and FWNTree lists
+					TMeshAABBTree3<FDynamicMesh3> Spatial(&Mesh);
+					TFastWindingTree<FDynamicMesh3> FastWinding(&Spatial);
+
+					// Filter points w/ the fast winding tree
+					TArray<bool> PointInside;
+					PointInside.SetNumZeroed(InSamplePoints.Num());
+					ParallelFor(PointInside.Num(), 
+						[&PointInside, &FastWinding, &InSamplePoints, InWindingThreshold]
+						(int32 PointIdx)
+						{
+							FVector Pt = InSamplePoints[PointIdx];
+							PointInside[PointIdx] = FastWinding.IsInside(Pt, InWindingThreshold);
+						}
+					);
+
+					// Move the points we're keeping to the front of the array, and trim the array
+					int32 FoundPoints = 0;
+					for (int32 Idx = 0; ; ++Idx)
+					{
+						while (Idx < InSamplePoints.Num() && PointInside[Idx] != bInKeepInside)
+						{
+							Idx++;
+						}
+						if (Idx < InSamplePoints.Num())
+						{
+							if (Idx != FoundPoints)
+							{
+								InSamplePoints[FoundPoints] = InSamplePoints[Idx];
+							}
+							FoundPoints++;
+						}
+						else
+						{
+							break;
+						}
+					}
+					InSamplePoints.SetNum(FoundPoints);
+				}
+			);
+		}
+		SetValue(Context, MoveTemp(InSamplePoints), &SamplePoints);
 	}
 }
 
diff --git a/Engine/Plugins/Experimental/GeometryCollectionPlugin/Source/GeometryCollectionNodes/Public/Dataflow/GeometryCollectionSamplingNodes.h b/Engine/Plugins/Experimental/GeometryCollectionPlugin/Source/GeometryCollectionNodes/Public/Dataflow/GeometryCollectionSamplingNodes.h
--- a/Engine/Plugins/Experimental/GeometryCollectionPlugin/Source/GeometryCollectionNodes/Public/Dataflow/GeometryCollectionSamplingNodes.h
+++ b/Engine/Plugins/Experimental/GeometryCollectionPlugin/Source/GeometryCollectionNodes/Public/Dataflow/GeometryCollectionSamplingNodes.h
@@ -16,6 +16,42 @@ namespace UE_DEPRECATED(5.5, "Use UE::Dataflow instead.") Dataflow {}
 class FGeometryCollection;
 class UDynamicMesh;
 
+/**
+ *
+ * Filter a point set to only the points inside or outside of a given mesh
+ *
+ */
+USTRUCT(meta = (DataflowGeometryCollection))
+struct FFilterPointSetWithMeshDataflowNode : public FDataflowNode
+{
+	GENERATED_USTRUCT_BODY()
+	DATAFLOW_NODE_DEFINE_INTERNAL(FFilterPointSetWithMeshDataflowNode, "FilterPointsWithMesh", "PointSampling", "")
+	DATAFLOW_NODE_RENDER_TYPE("PointsRender", FName("TArray<FVector>"), "SamplePoints")
+
+private:
+	/** Mesh to use to filter point set */
+	UPROPERTY(meta = (DataflowInput, DataflowIntrinsic))
+	TObjectPtr<UDynamicMesh> TargetMesh;
+
+	/** Whether to keep the points inside or (if false) outside the mesh. */
+	UPROPERTY(EditAnywhere, Category=Options, meta = (DataflowInput))
+	bool bKeepInside = true;
+
+	/** The winding number threshold to use for determining whether a point is inside or outside of the mesh */
+	UPROPERTY(EditAnywhere, Category = Options, meta = (DataflowInput))
+	float WindingThreshold = .5f;
+
+	/** Points to filter */
+	UPROPERTY(meta = (DataflowInput, DataflowOutput))
+	TArray<FVector> SamplePoints;
+
+
+	virtual void Evaluate(UE::Dataflow::FContext& Context, const FDataflowOutput* Out) const override;
+
+public:
+	FFilterPointSetWithMeshDataflowNode(const UE::Dataflow::FNodeParameters& InParam, FGuid InGuid = FGuid::NewGuid());
+};
+
 /**
  *
  * Uniform Sampling on a DynamicMesh
