add cluster point scatter dataflow node, matching the cluster fracture tool

diff --git a/Engine/Plugins/Experimental/GeometryCollectionPlugin/Source/GeometryCollectionNodes/Private/Dataflow/GeometryCollectionFracturingNodes.cpp b/Engine/Plugins/Experimental/GeometryCollectionPlugin/Source/GeometryCollectionNodes/Private/Dataflow/GeometryCollectionFracturingNodes.cpp
--- a/Engine/Plugins/Experimental/GeometryCollectionPlugin/Source/GeometryCollectionNodes/Private/Dataflow/GeometryCollectionFracturingNodes.cpp
+++ b/Engine/Plugins/Experimental/GeometryCollectionPlugin/Source/GeometryCollectionNodes/Private/Dataflow/GeometryCollectionFracturingNodes.cpp
@@ -53,6 +53,7 @@ namespace UE::Dataflow
 		DATAFLOW_NODE_REGISTER_CREATION_FACTORY(FRadialScatterPointsDataflowNode);
 		DATAFLOW_NODE_REGISTER_CREATION_FACTORY(FRadialScatterPointsDataflowNode_v2);
 		DATAFLOW_NODE_REGISTER_CREATION_FACTORY(FGridScatterPointsDataflowNode);
+		DATAFLOW_NODE_REGISTER_CREATION_FACTORY(FClusterScatterPointsDataflowNode);
 		DATAFLOW_NODE_REGISTER_CREATION_FACTORY(FVoronoiFractureDataflowNode);
 		DATAFLOW_NODE_REGISTER_CREATION_FACTORY(FVoronoiFractureDataflowNode_v2);
 		DATAFLOW_NODE_REGISTER_CREATION_FACTORY(FPlaneCutterDataflowNode);
@@ -71,6 +72,79 @@ namespace UE::Dataflow
 	}
 }
 
+FClusterScatterPointsDataflowNode::FClusterScatterPointsDataflowNode(const UE::Dataflow::FNodeParameters& InParam, FGuid InGuid)
+	: FDataflowNode(InParam, InGuid)
+{
+	RegisterInputConnection(&BoundingBox);
+
+	RegisterInputConnection(&NumberClustersMin).SetCanHidePin(true).SetPinIsHidden(true);
+	RegisterInputConnection(&NumberClustersMax).SetCanHidePin(true).SetPinIsHidden(true);
+	RegisterInputConnection(&PointsPerClusterMin).SetCanHidePin(true).SetPinIsHidden(true);
+	RegisterInputConnection(&PointsPerClusterMax).SetCanHidePin(true).SetPinIsHidden(true);
+	RegisterInputConnection(&ClusterRadiusFractionMin).SetCanHidePin(true).SetPinIsHidden(true);
+	RegisterInputConnection(&ClusterRadiusFractionMax).SetCanHidePin(true).SetPinIsHidden(true);
+	RegisterInputConnection(&ClusterRadiusOffset).SetCanHidePin(true).SetPinIsHidden(true);
+	RegisterInputConnection(&RandomSeed).SetCanHidePin(true).SetPinIsHidden(true);
+
+	RegisterOutputConnection(&Points);
+}
+
+void FClusterScatterPointsDataflowNode::Evaluate(UE::Dataflow::FContext& Context, const FDataflowOutput* Out) const
+{
+	if (Out->IsA(&Points))
+	{
+		const FBox& Bounds = GetValue<FBox>(Context, &BoundingBox);
+		if (Bounds.GetVolume() > 0.f)
+		{
+			FRandomStream RandStream(GetValue(Context, &RandomSeed));
+
+			const int32 InNumberClustersMin = FMath::Max(1, GetValue(Context, &NumberClustersMin));
+			const int32 InNumberClustersMax = FMath::Max(InNumberClustersMin, GetValue(Context, &NumberClustersMax));
+			const int32 ClusterCount = RandStream.RandRange(InNumberClustersMin, InNumberClustersMax);
+			const FVector Extent(Bounds.Max - Bounds.Min);
+
+			TArray<FVector> ClusterCenters;
+			ClusterCenters.Reserve(ClusterCount);
+			for (int32 ClusterIdx = 0; ClusterIdx < ClusterCount; ++ClusterIdx)
+			{
+				ClusterCenters.Emplace(Bounds.Min + FVector(RandStream.FRand(), RandStream.FRand(), RandStream.FRand()) * Extent);
+			}
+
+
+
+			const int32 InPointsPerClusterMin = FMath::Max(0, GetValue(Context, &PointsPerClusterMin));
+			const int32 InPointsPerClusterMax = FMath::Max(InPointsPerClusterMin, GetValue(Context, &PointsPerClusterMax));
+			const double InClusterRadiusOffset = (double)GetValue(Context, &ClusterRadiusOffset);
+			const double InClusterRadiusFractionMin = FMath::Max(0, (double)GetValue(Context, &ClusterRadiusFractionMin));
+			const double InClusterRadiusFractionMax = FMath::Max(InClusterRadiusFractionMin, (double)GetValue(Context, &ClusterRadiusFractionMax));
+			const double BoundsSize = Bounds.GetExtent().GetAbsMax();
+
+			TArray<FVector> NewPoints;
+			NewPoints.Reserve(ClusterCount * FMath::CeilToInt32(double(InPointsPerClusterMin + InPointsPerClusterMax) * .5));
+			for (int32 CenterIdx = 0; CenterIdx < ClusterCenters.Num(); ++CenterIdx)
+			{
+				const int32 SubPointCount = RandStream.RandRange(InPointsPerClusterMin, InPointsPerClusterMax);
+
+				for (int32 SubPointIdx = 0; SubPointIdx < SubPointCount; ++SubPointIdx)
+				{
+					FVector V(RandStream.VRand());
+					V.Normalize();
+					V *= InClusterRadiusOffset + (RandStream.FRandRange(InClusterRadiusFractionMin, InClusterRadiusFractionMax) * BoundsSize);
+					V += ClusterCenters[CenterIdx];
+					NewPoints.Emplace(V);
+				}
+			}
+
+			SetValue(Context, MoveTemp(NewPoints), &Points);
+		}
+		else
+		{
+			// ERROR: Invalid BoundingBox input
+			SetValue(Context, TArray<FVector>(), &Points);
+		}
+	}
+}
+
 void FUniformScatterPointsDataflowNode::Evaluate(UE::Dataflow::FContext& Context, const FDataflowOutput* Out) const
 {
 	if (Out->IsA<TArray<FVector>>(&Points))
diff --git a/Engine/Plugins/Experimental/GeometryCollectionPlugin/Source/GeometryCollectionNodes/Public/Dataflow/GeometryCollectionFracturingNodes.h b/Engine/Plugins/Experimental/GeometryCollectionPlugin/Source/GeometryCollectionNodes/Public/Dataflow/GeometryCollectionFracturingNodes.h
--- a/Engine/Plugins/Experimental/GeometryCollectionPlugin/Source/GeometryCollectionNodes/Public/Dataflow/GeometryCollectionFracturingNodes.h
+++ b/Engine/Plugins/Experimental/GeometryCollectionPlugin/Source/GeometryCollectionNodes/Public/Dataflow/GeometryCollectionFracturingNodes.h
@@ -99,6 +99,67 @@ public:
 
 };
 
+USTRUCT(meta = (DataflowGeometryCollection))
+struct FClusterScatterPointsDataflowNode : public FDataflowNode
+{
+	GENERATED_USTRUCT_BODY()
+	DATAFLOW_NODE_DEFINE_INTERNAL(FClusterScatterPointsDataflowNode, "ClusterScatterPoints", "Generators|Point", "")
+	DATAFLOW_NODE_RENDER_TYPE("PointsRender", FName("TArray<FVector>"), "Points")
+
+private:
+
+	/** Minimum number of clusters of points to create. The amount of clusters created will be chosen at random between Min and Max */
+	UPROPERTY(EditAnywhere, Category = Scatter, meta = (DisplayName = "Min Num Clusters", UIMax = "200", ClampMin = "1", DataflowInput))
+	int32 NumberClustersMin = 8;
+
+	/** Maximum number of clusters of points to create. The amount of clusters created will be chosen at random between Min and Max */
+	UPROPERTY(EditAnywhere, Category = Scatter, meta = (DisplayName = "Max Num Clusters", UIMax = "200", ClampMin = "1", DataflowInput))
+	int32 NumberClustersMax = 8;
+
+	/** Minimum number of points per cluster. The amount of points in each cluster will be chosen at random between Min and Max */
+	UPROPERTY(EditAnywhere, Category = Scatter, meta = (DisplayName = "Min Points Per Cluster", UIMax = "200", ClampMin = "0", DataflowInput))
+	int32 PointsPerClusterMin = 2;
+
+	/** Maximum number of points per cluster. The amount of points in each cluster will be chosen at random between Min and Max */
+	UPROPERTY(EditAnywhere, Category = Scatter, meta = (DisplayName = "Max Points Per Cluster", UIMax = "200", ClampMin = "0", DataflowInput))
+	int32 PointsPerClusterMax = 30;
+
+	/**
+	 * Minimum cluster radius (as fraction of the overall bounding box size). Cluster Radius Offset will be added to this.
+	 * Each point will be placed at least this far (plus the Cluster Radius Offset) from its cluster center.
+	 */
+	UPROPERTY(EditAnywhere, Category = Scatter, meta = (DisplayName = "Min Dist from Center (as Frac of Bounds)", ClampMin = "0.0", UIMax = "1.0", DataflowInput))
+	float ClusterRadiusFractionMin = .1;
+
+	/**
+	 * Maximum cluster radius (as fraction of the overall bounding box size). Cluster Radius Offset will be added to this.
+	 * Each point will be placed at most this far (plus the Cluster Radius Offset) from its cluster center.
+	 */
+	UPROPERTY(EditAnywhere, Category = Scatter, meta = (DisplayName = "Max Dist from Center (as Frac of Bounds)", ClampMin = "0.0", UIMax = "1.0", DataflowInput))
+	float ClusterRadiusFractionMax = .2;
+
+	/** Cluster radius offset (in cm). This offset will be added to the 'Min/Max Dist from Center' distance */
+	UPROPERTY(EditAnywhere, Category = Scatter, meta = (Units = "cm", DataflowInput))
+	float ClusterRadiusOffset = 0;
+
+	/** Seed for random */
+	UPROPERTY(EditAnywhere, Category = Scatter, meta = (DataflowInput, UIMin = "0"));
+	int32 RandomSeed = 0;
+
+	/** BoundingBox to generate points inside of */
+	UPROPERTY(meta = (DataflowInput))
+	FBox BoundingBox = FBox(ForceInit);
+
+	/** Generated points */
+	UPROPERTY(meta = (DataflowOutput))
+	TArray<FVector> Points;
+
+	virtual void Evaluate(UE::Dataflow::FContext& Context, const FDataflowOutput* Out) const override;
+
+public:
+	FClusterScatterPointsDataflowNode(const UE::Dataflow::FNodeParameters& InParam, FGuid InGuid = FGuid::NewGuid());
+};
+
 USTRUCT(meta = (DataflowGeometryCollection, Deprecated = "5.5"))
 struct FRadialScatterPointsDataflowNode : public FDataflowNode
 {
