Dataflow:
- Fixed missing check for enabling DebugDraw for SphereCoverings

diff --git a/Engine/Plugins/Experimental/GeometryCollectionPlugin/Source/GeometryCollectionNodes/Private/Dataflow/GeometryCollectionUtilityNodes.cpp b/Engine/Plugins/Experimental/GeometryCollectionPlugin/Source/GeometryCollectionNodes/Private/Dataflow/GeometryCollectionUtilityNodes.cpp
--- a/Engine/Plugins/Experimental/GeometryCollectionPlugin/Source/GeometryCollectionNodes/Private/Dataflow/GeometryCollectionUtilityNodes.cpp
+++ b/Engine/Plugins/Experimental/GeometryCollectionPlugin/Source/GeometryCollectionNodes/Private/Dataflow/GeometryCollectionUtilityNodes.cpp
@@ -260,11 +260,14 @@ void FCreateLeafConvexHullsDataflowNode::DebugDraw(UE::Dataflow::FContext& Conte
 			UE::Dataflow::Convex::DebugDrawProc(DataflowRenderingInterface, InCollection, bRandomizeColor, ColorRandomSeed, InOptionalSelectionFilter);
 		}
 
-		if (const FDataflowOutput* SphereCoveringOutput = FindOutput(&SphereCovering))
+		if (SphereCoveringDebugDrawRenderSettings.bDisplaySphereCovering)
 		{
-			const FDataflowSphereCovering& OutSpheres = SphereCoveringOutput->GetValue(Context, SphereCovering);
+			if (const FDataflowOutput* SphereCoveringOutput = FindOutput(&SphereCovering))
+			{
+				const FDataflowSphereCovering& OutSpheres = SphereCoveringOutput->GetValue(Context, SphereCovering);
 
-			UE::Dataflow::Convex::SphereCoveringDebugDrawProc(DataflowRenderingInterface, OutSpheres, SphereCoveringDebugDrawRenderSettings);
+				UE::Dataflow::Convex::SphereCoveringDebugDrawProc(DataflowRenderingInterface, OutSpheres, SphereCoveringDebugDrawRenderSettings);
+			}
 		}
 	}
 }
@@ -547,11 +550,14 @@ void FGenerateClusterConvexHullsFromLeafHullsDataflowNode::DebugDraw(UE::Dataflo
 			UE::Dataflow::Convex::DebugDrawProc(DataflowRenderingInterface, InCollection, bRandomizeColor, ColorRandomSeed, InOptionalSelectionFilter);
 		}
 
-		if (const FDataflowOutput* SphereCoveringOutput = FindOutput(&SphereCovering))
+		if (SphereCoveringDebugDrawRenderSettings.bDisplaySphereCovering)
 		{
-			const FDataflowSphereCovering& OutSpheres = SphereCoveringOutput->GetValue(Context, SphereCovering);
+			if (const FDataflowOutput* SphereCoveringOutput = FindOutput(&SphereCovering))
+			{
+				const FDataflowSphereCovering& OutSpheres = SphereCoveringOutput->GetValue(Context, SphereCovering);
 
-			UE::Dataflow::Convex::SphereCoveringDebugDrawProc(DataflowRenderingInterface, OutSpheres, SphereCoveringDebugDrawRenderSettings);
+				UE::Dataflow::Convex::SphereCoveringDebugDrawProc(DataflowRenderingInterface, OutSpheres, SphereCoveringDebugDrawRenderSettings);
+			}
 		}
 	}
 }
@@ -677,11 +683,14 @@ void FGenerateClusterConvexHullsFromChildrenHullsDataflowNode::DebugDraw(UE::Dat
 			UE::Dataflow::Convex::DebugDrawProc(DataflowRenderingInterface, InCollection, bRandomizeColor, ColorRandomSeed, InOptionalSelectionFilter);
 		}
 
-		if (const FDataflowOutput* SphereCoveringOutput = FindOutput(&SphereCovering))
+		if (SphereCoveringDebugDrawRenderSettings.bDisplaySphereCovering)
 		{
-			const FDataflowSphereCovering& OutSpheres = SphereCoveringOutput->GetValue(Context, SphereCovering);
+			if (const FDataflowOutput* SphereCoveringOutput = FindOutput(&SphereCovering))
+			{
+				const FDataflowSphereCovering& OutSpheres = SphereCoveringOutput->GetValue(Context, SphereCovering);
 
-			UE::Dataflow::Convex::SphereCoveringDebugDrawProc(DataflowRenderingInterface, OutSpheres, SphereCoveringDebugDrawRenderSettings);
+				UE::Dataflow::Convex::SphereCoveringDebugDrawProc(DataflowRenderingInterface, OutSpheres, SphereCoveringDebugDrawRenderSettings);
+			}
 		}
 	}
 }
@@ -894,11 +903,14 @@ void FMergeConvexHullsDataflowNode::DebugDraw(UE::Dataflow::FContext& Context, I
 			UE::Dataflow::Convex::DebugDrawProc(DataflowRenderingInterface, InCollection, bRandomizeColor, ColorRandomSeed, InOptionalSelectionFilter);
 		}
 
-		if (const FDataflowOutput* SphereCoveringOutput = FindOutput(&SphereCovering))
+		if (SphereCoveringDebugDrawRenderSettings.bDisplaySphereCovering)
 		{
-			const FDataflowSphereCovering& OutSpheres = SphereCoveringOutput->GetValue(Context, SphereCovering);
+			if (const FDataflowOutput* SphereCoveringOutput = FindOutput(&SphereCovering))
+			{
+				const FDataflowSphereCovering& OutSpheres = SphereCoveringOutput->GetValue(Context, SphereCovering);
 
-			UE::Dataflow::Convex::SphereCoveringDebugDrawProc(DataflowRenderingInterface, OutSpheres, SphereCoveringDebugDrawRenderSettings);
+				UE::Dataflow::Convex::SphereCoveringDebugDrawProc(DataflowRenderingInterface, OutSpheres, SphereCoveringDebugDrawRenderSettings);
+			}
 		}
 	}
 }
