Dataflow : Add support for geometrty collection root proxies
Also deprecate nodes that were still outputing normal material arrays instead of material interface arrays

diff --git a/Engine/Plugins/Experimental/GeometryCollectionPlugin/Source/GeometryCollectionNodes/Private/Dataflow/GeometryCollectionAssetNodes.cpp b/Engine/Plugins/Experimental/GeometryCollectionPlugin/Source/GeometryCollectionNodes/Private/Dataflow/GeometryCollectionAssetNodes.cpp
--- a/Engine/Plugins/Experimental/GeometryCollectionPlugin/Source/GeometryCollectionNodes/Private/Dataflow/GeometryCollectionAssetNodes.cpp
+++ b/Engine/Plugins/Experimental/GeometryCollectionPlugin/Source/GeometryCollectionNodes/Private/Dataflow/GeometryCollectionAssetNodes.cpp
@@ -21,12 +21,78 @@ namespace UE::Dataflow
 		DATAFLOW_NODE_REGISTER_CREATION_FACTORY(FGeometryCollectionTerminalDataflowNode_v2);
 		DATAFLOW_NODE_REGISTER_CREATION_FACTORY(FGetGeometryCollectionAssetDataflowNode);
 		DATAFLOW_NODE_REGISTER_CREATION_FACTORY(FGetGeometryCollectionSourcesDataflowNode);
+		DATAFLOW_NODE_REGISTER_CREATION_FACTORY(FCreateGeometryCollectionFromSourcesDataflowNode_v2);
+		DATAFLOW_NODE_REGISTER_CREATION_FACTORY(FGeometryCollectionToCollectionDataflowNode_v2);
+		DATAFLOW_NODE_REGISTER_CREATION_FACTORY(FBlueprintToCollectionDataflowNode_v2);
+		DATAFLOW_NODE_REGISTER_CREATION_FACTORY(FMakeRootProxyMeshDataflowNode);
+		DATAFLOW_NODE_REGISTER_CREATION_FACTORY(FMakeRootProxyMeshArrayDataflowNode);
+		DATAFLOW_NODE_REGISTER_CREATION_FACTORY(FAddRootProxyMeshToArrayDataflowNode);
+
+		// deprecated nodes (need to stay registered)
+		DATAFLOW_NODE_REGISTER_CREATION_FACTORY(FGeometryCollectionTerminalDataflowNode);
 		DATAFLOW_NODE_REGISTER_CREATION_FACTORY(FCreateGeometryCollectionFromSourcesDataflowNode);
 		DATAFLOW_NODE_REGISTER_CREATION_FACTORY(FGeometryCollectionToCollectionDataflowNode);
 		DATAFLOW_NODE_REGISTER_CREATION_FACTORY(FBlueprintToCollectionDataflowNode);
+	}
+}
 
-		// deprecated nodes (need to stay registered)
-		DATAFLOW_NODE_REGISTER_CREATION_FACTORY(FGeometryCollectionTerminalDataflowNode);
+// ===========================================================================================================================
+FMakeRootProxyMeshDataflowNode::FMakeRootProxyMeshDataflowNode(const UE::Dataflow::FNodeParameters& InParam, FGuid InGuid)
+	: FDataflowNode(InParam, InGuid)
+{
+	RegisterInputConnection(&Mesh);
+	RegisterInputConnection(&Transform);
+
+	RegisterOutputConnection(&RootProxyMesh);
+}
+
+void FMakeRootProxyMeshDataflowNode::Evaluate(UE::Dataflow::FContext& Context, const FDataflowOutput* Out) const
+{
+	if (Out->IsA(&RootProxyMesh))
+	{
+		FDataflowRootProxyMesh OutRootProxyMesh;
+		OutRootProxyMesh.Mesh = GetValue(Context, &Mesh);
+		OutRootProxyMesh.Transform = GetValue(Context, &Transform);
+
+		SetValue(Context, MoveTemp(OutRootProxyMesh), &RootProxyMesh);
+	}
+}
+
+// ===========================================================================================================================
+FMakeRootProxyMeshArrayDataflowNode::FMakeRootProxyMeshArrayDataflowNode(const UE::Dataflow::FNodeParameters& InParam, FGuid InGuid)
+	: FDataflowNode(InParam, InGuid)
+{
+	RegisterOutputConnection(&RootProxyMeshes);
+}
+
+void FMakeRootProxyMeshArrayDataflowNode::Evaluate(UE::Dataflow::FContext& Context, const FDataflowOutput* Out) const
+{
+	if (Out->IsA(&RootProxyMeshes))
+	{
+		SetValue(Context, RootProxyMeshes, &RootProxyMeshes);
+	}
+}
+
+// ===========================================================================================================================
+FAddRootProxyMeshToArrayDataflowNode::FAddRootProxyMeshToArrayDataflowNode(const UE::Dataflow::FNodeParameters& InParam, FGuid InGuid)
+	: FDataflowNode(InParam, InGuid)
+{
+	RegisterInputConnection(&RootProxyMeshes);
+	RegisterInputConnection(&RootProxyMesh);
+
+	RegisterOutputConnection(&RootProxyMeshes, &RootProxyMeshes);
+}
+
+void FAddRootProxyMeshToArrayDataflowNode::Evaluate(UE::Dataflow::FContext& Context, const FDataflowOutput* Out) const
+{
+	if (Out->IsA(&RootProxyMeshes))
+	{
+		const FDataflowRootProxyMesh& InRootProxyMesh = GetValue(Context, &RootProxyMesh);
+
+		TArray<FDataflowRootProxyMesh> OutArray;
+		OutArray.Add(InRootProxyMesh);
+
+		SetValue(Context, MoveTemp(OutArray), &RootProxyMeshes);
 	}
 }
 
@@ -36,10 +102,12 @@ FGeometryCollectionTerminalDataflowNode_v2::FGeometryCollectionTerminalDataflowN
 	: FDataflowTerminalNode(InParam, InGuid)
 {
 	RegisterInputConnection(&Collection);
-	RegisterOutputConnection(&Collection, &Collection);
 	RegisterInputConnection(&Materials);
-	RegisterOutputConnection(&Materials, &Materials);
 	RegisterInputConnection(&InstancedMeshes);
+	RegisterInputConnection(&RootProxyMeshes);
+
+	RegisterOutputConnection(&Collection, &Collection);
+	RegisterOutputConnection(&Materials, &Materials);
 	RegisterOutputConnection(&InstancedMeshes, &InstancedMeshes);
 }
 
@@ -48,6 +116,7 @@ void FGeometryCollectionTerminalDataflowNode_v2::SetAssetValue(TObjectPtr<UObjec
 	using FGeometryCollectionPtr = TSharedPtr<FGeometryCollection, ESPMode::ThreadSafe>;
 	using FMaterialArray = TArray<TObjectPtr<UMaterialInterface>>;
 	using FInstancedMeshesArray = TArray<FGeometryCollectionAutoInstanceMesh>;
+	using FRootProxyMeshesArray = TArray<FDataflowRootProxyMesh>;
 
 	if (UGeometryCollection* CollectionAsset = Cast<UGeometryCollection>(Asset.Get()))
 	{
@@ -56,12 +125,27 @@ void FGeometryCollectionTerminalDataflowNode_v2::SetAssetValue(TObjectPtr<UObjec
 			const FManagedArrayCollection& InCollection = GetValue(Context, &Collection);
 			const FMaterialArray& InMaterials = GetValue(Context, &Materials);
 			const FInstancedMeshesArray& InInstancedMeshes = GetValue(Context, &InstancedMeshes);
+			const FRootProxyMeshesArray& InRootProxyMeshes = GetValue(Context, &RootProxyMeshes);
 
 			constexpr bool bHasInternalMaterial = false; // with dataflow there's no assumption of internal materials
 			CollectionAsset->ResetFrom(InCollection, InMaterials, bHasInternalMaterial);
 
 			CollectionAsset->SetAutoInstanceMeshes(InInstancedMeshes);
 
+			if (IsConnected(&RootProxyMeshes))
+			{
+				FGeometryCollectionProxyMeshData& Data = CollectionAsset->RootProxyData;
+				Data.ProxyMeshes.Reset();
+				Data.ProxyMeshes.Reserve(InRootProxyMeshes.Num());
+				Data.MeshTransforms.Reset();
+				Data.MeshTransforms.Reserve(InRootProxyMeshes.Num());
+				for (const FDataflowRootProxyMesh& ProxyMesh : InRootProxyMeshes)
+				{
+					Data.ProxyMeshes.Add(ProxyMesh.Mesh);
+					Data.MeshTransforms.Add(FTransform3f(ProxyMesh.Transform));
+				}
+			}
+
 #if WITH_EDITOR
 			// make sure we rebuild the render data when we are done setting everything 
 			CollectionAsset->RebuildRenderData();
@@ -81,7 +165,7 @@ void FGeometryCollectionTerminalDataflowNode_v2::Evaluate(UE::Dataflow::FContext
 }
 
 // ===========================================================================================================================
-
+// DEPRECATED 5.6 : see FGeometryCollectionTerminalDataflowNode_v2
 FGeometryCollectionTerminalDataflowNode::FGeometryCollectionTerminalDataflowNode(const UE::Dataflow::FNodeParameters& InParam, FGuid InGuid)
 	: FDataflowTerminalNode(InParam, InGuid)
 {
@@ -192,6 +276,7 @@ void FGetGeometryCollectionSourcesDataflowNode::Evaluate(UE::Dataflow::FContext&
 }
 
 // ===========================================================================================================================
+// DEPRECATED 5.6 : see FCreateGeometryCollectionFromSourcesDataflowNode_v2
 
 FCreateGeometryCollectionFromSourcesDataflowNode::FCreateGeometryCollectionFromSourcesDataflowNode(const UE::Dataflow::FNodeParameters& InParam, FGuid InGuid)
 		: FDataflowNode(InParam, InGuid)
@@ -283,6 +368,92 @@ void FCreateGeometryCollectionFromSourcesDataflowNode::Evaluate(UE::Dataflow::FC
 
 // ===========================================================================================================================
 
+FCreateGeometryCollectionFromSourcesDataflowNode_v2::FCreateGeometryCollectionFromSourcesDataflowNode_v2(const UE::Dataflow::FNodeParameters& InParam, FGuid InGuid)
+	: FDataflowNode(InParam, InGuid)
+{
+	RegisterInputConnection(&Sources);
+	RegisterOutputConnection(&Collection);
+	RegisterOutputConnection(&Materials);
+	RegisterOutputConnection(&InstancedMeshes);
+}
+
+void FCreateGeometryCollectionFromSourcesDataflowNode_v2::Evaluate(UE::Dataflow::FContext& Context, const FDataflowOutput* Out) const
+{
+	ensure(Out->IsA(&Collection) || Out->IsA(&Materials) || Out->IsA(&InstancedMeshes));
+
+	const TArray<FGeometryCollectionSource>& InSources = GetValue(Context, &Sources);
+
+	FGeometryCollection OutCollection;
+	TArray<TObjectPtr<UMaterialInterface>> OutMaterials;
+	TArray<FGeometryCollectionAutoInstanceMesh> OutInstancedMeshes;
+
+	// make sure we have an attribute for instanced meshes
+	GeometryCollection::Facades::FCollectionInstancedMeshFacade InstancedMeshFacade(OutCollection);
+	InstancedMeshFacade.DefineSchema();
+
+	constexpr bool bReindexMaterialsInLoop = false;
+	for (int32 SourceIndex = 0; SourceIndex < InSources.Num(); SourceIndex++)
+	{
+		const FGeometryCollectionSource& Source = InSources[SourceIndex];
+		const int32 NumTransformsBeforeAppending = OutCollection.NumElements(FGeometryCollection::TransformGroup);
+
+		// todo: change AppendGeometryCollectionSource to take a FManagedArrayCollection so we could move the collection when assigning it to the output
+		FGeometryCollectionEngineConversion::AppendGeometryCollectionSource(Source, OutCollection, MutableView(OutMaterials), bReindexMaterialsInLoop);
+
+		// todo(chaos) if the source is a geometry collection this will not work properly 
+		FGeometryCollectionAutoInstanceMesh InstancedMesh;
+		InstancedMesh.Mesh = Cast<UStaticMesh>(Source.SourceGeometryObject.TryLoad());
+		InstancedMesh.Materials = Source.SourceMaterial;
+		InstancedMesh.NumInstances = 0;
+		InstancedMesh.CustomData.Reset();
+
+		int32 InstancedMeshIndex = OutInstancedMeshes.Find(InstancedMesh);
+		if (InstancedMeshIndex == INDEX_NONE)
+		{
+			InstancedMeshIndex = OutInstancedMeshes.Add(InstancedMesh);
+		}
+		OutInstancedMeshes[InstancedMeshIndex].NumInstances++;
+		OutInstancedMeshes[InstancedMeshIndex].CustomData.Append(Source.InstanceCustomData);
+
+		// add the instanced mesh  for all the newly added transforms 
+		const int32 NumTransformsAfterAppending = OutCollection.NumElements(FGeometryCollection::TransformGroup);
+		//ensure((NumTransformsAfterAppending - NumTransformsBeforeAppending) == 1);
+		for (int32 TransformIndex = NumTransformsBeforeAppending; TransformIndex < NumTransformsAfterAppending; TransformIndex++)
+		{
+			InstancedMeshFacade.SetIndex(TransformIndex, InstancedMeshIndex);
+		}
+	}
+	if (bReindexMaterialsInLoop == false)
+	{
+		OutCollection.ReindexMaterials();
+	}
+
+	// add the instanced mesh indices
+
+	const int32 NumTransforms = InstancedMeshFacade.GetNumIndices();
+	for (int32 TransformIndex = 0; TransformIndex < NumTransforms; TransformIndex++)
+	{
+	}
+
+	// make sure we have only one root
+	if (FGeometryCollectionClusteringUtility::ContainsMultipleRootBones(&OutCollection))
+	{
+		FGeometryCollectionClusteringUtility::ClusterAllBonesUnderNewRoot(&OutCollection);
+	}
+
+	// make sure we have a level attribute
+	Chaos::Facades::FCollectionHierarchyFacade HierarchyFacade(OutCollection);
+	HierarchyFacade.GenerateLevelAttribute();
+
+	// we have to make a copy since we have generated a FGeometryCollection which is inherited from FManagedArrayCollection
+	SetValue(Context, static_cast<const FManagedArrayCollection&>(OutCollection), &Collection);
+	SetValue(Context, MoveTemp(OutMaterials), &Materials);
+	SetValue(Context, MoveTemp(OutInstancedMeshes), &InstancedMeshes);
+}
+
+// ===========================================================================================================================
+// DEPRECATED 5.6 : see FGeometryCollectionToCollectionDataflowNode_v2
+
 FGeometryCollectionToCollectionDataflowNode::FGeometryCollectionToCollectionDataflowNode(const UE::Dataflow::FNodeParameters& InParam, FGuid InGuid)
 	: FDataflowNode(InParam, InGuid)
 {
@@ -317,6 +488,36 @@ void FGeometryCollectionToCollectionDataflowNode::Evaluate(UE::Dataflow::FContex
 
 // ===========================================================================================================================
 
+FGeometryCollectionToCollectionDataflowNode_v2::FGeometryCollectionToCollectionDataflowNode_v2(const UE::Dataflow::FNodeParameters& InParam, FGuid InGuid)
+	: FDataflowNode(InParam, InGuid)
+{
+	RegisterOutputConnection(&Collection);
+	RegisterOutputConnection(&Materials);
+	RegisterOutputConnection(&InstancedMeshes);
+}
+
+void FGeometryCollectionToCollectionDataflowNode_v2::Evaluate(UE::Dataflow::FContext& Context, const FDataflowOutput* Out) const
+{
+	ensure(Out->IsA(&Collection) || Out->IsA(&Materials) || Out->IsA(&InstancedMeshes));
+
+	FManagedArrayCollection OutCollection;
+	TArray<TObjectPtr<UMaterialInterface>> OutMaterials;
+	TArray<FGeometryCollectionAutoInstanceMesh> OutInstancedMeshes;
+
+	if (GeometryCollection)
+	{
+		FGeometryCollectionEngineConversion::ConvertGeometryCollectionToGeometryCollection(GeometryCollection, OutCollection, OutMaterials, OutInstancedMeshes);
+	}
+
+	// Set Outputs
+	SetValue(Context, MoveTemp(OutCollection), &Collection);
+	SetValue(Context, MoveTemp(OutMaterials), &Materials);
+	SetValue(Context, MoveTemp(OutInstancedMeshes), &InstancedMeshes);
+}
+
+// ===========================================================================================================================
+// DEPRECATED 5.6 : see FBlueprintToCollectionDataflowNode
+
 FBlueprintToCollectionDataflowNode::FBlueprintToCollectionDataflowNode(const UE::Dataflow::FNodeParameters& InParam, FGuid InGuid)
 	: FDataflowNode(InParam, InGuid)
 {
@@ -364,3 +565,46 @@ void FBlueprintToCollectionDataflowNode::Evaluate(UE::Dataflow::FContext& Contex
 	SetValue(Context, MoveTemp(OutInstancedMeshes), &InstancedMeshes);
 }
 
+// ===========================================================================================================================
+
+FBlueprintToCollectionDataflowNode_v2::FBlueprintToCollectionDataflowNode_v2(const UE::Dataflow::FNodeParameters& InParam, FGuid InGuid)
+	: FDataflowNode(InParam, InGuid)
+{
+	RegisterOutputConnection(&Collection);
+	RegisterOutputConnection(&Materials);
+	RegisterOutputConnection(&InstancedMeshes);
+}
+
+void FBlueprintToCollectionDataflowNode_v2::Evaluate(UE::Dataflow::FContext& Context, const FDataflowOutput* Out) const
+{
+	ensure(Out->IsA(&Collection) || Out->IsA(&Materials) || Out->IsA(&InstancedMeshes));
+
+	FManagedArrayCollection OutCollection;
+	TArray<TObjectPtr<UMaterialInterface>> OutMaterials;
+	TArray<FGeometryCollectionAutoInstanceMesh> OutInstancedMeshes;
+
+	if (Blueprint)
+	{
+		if (TUniquePtr<FPreviewScene> PreviewScene = MakeUnique<FPreviewScene>(FPreviewScene::ConstructionValues()))
+		{
+			if (UWorld* PreviewWorld = PreviewScene->GetWorld())
+			{
+				FActorSpawnParameters SpawnInfo;
+				SpawnInfo.SpawnCollisionHandlingOverride = ESpawnActorCollisionHandlingMethod::AlwaysSpawn;
+				SpawnInfo.bNoFail = true;
+				SpawnInfo.ObjectFlags = RF_Transient;
+
+				if (AActor* PreviewActor = PreviewWorld->SpawnActor(Blueprint->GeneratedClass, nullptr, SpawnInfo))
+				{
+					FGeometryCollectionEngineConversion::FSkeletalMeshToCollectionConversionParameters ConversionParameters;
+					FGeometryCollectionEngineConversion::ConvertActorToGeometryCollection(PreviewActor, OutCollection, OutMaterials, OutInstancedMeshes, ConversionParameters, bSplitComponents);
+				}
+			}
+		}
+	}
+
+	// Set Outputs
+	SetValue(Context, MoveTemp(OutCollection), &Collection);
+	SetValue(Context, MoveTemp(OutMaterials), &Materials);
+	SetValue(Context, MoveTemp(OutInstancedMeshes), &InstancedMeshes);
+}
\ No newline at end of file
diff --git a/Engine/Plugins/Experimental/GeometryCollectionPlugin/Source/GeometryCollectionNodes/Public/Dataflow/GeometryCollectionAssetNodes.h b/Engine/Plugins/Experimental/GeometryCollectionPlugin/Source/GeometryCollectionNodes/Public/Dataflow/GeometryCollectionAssetNodes.h
--- a/Engine/Plugins/Experimental/GeometryCollectionPlugin/Source/GeometryCollectionNodes/Public/Dataflow/GeometryCollectionAssetNodes.h
+++ b/Engine/Plugins/Experimental/GeometryCollectionPlugin/Source/GeometryCollectionNodes/Public/Dataflow/GeometryCollectionAssetNodes.h
@@ -19,6 +19,92 @@ namespace UE_DEPRECATED(5.5, "Use UE::Dataflow instead.") Dataflow {}
 class UGeometryCollection;
 class UBlueprint;
 
+USTRUCT()
+struct FDataflowRootProxyMesh
+{
+	GENERATED_BODY()
+
+	UPROPERTY(EditAnywhere, Category = RootProxyMesh);
+	TObjectPtr<UStaticMesh> Mesh;
+
+	UPROPERTY(EditAnywhere, Category = RootProxyMesh);
+	FTransform Transform;
+};
+
+/**
+* Create a RootProxyMesh object
+* (used by geometry collection assets)
+*/
+USTRUCT()
+struct FMakeRootProxyMeshDataflowNode : public FDataflowNode
+{
+	GENERATED_USTRUCT_BODY()
+	DATAFLOW_NODE_DEFINE_INTERNAL(FMakeRootProxyMeshDataflowNode, "MakeRootProxyMesh", "GeometryCollection", "")
+
+private:
+	/** mesh to use as a proxy */
+	UPROPERTY(EditAnywhere, Category = RootProxyMesh, meta = (DataflowInput));
+	TObjectPtr<UStaticMesh> Mesh;
+
+	/** transform to use for the proxy, relative to the asset it will be used for */
+	UPROPERTY(EditAnywhere, Category = RootProxyMesh, meta = (DataflowInput));
+	FTransform Transform;
+
+	/** mesh to use as a proxy */
+	UPROPERTY(meta = (DataflowOutput));
+	FDataflowRootProxyMesh RootProxyMesh;
+
+	virtual void Evaluate(UE::Dataflow::FContext& Context, const FDataflowOutput* Out) const override;
+
+public:
+	FMakeRootProxyMeshDataflowNode(const UE::Dataflow::FNodeParameters& InParam, FGuid InGuid = FGuid::NewGuid());
+};
+
+/**
+* Create a RootProxyMesh Array
+* (used by geometry collection assets)
+*/
+USTRUCT()
+struct FMakeRootProxyMeshArrayDataflowNode : public FDataflowNode
+{
+	GENERATED_USTRUCT_BODY()
+	DATAFLOW_NODE_DEFINE_INTERNAL(FMakeRootProxyMeshArrayDataflowNode, "MakeRootProxyMeshArray", "GeometryCollection", "")
+
+private:
+	/** newly created array */
+	UPROPERTY(EditAnywhere, Category = RootProxyMesh, meta = (DataflowOutput));
+	TArray<FDataflowRootProxyMesh> RootProxyMeshes;
+
+	virtual void Evaluate(UE::Dataflow::FContext& Context, const FDataflowOutput* Out) const override;
+
+public:
+	FMakeRootProxyMeshArrayDataflowNode(const UE::Dataflow::FNodeParameters& InParam, FGuid InGuid = FGuid::NewGuid());
+};
+
+/*
+* Add a root proxy object to an array of root proxy mesh
+* * (used by geometry collection assets)
+*/
+USTRUCT(meta = (DataflowGeometryCollection, DataflowTerminal))
+struct FAddRootProxyMeshToArrayDataflowNode : public FDataflowNode
+{
+	GENERATED_USTRUCT_BODY()
+	DATAFLOW_NODE_DEFINE_INTERNAL(FAddRootProxyMeshToArrayDataflowNode, "AddRootProxyMeshToArray", "GeometryCollection", "")
+
+private:
+	/** Root proxy array to add the mesh to */
+	UPROPERTY(meta = (DataflowInput, DataflowOutput, DataflowPassthrough = RootProxyMeshes));
+	TArray<FDataflowRootProxyMesh> RootProxyMeshes;
+
+	UPROPERTY(EditAnywhere, Category = RootProxyMesh, meta = (DataflowInput));
+	FDataflowRootProxyMesh RootProxyMesh;
+
+	virtual void Evaluate(UE::Dataflow::FContext& Context, const FDataflowOutput* Out) const override;
+
+public:
+	FAddRootProxyMeshToArrayDataflowNode(const UE::Dataflow::FNodeParameters& InParam, FGuid InGuid = FGuid::NewGuid());
+};
+
 /*
 * Geometry Collection asset terminal node
 */
@@ -41,6 +127,9 @@ private:
 	UPROPERTY(meta = (DataflowInput, DataflowOutput, DataflowPassthrough = "InstancedMeshes"))
 	TArray<FGeometryCollectionAutoInstanceMesh> InstancedMeshes;
 
+	UPROPERTY(meta = (DataflowInput));
+	TArray<FDataflowRootProxyMesh> RootProxyMeshes;
+
 	virtual void Evaluate(UE::Dataflow::FContext& Context) const override;
 	virtual void SetAssetValue(TObjectPtr<UObject> Asset, UE::Dataflow::FContext& Context) const override;
 
@@ -128,8 +217,9 @@ public:
 
 /**
  * create a geometry collection from a set of geometry sources    
+ * DEPRECATED 5.6 : use the new node version with a single material array output 
  */
-USTRUCT(meta = (DataflowGeometryCollection))
+USTRUCT(meta = (DataflowGeometryCollection, Deprecated = "5.6"))
 struct FCreateGeometryCollectionFromSourcesDataflowNode : public FDataflowNode
 {
 	GENERATED_USTRUCT_BODY()
@@ -160,11 +250,43 @@ public:
 	TArray<FGeometryCollectionAutoInstanceMesh> InstancedMeshes;
 };
 
+/**
+ * create a geometry collection from a set of geometry sources
+  */
+USTRUCT(meta = (DataflowGeometryCollection))
+struct FCreateGeometryCollectionFromSourcesDataflowNode_v2 : public FDataflowNode
+{
+	GENERATED_USTRUCT_BODY()
+	DATAFLOW_NODE_DEFINE_INTERNAL(FCreateGeometryCollectionFromSourcesDataflowNode_v2, "CreateGeometryCollectionFromSources", "GeometryCollection|Asset", "")
+
+public:
+	FCreateGeometryCollectionFromSourcesDataflowNode_v2(const UE::Dataflow::FNodeParameters& InParam, FGuid InGuid = FGuid::NewGuid());
+
+private:
+	virtual void Evaluate(UE::Dataflow::FContext& Context, const FDataflowOutput* Out) const override;
+
+	/** array of geometry sources */
+	UPROPERTY(meta = (DataflowInput))
+	TArray<FGeometryCollectionSource> Sources;
+
+	/** Geometry collection newly created */
+	UPROPERTY(meta = (DataflowOutput))
+	FManagedArrayCollection Collection;
+
+	/** Materials array to use for this asset */
+	UPROPERTY(meta = (DataflowOutput))
+	TArray<TObjectPtr<UMaterialInterface>> Materials;
+
+	/** array of instanced meshes*/
+	UPROPERTY(meta = (DataflowOutput))
+	TArray<FGeometryCollectionAutoInstanceMesh> InstancedMeshes;
+};
 
 /**
  * Converts a UGeometryCollection asset to an FManagedArrayCollection
+ * DEPRECATED 5.6 : use the new version  that only has one material array output 
  */
-USTRUCT(meta = (DataflowContext = "GeometryCollection", DataflowGeometryCollection, DataflowTerminal))
+USTRUCT(meta = (DataflowContext = "GeometryCollection", DataflowGeometryCollection, Deprecated = "5.6"))
 struct FGeometryCollectionToCollectionDataflowNode : public FDataflowNode
 {
 	GENERATED_USTRUCT_BODY()
@@ -196,11 +318,44 @@ public:
 	virtual void Evaluate(UE::Dataflow::FContext& Context, const FDataflowOutput* Out) const override;
 };
 
+/**
+ * Converts a UGeometryCollection asset to an FManagedArrayCollection
+ */
+USTRUCT(meta = (DataflowContext = "GeometryCollection", DataflowGeometryCollection))
+struct FGeometryCollectionToCollectionDataflowNode_v2 : public FDataflowNode
+{
+	GENERATED_USTRUCT_BODY()
+	DATAFLOW_NODE_DEFINE_INTERNAL(FGeometryCollectionToCollectionDataflowNode_v2, "GeometryCollectionToCollection", "GeometryCollection|Asset", "")
+		DATAFLOW_NODE_RENDER_TYPE("SurfaceRender", FGeometryCollection::StaticType(), "Collection")
+
+public:
+	FGeometryCollectionToCollectionDataflowNode_v2(const UE::Dataflow::FNodeParameters& InParam, FGuid InGuid = FGuid::NewGuid());
+
+private:
+	/** Asset input */
+	UPROPERTY(EditAnywhere, Category = "Asset");
+	TObjectPtr<UGeometryCollection> GeometryCollection;
+
+	/** Geometry collection newly created */
+	UPROPERTY(meta = (DataflowOutput))
+	FManagedArrayCollection Collection;
+
+	/** Material instances array from the static mesh */
+	UPROPERTY(meta = (DataflowOutput))
+	TArray<TObjectPtr<UMaterialInterface>> Materials;
+
+	/** Array of instanced meshes*/
+	UPROPERTY(meta = (DataflowOutput))
+	TArray<FGeometryCollectionAutoInstanceMesh> InstancedMeshes;
+
+	virtual void Evaluate(UE::Dataflow::FContext& Context, const FDataflowOutput* Out) const override;
+};
 
 /**
  * Create a geometry collection from an asset
+ * DEPRECATED 5.6 : use the new version  that only has one material array output 
  */
- USTRUCT(meta = (DataflowContext = "GeometryCollection", DataflowGeometryCollection, DataflowTerminal))
+ USTRUCT(meta = (DataflowContext = "GeometryCollection", DataflowGeometryCollection, Deprecated = "5.6"))
  struct FBlueprintToCollectionDataflowNode : public FDataflowNode
  {
  	GENERATED_USTRUCT_BODY()
@@ -236,7 +391,44 @@ public:
  	virtual void Evaluate(UE::Dataflow::FContext& Context, const FDataflowOutput* Out) const override;
  };
 
- 
+
+/**
+* Create a geometry collection from an asset
+*/
+USTRUCT(meta = (DataflowContext = "GeometryCollection", DataflowGeometryCollection))
+	struct FBlueprintToCollectionDataflowNode_v2 : public FDataflowNode
+{
+	GENERATED_USTRUCT_BODY()
+	DATAFLOW_NODE_DEFINE_INTERNAL(FBlueprintToCollectionDataflowNode_v2, "BlueprintToCollection", "GeometryCollection|Asset", "")
+	DATAFLOW_NODE_RENDER_TYPE("SurfaceRender", FGeometryCollection::StaticType(), "Collection")
+
+public:
+	FBlueprintToCollectionDataflowNode_v2(const UE::Dataflow::FNodeParameters& InParam, FGuid InGuid = FGuid::NewGuid());
+
+private:
+	/** Asset input */
+	UPROPERTY(EditAnywhere, Category = "Asset");
+	TObjectPtr<UBlueprint> Blueprint;
+
+	/** Split components */
+	UPROPERTY(EditAnywhere, Category = "Asset");
+	bool bSplitComponents = false;
+
+	/** Geometry collection newly created */
+	UPROPERTY(meta = (DataflowOutput))
+	FManagedArrayCollection Collection;
+
+	/** Material instances array from the static mesh */
+	UPROPERTY(meta = (DataflowOutput))
+	TArray<TObjectPtr<UMaterialInterface>> Materials;
+
+	/** Array of instanced meshes*/
+	UPROPERTY(meta = (DataflowOutput))
+	TArray<FGeometryCollectionAutoInstanceMesh> InstancedMeshes;
+
+	virtual void Evaluate(UE::Dataflow::FContext& Context, const FDataflowOutput* Out) const override;
+};
+
 namespace UE::Dataflow
 {
 	void GeometryCollectionEngineAssetNodes();
