geometry collection : make sure spawning a geometry collection actor in editor respect the hysics materials set on the asset

diff --git a/Engine/Plugins/Experimental/ChaosEditor/Source/FractureEditor/Private/FractureToolGenerators.cpp b/Engine/Plugins/Experimental/ChaosEditor/Source/FractureEditor/Private/FractureToolGenerators.cpp
--- a/Engine/Plugins/Experimental/ChaosEditor/Source/FractureEditor/Private/FractureToolGenerators.cpp
+++ b/Engine/Plugins/Experimental/ChaosEditor/Source/FractureEditor/Private/FractureToolGenerators.cpp
@@ -1063,7 +1063,10 @@ class AGeometryCollectionActor* UFractureToolGenerateAsset::CreateNewGeometryAct
 
 	// Set the Geometry Collection asset in the new actor
 	NewActor->GetGeometryCollectionComponent()->SetRestCollection(InGeometryCollection);
-	NewActor->GetGeometryCollectionComponent()->SetPhysMaterialOverride(GEngine->DefaultDestructiblePhysMaterial);
+	if (InGeometryCollection->PhysicsMaterial == nullptr)
+	{
+		NewActor->GetGeometryCollectionComponent()->SetPhysMaterialOverride(GEngine->DefaultDestructiblePhysMaterial);
+	}
 
 	// copy transform of original static mesh actor to this new actor
 	NewActor->SetActorLabel(UniqueAssetName);
diff --git a/Engine/Plugins/Experimental/GeometryCollectionPlugin/Source/GeometryCollectionEditor/Private/GeometryCollection/ActorFactoryGeometryCollection.cpp b/Engine/Plugins/Experimental/GeometryCollectionPlugin/Source/GeometryCollectionEditor/Private/GeometryCollection/ActorFactoryGeometryCollection.cpp
--- a/Engine/Plugins/Experimental/GeometryCollectionPlugin/Source/GeometryCollectionEditor/Private/GeometryCollection/ActorFactoryGeometryCollection.cpp
+++ b/Engine/Plugins/Experimental/GeometryCollectionPlugin/Source/GeometryCollectionEditor/Private/GeometryCollection/ActorFactoryGeometryCollection.cpp
@@ -55,7 +55,10 @@ void UActorFactoryGeometryCollection::PostSpawnActor(UObject* Asset, AActor* New
 	NewGeometryCollectionActor->GetGeometryCollectionComponent()->EnableClustering = GeometryCollection->EnableClustering;
 	NewGeometryCollectionActor->GetGeometryCollectionComponent()->ClusterGroupIndex = GeometryCollection->ClusterGroupIndex;
 	NewGeometryCollectionActor->GetGeometryCollectionComponent()->MaxClusterLevel = GeometryCollection->MaxClusterLevel;
-	NewGeometryCollectionActor->GetGeometryCollectionComponent()->SetPhysMaterialOverride(GEngine->DefaultDestructiblePhysMaterial);
+	if (GeometryCollection->PhysicsMaterial == nullptr)
+	{
+		NewGeometryCollectionActor->GetGeometryCollectionComponent()->SetPhysMaterialOverride(GEngine->DefaultDestructiblePhysMaterial);
+	}
 
 	// Init Component
 	NewGeometryCollectionActor->GetGeometryCollectionComponent()->RegisterComponent();
