make 'copy convex from root' node's convex debug rendering only render the root convexes, plus default to skipping bones w/ no convexes (to not clobber collision in that case)

diff --git a/Engine/Plugins/Experimental/GeometryCollectionPlugin/Source/GeometryCollectionNodes/Private/Dataflow/GeometryCollectionUtilityNodes.cpp b/Engine/Plugins/Experimental/GeometryCollectionPlugin/Source/GeometryCollectionNodes/Private/Dataflow/GeometryCollectionUtilityNodes.cpp
--- a/Engine/Plugins/Experimental/GeometryCollectionPlugin/Source/GeometryCollectionNodes/Private/Dataflow/GeometryCollectionUtilityNodes.cpp
+++ b/Engine/Plugins/Experimental/GeometryCollectionPlugin/Source/GeometryCollectionNodes/Private/Dataflow/GeometryCollectionUtilityNodes.cpp
@@ -735,9 +735,15 @@ void FCopyConvexHullsFromRootDataflowNode::DebugDraw(UE::Dataflow::FContext& Con
 		if (const FDataflowOutput* Output = FindOutput(&Collection))
 		{
 			const FManagedArrayCollection& InCollection = Output->GetValue(Context, Collection);
-			FDataflowTransformSelection EmptySelection;
 
-			UE::Dataflow::Convex::DebugDrawProc(DataflowRenderingInterface, InCollection, bRandomizeColor, ColorRandomSeed, EmptySelection);
+			GeometryCollection::Facades::FCollectionTransformSelectionFacade TransformSelectionFacade(InCollection);
+			const TArray<int32>& SelectionArr = TransformSelectionFacade.SelectRootBones();
+
+			FDataflowTransformSelection RootSelection;
+			RootSelection.Initialize(InCollection.NumElements(FTransformCollection::TransformGroup), false);
+			RootSelection.SetFromArray(SelectionArr);
+
+			UE::Dataflow::Convex::DebugDrawProc(DataflowRenderingInterface, InCollection, bRandomizeColor, ColorRandomSeed, RootSelection);
 		}
 	}
 }
@@ -749,6 +755,7 @@ FCopyConvexHullsFromRootDataflowNode::FCopyConvexHullsFromRootDataflowNode(const
 {
 	RegisterInputConnection(&Collection);
 	RegisterInputConnection(&FromCollection);
+	RegisterInputConnection(&bSkipIfEmpty).SetCanHidePin(true).SetPinIsHidden(true);
 
 	RegisterOutputConnection(&Collection, &Collection);
 	
@@ -766,6 +773,7 @@ void FCopyConvexHullsFromRootDataflowNode::Evaluate(UE::Dataflow::FContext& Cont
 		if (IsConnected(&Collection) && IsConnected(&FromCollection))
 		{
 			const FManagedArrayCollection& InFromCollection = GetValue(Context, &FromCollection);
+			const bool bInSkipIfEmpty = GetValue(Context, &bSkipIfEmpty);
 
 			if (FGeometryCollectionConvexUtility::HasConvexHullData(&InFromCollection))
 			{
@@ -779,12 +787,17 @@ void FCopyConvexHullsFromRootDataflowNode::Evaluate(UE::Dataflow::FContext& Cont
 				}
 				else
 				{
-					FGeometryCollectionConvexUtility::CopyConvexHulls(InCollection, ToRoots, InFromCollection, FromRoots);
+					FGeometryCollectionConvexUtility::CopyConvexHulls(InCollection, ToRoots, InFromCollection, FromRoots, bInSkipIfEmpty);
 				}
 			}
 			else
 			{
-				UE_LOG(LogChaosDataflow, Warning, TEXT("Failed to copy root collision: FromCollection has no convex hull data"));
+				if (!bInSkipIfEmpty)
+				{
+					GeometryCollection::Facades::FCollectionTransformSelectionFacade ToTransformSelectionFacade(InCollection);
+					const TArray<int32> ToRoots = ToTransformSelectionFacade.SelectRootBones();
+					FGeometryCollectionConvexUtility::RemoveConvexHulls(&InCollection, ToRoots);
+				}
 			}
 		}
 		
diff --git a/Engine/Plugins/Experimental/GeometryCollectionPlugin/Source/GeometryCollectionNodes/Public/Dataflow/GeometryCollectionUtilityNodes.h b/Engine/Plugins/Experimental/GeometryCollectionPlugin/Source/GeometryCollectionNodes/Public/Dataflow/GeometryCollectionUtilityNodes.h
--- a/Engine/Plugins/Experimental/GeometryCollectionPlugin/Source/GeometryCollectionNodes/Public/Dataflow/GeometryCollectionUtilityNodes.h
+++ b/Engine/Plugins/Experimental/GeometryCollectionPlugin/Source/GeometryCollectionNodes/Public/Dataflow/GeometryCollectionUtilityNodes.h
@@ -632,6 +632,10 @@ private:
 	UPROPERTY(meta = (DataflowInput, DataflowIntrinsic))
 	FManagedArrayCollection FromCollection;
 
+	// Whether to skip copying from roots with no convex hulls
+	UPROPERTY(EditAnywhere, Category = "Options", meta = (DataflowInput))
+	bool bSkipIfEmpty = true;
+
 	UPROPERTY(EditAnywhere, Category = "Debug Draw", meta = (ShowOnlyInnerProperties))
 	FDataflowNodeDebugDrawSettings DebugDrawRenderSettings;
 
diff --git a/Engine/Source/Runtime/Experimental/Chaos/Private/GeometryCollection/GeometryCollectionConvexUtility.cpp b/Engine/Source/Runtime/Experimental/Chaos/Private/GeometryCollection/GeometryCollectionConvexUtility.cpp
--- a/Engine/Source/Runtime/Experimental/Chaos/Private/GeometryCollection/GeometryCollectionConvexUtility.cpp
+++ b/Engine/Source/Runtime/Experimental/Chaos/Private/GeometryCollection/GeometryCollectionConvexUtility.cpp
@@ -2538,7 +2538,7 @@ Chaos::FConvexPtr FGeometryCollectionConvexUtility::GetConvexHull(const FGeometr
 
 
 bool FGeometryCollectionConvexUtility::CopyConvexHulls(FManagedArrayCollection& Collection, const TArray<int32>& TransformIndices,
-	const FManagedArrayCollection& SourceCollection, const TArray<int32>& SourceTransformIndices)
+	const FManagedArrayCollection& SourceCollection, const TArray<int32>& SourceTransformIndices, bool bSkipIfEmpty)
 {
 	static FName ConvexGroupName = FGeometryCollection::ConvexGroup;
 	static FName ConvexHullAttributeName = FGeometryCollection::ConvexHullAttribute;
@@ -2589,6 +2589,12 @@ bool FGeometryCollectionConvexUtility::CopyConvexHulls(FManagedArrayCollection&
 			continue;
 		}
 		const TSet<int32>& FromIndices = FromTransformToConvexIndicesAttribute[FromTransformIdx];
+
+		if (bSkipIfEmpty && FromIndices.IsEmpty())
+		{
+			continue;
+		}
+
 		ConvexToRemove.Append(ToTransformToConvexIndices[ToTransformIdx]);
 		TSet<int32>& ToConvexIndices = ToTransformToConvexIndices[ToTransformIdx];
 		ToConvexIndices.Reset();
diff --git a/Engine/Source/Runtime/Experimental/Chaos/Public/GeometryCollection/GeometryCollectionConvexUtility.h b/Engine/Source/Runtime/Experimental/Chaos/Public/GeometryCollection/GeometryCollectionConvexUtility.h
--- a/Engine/Source/Runtime/Experimental/Chaos/Public/GeometryCollection/GeometryCollectionConvexUtility.h
+++ b/Engine/Source/Runtime/Experimental/Chaos/Public/GeometryCollection/GeometryCollectionConvexUtility.h
@@ -281,9 +281,10 @@ public:
 	 * @param TransformIndices The transforms on the GeometryCollection to update w/ copied convex hulls. Must be same length as SourceTransformIndices array.
 	 * @param SourceCollection The collection to copy convex hulls from
 	 * @param SourceTransformIndices The transforms on the SourceCollection to copy from. Must be same length as TransformIndices array.
+	 * @param bSkipIfEmpty Whether to skip copying from transforms that have no convex hulls
 	 */
 	static CHAOS_API bool CopyConvexHulls(FManagedArrayCollection& Collection, const TArray<int32>& TransformIndices,
-		const FManagedArrayCollection& SourceCollection, const TArray<int32>& SourceTransformIndices);
+		const FManagedArrayCollection& SourceCollection, const TArray<int32>& SourceTransformIndices, bool bSkipIfEmpty = false);
 
 	/** Delete the convex hulls pointed at by the transform indices provided. */
 	static CHAOS_API void RemoveConvexHulls(FManagedArrayCollection* GeometryCollection, const TArray<int32>& TransformsToClearHullsFrom);
