Dataflow : display collection in construction view for most geometry collection selection nodes and show shapes for shape related selections in debug draw

diff --git a/Engine/Plugins/Experimental/GeometryCollectionPlugin/Source/GeometryCollectionNodes/Private/Dataflow/GeometryCollectionSelectionNodes.cpp b/Engine/Plugins/Experimental/GeometryCollectionPlugin/Source/GeometryCollectionNodes/Private/Dataflow/GeometryCollectionSelectionNodes.cpp
--- a/Engine/Plugins/Experimental/GeometryCollectionPlugin/Source/GeometryCollectionNodes/Private/Dataflow/GeometryCollectionSelectionNodes.cpp
+++ b/Engine/Plugins/Experimental/GeometryCollectionPlugin/Source/GeometryCollectionNodes/Private/Dataflow/GeometryCollectionSelectionNodes.cpp
@@ -3,6 +3,12 @@
 #include "Dataflow/GeometryCollectionSelectionNodes.h"
 #include "Dataflow/DataflowCore.h"
 
+#if WITH_EDITOR
+#include "Dataflow/DataflowRenderingViewMode.h"
+#endif
+
+#include "Dataflow/DataflowDebugDrawInterface.h"
+
 #include "Engine/StaticMesh.h"
 #include "GeometryCollection/GeometryCollectionObject.h"
 #include "GeometryCollection/ManagedArrayCollection.h"
@@ -824,6 +830,7 @@ void FCollectionTransformSelectionBySizeDataflowNode::Evaluate(UE::Dataflow::FCo
 	}
 }
 
+//////////////////////////////////////////////////////////////////////////////////////////////////
 
 void FCollectionTransformSelectionByVolumeDataflowNode::Evaluate(UE::Dataflow::FContext& Context, const FDataflowOutput* Out) const
 {
@@ -849,7 +856,6 @@ void FCollectionTransformSelectionByVolumeDataflowNode::Evaluate(UE::Dataflow::F
 	}
 }
 
-
 void FCollectionTransformSelectionInBoxDataflowNode::Evaluate(UE::Dataflow::FContext& Context, const FDataflowOutput* Out) const
 {
 	if (Out->IsA<FDataflowTransformSelection>(&TransformSelection))
@@ -886,6 +892,34 @@ void FCollectionTransformSelectionInBoxDataflowNode::Evaluate(UE::Dataflow::FCon
 	}
 }
 
+#if WITH_EDITOR
+
+bool FCollectionTransformSelectionInBoxDataflowNode::CanDebugDrawViewMode(const FName& ViewModeName) const
+{
+	return ViewModeName == UE::Dataflow::FDataflowConstruction3DViewMode::Name;
+}
+
+void FCollectionTransformSelectionInBoxDataflowNode::DebugDraw(UE::Dataflow::FContext& Context,	IDataflowDebugDrawInterface& DataflowRenderingInterface, const FDebugDrawParameters& DebugDrawParameters) const
+{
+	if ((DebugDrawParameters.bNodeIsSelected || DebugDrawParameters.bNodeIsPinned))
+	{
+		const FBox& InBox = GetValue<FBox>(Context, &Box);
+		const FTransform& InTransform = GetValue<FTransform>(Context, &Transform);
+
+		DataflowRenderingInterface.SetLineWidth(1.0);
+		DataflowRenderingInterface.SetWireframe(true);
+		DataflowRenderingInterface.SetWorldPriority();
+		DataflowRenderingInterface.SetColor(FLinearColor::Red);
+
+		const FVector TransformedCenter = InBox.GetCenter() + InTransform.GetTranslation();
+		const FVector ScaledExtent = InBox.GetExtent() * InTransform.GetScale3D();
+		DataflowRenderingInterface.DrawBox(ScaledExtent, InTransform.GetRotation(), TransformedCenter, 1.0);
+	}
+}
+
+#endif
+
+//////////////////////////////////////////////////////////////////////////////////////////////////
 
 void FCollectionTransformSelectionInSphereDataflowNode::Evaluate(UE::Dataflow::FContext& Context, const FDataflowOutput* Out) const
 {
@@ -923,6 +957,34 @@ void FCollectionTransformSelectionInSphereDataflowNode::Evaluate(UE::Dataflow::F
 	}
 }
 
+#if WITH_EDITOR
+
+bool FCollectionTransformSelectionInSphereDataflowNode::CanDebugDrawViewMode(const FName& ViewModeName) const
+{
+	return ViewModeName == UE::Dataflow::FDataflowConstruction3DViewMode::Name;
+}
+
+void FCollectionTransformSelectionInSphereDataflowNode::DebugDraw(UE::Dataflow::FContext& Context, IDataflowDebugDrawInterface& DataflowRenderingInterface, const FDebugDrawParameters& DebugDrawParameters) const
+{
+	if ((DebugDrawParameters.bNodeIsSelected || DebugDrawParameters.bNodeIsPinned))
+	{
+		const FSphere& InSphere = GetValue<FSphere>(Context, &Sphere);
+		const FTransform& InTransform = GetValue<FTransform>(Context, &Transform);
+
+		DataflowRenderingInterface.SetLineWidth(1.0);
+		DataflowRenderingInterface.SetWireframe(true);
+		DataflowRenderingInterface.SetWorldPriority();
+		DataflowRenderingInterface.SetColor(FLinearColor::Red);
+
+		const FVector TransformedCenter = InSphere.Center + InTransform.GetTranslation();
+		const double ScaledRadius = InSphere.W * InTransform.GetScale3D().GetMax();
+		DataflowRenderingInterface.DrawSphere(TransformedCenter, ScaledRadius);
+	}
+}
+
+#endif
+
+//////////////////////////////////////////////////////////////////////////////////////////////////
 
 void FCollectionTransformSelectionByFloatAttrDataflowNode::Evaluate(UE::Dataflow::FContext& Context, const FDataflowOutput* Out) const
 {
diff --git a/Engine/Plugins/Experimental/GeometryCollectionPlugin/Source/GeometryCollectionNodes/Public/Dataflow/GeometryCollectionSelectionNodes.h b/Engine/Plugins/Experimental/GeometryCollectionPlugin/Source/GeometryCollectionNodes/Public/Dataflow/GeometryCollectionSelectionNodes.h
--- a/Engine/Plugins/Experimental/GeometryCollectionPlugin/Source/GeometryCollectionNodes/Public/Dataflow/GeometryCollectionSelectionNodes.h
+++ b/Engine/Plugins/Experimental/GeometryCollectionPlugin/Source/GeometryCollectionNodes/Public/Dataflow/GeometryCollectionSelectionNodes.h
@@ -32,6 +32,7 @@ struct FCollectionTransformSelectionAllDataflowNode : public FDataflowNode
 {
 	GENERATED_USTRUCT_BODY()
 	DATAFLOW_NODE_DEFINE_INTERNAL(FCollectionTransformSelectionAllDataflowNode, "CollectionTransformSelectAll", "GeometryCollection|Selection|Transform", "")
+	DATAFLOW_NODE_RENDER_TYPE("SurfaceRender", FGeometryCollection::StaticType(), "Collection")
 
 public:
 	/** GeometryCollection for the selection */
@@ -159,6 +160,7 @@ struct FCollectionTransformSelectionNoneDataflowNode : public FDataflowNode
 {
 	GENERATED_USTRUCT_BODY()
 	DATAFLOW_NODE_DEFINE_INTERNAL(FCollectionTransformSelectionNoneDataflowNode, "CollectionTransformSelectNone", "GeometryCollection|Selection|Transform", "")
+	DATAFLOW_NODE_RENDER_TYPE("SurfaceRender", FGeometryCollection::StaticType(), "Collection")
 
 public:
 	/** GeometryCollection for the selection */
@@ -220,6 +222,7 @@ struct FCollectionTransformSelectionRandomDataflowNode : public FDataflowNode
 {
 	GENERATED_USTRUCT_BODY()
 	DATAFLOW_NODE_DEFINE_INTERNAL(FCollectionTransformSelectionRandomDataflowNode, "CollectionTransformSelectRandom", "GeometryCollection|Selection|Transform", "")
+	DATAFLOW_NODE_RENDER_TYPE("SurfaceRender", FGeometryCollection::StaticType(), "Collection")
 
 public:
 	/** If true, it always generates the same result for the same RandomSeed */
@@ -268,6 +271,7 @@ struct FCollectionTransformSelectionRootDataflowNode : public FDataflowNode
 {
 	GENERATED_USTRUCT_BODY()
 	DATAFLOW_NODE_DEFINE_INTERNAL(FCollectionTransformSelectionRootDataflowNode, "CollectionTransformSelectRoot", "GeometryCollection|Selection|Transform", "")
+	DATAFLOW_NODE_RENDER_TYPE("SurfaceRender", FGeometryCollection::StaticType(), "Collection")
 
 public:
 	/** GeometryCollection for the selection */
@@ -302,6 +306,7 @@ struct FCollectionTransformSelectionCustomDataflowNode : public FDataflowNode
 {
 	GENERATED_USTRUCT_BODY()
 	DATAFLOW_NODE_DEFINE_INTERNAL(FCollectionTransformSelectionCustomDataflowNode, "CollectionTransformSelectCustom", "GeometryCollection|Selection|Transform", "")
+	DATAFLOW_NODE_RENDER_TYPE("SurfaceRender", FGeometryCollection::StaticType(), "Collection")
 
 public:
 	/** GeometryCollection for the selection */
@@ -340,6 +345,7 @@ struct FCollectionTransformSelectionCustomDataflowNode_v2 : public FDataflowNode
 {
 	GENERATED_USTRUCT_BODY()
 	DATAFLOW_NODE_DEFINE_INTERNAL(FCollectionTransformSelectionCustomDataflowNode_v2, "CollectionTransformSelectCustom", "GeometryCollection|Selection|Transform", "")
+	DATAFLOW_NODE_RENDER_TYPE("SurfaceRender", FGeometryCollection::StaticType(), "Collection")
 
 public:
 	/** GeometryCollection for the selection */
@@ -369,6 +375,7 @@ struct FCollectionTransformSelectionFromIndexArrayDataflowNode : public FDataflo
 {
 	GENERATED_USTRUCT_BODY()
 	DATAFLOW_NODE_DEFINE_INTERNAL(FCollectionTransformSelectionFromIndexArrayDataflowNode, "CollectionTransformSelectionFromIndexArray", "GeometryCollection|Selection|Array", "")
+	DATAFLOW_NODE_RENDER_TYPE("SurfaceRender", FGeometryCollection::StaticType(), "Collection")
 
 public:
 
@@ -408,6 +415,7 @@ struct FCollectionTransformSelectionParentDataflowNode : public FDataflowNode
 {
 	GENERATED_USTRUCT_BODY()
 	DATAFLOW_NODE_DEFINE_INTERNAL(FCollectionTransformSelectionParentDataflowNode, "CollectionTransformSelectParent", "GeometryCollection|Selection|Transform", "")
+	DATAFLOW_NODE_RENDER_TYPE("SurfaceRender", FGeometryCollection::StaticType(), "Collection")
 
 public:
 	/** Array of the selected bone indices */
@@ -485,6 +493,7 @@ struct FCollectionTransformSelectionChildrenDataflowNode : public FDataflowNode
 {
 	GENERATED_USTRUCT_BODY()
 	DATAFLOW_NODE_DEFINE_INTERNAL(FCollectionTransformSelectionChildrenDataflowNode, "CollectionTransformSelectChildren", "GeometryCollection|Selection|Transform", "")
+	DATAFLOW_NODE_RENDER_TYPE("SurfaceRender", FGeometryCollection::StaticType(), "Collection")
 
 public:
 	/** Array of the selected bone indices */
@@ -519,6 +528,7 @@ struct FCollectionTransformSelectionSiblingsDataflowNode : public FDataflowNode
 {
 	GENERATED_USTRUCT_BODY()
 	DATAFLOW_NODE_DEFINE_INTERNAL(FCollectionTransformSelectionSiblingsDataflowNode, "CollectionTransformSelectSiblings", "GeometryCollection|Selection|Transform", "")
+	DATAFLOW_NODE_RENDER_TYPE("SurfaceRender", FGeometryCollection::StaticType(), "Collection")
 
 public:
 	/** Array of the selected bone indices */
@@ -553,6 +563,7 @@ struct FCollectionTransformSelectionLevelDataflowNode : public FDataflowNode
 {
 	GENERATED_USTRUCT_BODY()
 	DATAFLOW_NODE_DEFINE_INTERNAL(FCollectionTransformSelectionLevelDataflowNode, "CollectionTransformSelectSameLevel", "GeometryCollection|Selection|Transform", "")
+	DATAFLOW_NODE_RENDER_TYPE("SurfaceRender", FGeometryCollection::StaticType(), "Collection")
 
 public:
 	/** Array of the selected bone indices */
@@ -587,6 +598,7 @@ struct FCollectionTransformSelectionTargetLevelDataflowNode : public FDataflowNo
 {
 	GENERATED_USTRUCT_BODY()
 	DATAFLOW_NODE_DEFINE_INTERNAL(FCollectionTransformSelectionTargetLevelDataflowNode, "CollectionTransformSelectTargetLevel", "GeometryCollection|Selection|Transform", "")
+	DATAFLOW_NODE_RENDER_TYPE("SurfaceRender", FGeometryCollection::StaticType(), "Collection")
 
 public:
 	/** GeometryCollection for the selection */
@@ -630,6 +642,7 @@ struct FCollectionTransformSelectionContactDataflowNode : public FDataflowNode
 {
 	GENERATED_USTRUCT_BODY()
 	DATAFLOW_NODE_DEFINE_INTERNAL(FCollectionTransformSelectionContactDataflowNode, "CollectionTransformSelectContact", "GeometryCollection|Selection|Transform", "")
+	DATAFLOW_NODE_RENDER_TYPE("SurfaceRender", FGeometryCollection::StaticType(), "Collection")
 
 public:
 	/** Array of the selected bone indices */
@@ -668,6 +681,7 @@ struct FCollectionTransformSelectionLeafDataflowNode : public FDataflowNode
 {
 	GENERATED_USTRUCT_BODY()
 	DATAFLOW_NODE_DEFINE_INTERNAL(FCollectionTransformSelectionLeafDataflowNode, "CollectionTransformSelectLeaf", "GeometryCollection|Selection|Transform", "")
+	DATAFLOW_NODE_RENDER_TYPE("SurfaceRender", FGeometryCollection::StaticType(), "Collection")
 
 public:
 	/** GeometryCollection for the selection */
@@ -703,6 +717,7 @@ struct FCollectionTransformSelectionClusterDataflowNode : public FDataflowNode
 {
 	GENERATED_USTRUCT_BODY()
 	DATAFLOW_NODE_DEFINE_INTERNAL(FCollectionTransformSelectionClusterDataflowNode, "CollectionTransformSelectCluster", "GeometryCollection|Selection|Transform", "")
+	DATAFLOW_NODE_RENDER_TYPE("SurfaceRender", FGeometryCollection::StaticType(), "Collection")
 
 public:
 	/** GeometryCollection for the selection */
@@ -736,6 +751,7 @@ struct FCollectionTransformSelectionClusterDataflowNode_v2 : public FDataflowNod
 {
 	GENERATED_USTRUCT_BODY()
 	DATAFLOW_NODE_DEFINE_INTERNAL(FCollectionTransformSelectionClusterDataflowNode_v2, "CollectionTransformSelectCluster", "GeometryCollection|Selection|Transform", "")
+	DATAFLOW_NODE_RENDER_TYPE("SurfaceRender", FGeometryCollection::StaticType(), "Collection")
 
 public:
 	/** GeometryCollection for the selection */
@@ -832,6 +848,7 @@ struct FCollectionTransformSelectionBySizeDataflowNode : public FDataflowNode
 {
 	GENERATED_USTRUCT_BODY()
 	DATAFLOW_NODE_DEFINE_INTERNAL(FCollectionTransformSelectionBySizeDataflowNode, "CollectionTransformSelectBySize", "GeometryCollection|Selection|Transform", "")
+	DATAFLOW_NODE_RENDER_TYPE("SurfaceRender", FGeometryCollection::StaticType(), "Collection")
 
 public:
 	/** GeometryCollection for the selection */
@@ -887,6 +904,7 @@ struct FCollectionTransformSelectionByVolumeDataflowNode : public FDataflowNode
 {
 	GENERATED_USTRUCT_BODY()
 	DATAFLOW_NODE_DEFINE_INTERNAL(FCollectionTransformSelectionByVolumeDataflowNode, "CollectionTransformSelectByVolume", "GeometryCollection|Selection|Transform", "")
+	DATAFLOW_NODE_RENDER_TYPE("SurfaceRender", FGeometryCollection::StaticType(), "Collection")
 
 public:
 	/** GeometryCollection for the selection */
@@ -952,6 +970,7 @@ struct FCollectionTransformSelectionInBoxDataflowNode : public FDataflowNode
 {
 	GENERATED_USTRUCT_BODY()
 	DATAFLOW_NODE_DEFINE_INTERNAL(FCollectionTransformSelectionInBoxDataflowNode, "CollectionTransformSelectInBox", "GeometryCollection|Selection|Transform", "")
+	DATAFLOW_NODE_RENDER_TYPE("SurfaceRender", FGeometryCollection::StaticType(), "Collection")
 
 public:
 	/** GeometryCollection for the selection */
@@ -990,6 +1009,10 @@ public:
 
 	virtual void Evaluate(UE::Dataflow::FContext& Context, const FDataflowOutput* Out) const override;
 
+#if WITH_EDITOR
+	virtual bool CanDebugDrawViewMode(const FName& ViewModeName) const override;
+	virtual void DebugDraw(UE::Dataflow::FContext& Context, IDataflowDebugDrawInterface& DataflowRenderingInterface, const FDebugDrawParameters& DebugDrawParameters) const override;
+#endif //WITH_EDITOR
 };
 
 
@@ -1003,6 +1026,7 @@ struct FCollectionTransformSelectionInSphereDataflowNode : public FDataflowNode
 {
 	GENERATED_USTRUCT_BODY()
 	DATAFLOW_NODE_DEFINE_INTERNAL(FCollectionTransformSelectionInSphereDataflowNode, "CollectionTransformSelectInSphere", "GeometryCollection|Selection|Transform", "")
+	DATAFLOW_NODE_RENDER_TYPE("SurfaceRender", FGeometryCollection::StaticType(), "Collection")
 
 public:
 	/** GeometryCollection for the selection */
@@ -1041,6 +1065,10 @@ public:
 
 	virtual void Evaluate(UE::Dataflow::FContext& Context, const FDataflowOutput* Out) const override;
 
+#if WITH_EDITOR
+	virtual bool CanDebugDrawViewMode(const FName& ViewModeName) const override;
+	virtual void DebugDraw(UE::Dataflow::FContext& Context, IDataflowDebugDrawInterface& DataflowRenderingInterface, const FDebugDrawParameters& DebugDrawParameters) const override;
+#endif //WITH_EDITOR
 };
 
 
@@ -1054,6 +1082,7 @@ struct FCollectionTransformSelectionByFloatAttrDataflowNode : public FDataflowNo
 {
 	GENERATED_USTRUCT_BODY()
 	DATAFLOW_NODE_DEFINE_INTERNAL(FCollectionTransformSelectionByFloatAttrDataflowNode, "CollectionTransformSelectByFloatAttribute", "GeometryCollection|Selection|Transform", "")
+	DATAFLOW_NODE_RENDER_TYPE("SurfaceRender", FGeometryCollection::StaticType(), "Collection")
 
 public:
 	/** GeometryCollection for the selection */
@@ -1113,6 +1142,7 @@ struct FCollectionTransformSelectionByIntAttrDataflowNode : public FDataflowNode
 {
 	GENERATED_USTRUCT_BODY()
 	DATAFLOW_NODE_DEFINE_INTERNAL(FCollectionTransformSelectionByIntAttrDataflowNode, "CollectionTransformSelectByIntAttribute", "GeometryCollection|Selection|Transform", "")
+	DATAFLOW_NODE_RENDER_TYPE("SurfaceRender", FGeometryCollection::StaticType(), "Collection")
 
 public:
 	/** GeometryCollection for the selection */
@@ -1173,6 +1203,7 @@ struct FCollectionVertexSelectionCustomDataflowNode : public FDataflowNode
 {
 	GENERATED_USTRUCT_BODY()
 	DATAFLOW_NODE_DEFINE_INTERNAL(FCollectionVertexSelectionCustomDataflowNode, "CollectionVertexSelectCustom", "GeometryCollection|Selection|Vertex", "")
+	DATAFLOW_NODE_RENDER_TYPE("SurfaceRender", FGeometryCollection::StaticType(), "Collection")
 
 public:
 	/** GeometryCollection for the selection */
@@ -1212,6 +1243,7 @@ struct FCollectionFaceSelectionCustomDataflowNode : public FDataflowNode
 {
 	GENERATED_USTRUCT_BODY()
 	DATAFLOW_NODE_DEFINE_INTERNAL(FCollectionFaceSelectionCustomDataflowNode, "CollectionFaceSelectCustom", "GeometryCollection|Selection|Face", "")
+	DATAFLOW_NODE_RENDER_TYPE("SurfaceRender", FGeometryCollection::StaticType(), "Collection")
 
 public:
 	/** GeometryCollection for the selection */
@@ -1250,6 +1282,7 @@ struct FCollectionSelectionConvertDataflowNode : public FDataflowNode
 {
 	GENERATED_USTRUCT_BODY()
 	DATAFLOW_NODE_DEFINE_INTERNAL(FCollectionSelectionConvertDataflowNode, "CollectionSelectionConvert", "GeometryCollection|Selection", "")
+	DATAFLOW_NODE_RENDER_TYPE("SurfaceRender", FGeometryCollection::StaticType(), "Collection")
 
 public:
 	/** GeometryCollection for the selection */
@@ -1452,6 +1485,7 @@ struct FCollectionSelectionByAttrDataflowNode : public FDataflowNode
 {
 	GENERATED_USTRUCT_BODY()
 	DATAFLOW_NODE_DEFINE_INTERNAL(FCollectionSelectionByAttrDataflowNode, "CollectionSelectByAttr", "GeometryCollection|Selection|All", "")
+	DATAFLOW_NODE_RENDER_TYPE("SurfaceRender", FGeometryCollection::StaticType(), "Collection")
 
 public:
 	/** GeometryCollection for the selection */
@@ -1619,6 +1653,7 @@ struct FCollectionSelectInternalFacesDataflowNode : public FDataflowNode
 {
 	GENERATED_USTRUCT_BODY()
 	DATAFLOW_NODE_DEFINE_INTERNAL(FCollectionSelectInternalFacesDataflowNode, "CollectionSelectInternalFaces", "GeometryCollection|Selection", "")
+	DATAFLOW_NODE_RENDER_TYPE("SurfaceRender", FGeometryCollection::StaticType(), "Collection")
 
 private:
 	/** Collection to select the internal faces from */
