Dataflow : add support for collection rendering for a lot of geometry collection nodes
- also added debug draw for SetAnchorState to show the center of the transforms affected by the state change
- fix points to allow for foreground rendering

diff --git a/Engine/Plugins/Experimental/GeometryCollectionPlugin/Source/GeometryCollectionNodes/Private/Dataflow/GeometryCollectionNodes.cpp b/Engine/Plugins/Experimental/GeometryCollectionPlugin/Source/GeometryCollectionNodes/Private/Dataflow/GeometryCollectionNodes.cpp
--- a/Engine/Plugins/Experimental/GeometryCollectionPlugin/Source/GeometryCollectionNodes/Private/Dataflow/GeometryCollectionNodes.cpp
+++ b/Engine/Plugins/Experimental/GeometryCollectionPlugin/Source/GeometryCollectionNodes/Private/Dataflow/GeometryCollectionNodes.cpp
@@ -868,6 +868,46 @@ void FSetAnchorStateDataflowNode::Evaluate(UE::Dataflow::FContext& Context, cons
 	}
 }
 
+
+#if WITH_EDITOR
+
+bool FSetAnchorStateDataflowNode::CanDebugDrawViewMode(const FName& ViewModeName) const
+{
+	return ViewModeName == UE::Dataflow::FDataflowConstruction3DViewMode::Name;
+}
+
+void FSetAnchorStateDataflowNode::DebugDraw(UE::Dataflow::FContext& Context, IDataflowDebugDrawInterface& DataflowRenderingInterface, const FDebugDrawParameters& DebugDrawParameters) const
+{
+	if ((DebugDrawParameters.bNodeIsSelected || DebugDrawParameters.bNodeIsPinned))
+	{
+		const FManagedArrayCollection& InCollection = GetValue(Context, &Collection);
+		const FDataflowTransformSelection& InTransformSelection = GetValue(Context, &TransformSelection);
+
+		GeometryCollection::Facades::FBoundsFacade BoundsFacade(InCollection);
+		const TArray<FVector> Centroids = BoundsFacade.GetCentroids();
+
+		GeometryCollection::Facades::FCollectionTransformFacade TransformFacade(InCollection);
+
+		DataflowRenderingInterface.SetColor(FLinearColor::Blue);
+		DataflowRenderingInterface.SetPointSize(5.f);
+		DataflowRenderingInterface.ReservePoints(InTransformSelection.NumSelected());
+		DataflowRenderingInterface.SetForegroundPriority();
+
+		const int32 NumCentroids = Centroids.Num();
+		for (int32 TransformIdx = 0; TransformIdx < InTransformSelection.Num(); TransformIdx++)
+		{
+			if (TransformIdx < NumCentroids && InTransformSelection.IsSelected(TransformIdx))
+			{
+				const FTransform CollectionSpaceTransform = TransformFacade.ComputeCollectionSpaceTransform(TransformIdx);
+				const FVector Point = CollectionSpaceTransform.TransformPosition(Centroids[TransformIdx]);
+				DataflowRenderingInterface.DrawPoint(Point);
+			}
+		}
+	}
+}
+
+#endif
+
 //////////////////////////////////////////////////////////////////////////////////////
 
 FSetDynamicStateDataflowNode::FSetDynamicStateDataflowNode(const UE::Dataflow::FNodeParameters& InParam, FGuid InGuid)
diff --git a/Engine/Plugins/Experimental/GeometryCollectionPlugin/Source/GeometryCollectionNodes/Public/Dataflow/GeometryCollectionNodes.h b/Engine/Plugins/Experimental/GeometryCollectionPlugin/Source/GeometryCollectionNodes/Public/Dataflow/GeometryCollectionNodes.h
--- a/Engine/Plugins/Experimental/GeometryCollectionPlugin/Source/GeometryCollectionNodes/Public/Dataflow/GeometryCollectionNodes.h
+++ b/Engine/Plugins/Experimental/GeometryCollectionPlugin/Source/GeometryCollectionNodes/Public/Dataflow/GeometryCollectionNodes.h
@@ -36,6 +36,7 @@ struct FGetCollectionFromAssetDataflowNode : public FDataflowNode
 {
 	GENERATED_USTRUCT_BODY()
 	DATAFLOW_NODE_DEFINE_INTERNAL(FGetCollectionFromAssetDataflowNode, "GetCollectionFromAsset", "GeometryCollection|Asset", "")
+	DATAFLOW_NODE_RENDER_TYPE("SurfaceRender", FGeometryCollection::StaticType(), "Collection")
 
 public:
 	UPROPERTY(EditAnywhere, Category = "Dataflow", meta = (DataflowInput, DisplayName = "Asset"))
@@ -66,6 +67,7 @@ struct FAppendCollectionAssetsDataflowNode : public FDataflowNode
 {
 	GENERATED_USTRUCT_BODY()
 	DATAFLOW_NODE_DEFINE_INTERNAL(FAppendCollectionAssetsDataflowNode, "AppendCollections", "GeometryCollection", "")
+	DATAFLOW_NODE_RENDER_TYPE("SurfaceRender", FGeometryCollection::StaticType(), "Collection")
 
 public:
 	typedef FManagedArrayCollection DataType;
@@ -993,6 +995,8 @@ struct FBranchCollectionDataflowNode : public FDataflowNode
 {
 	GENERATED_USTRUCT_BODY()
 	DATAFLOW_NODE_DEFINE_INTERNAL(FBranchCollectionDataflowNode, "BranchCollection", "Utilities|FlowControl", "")
+	DATAFLOW_NODE_RENDER_TYPE("SurfaceRender", FGeometryCollection::StaticType(), "ChosenCollection")
+
 
 public:
 	/** Collection input for the 'true' case */
@@ -1063,6 +1067,7 @@ struct FRemoveOnBreakDataflowNode : public FDataflowNode
 {
 	GENERATED_USTRUCT_BODY()
 	DATAFLOW_NODE_DEFINE_INTERNAL(FRemoveOnBreakDataflowNode, "RemoveOnBreak", "GeometryCollection|Utilities", "")
+	DATAFLOW_NODE_RENDER_TYPE("SurfaceRender", FGeometryCollection::StaticType(), "Collection")
 
 public:
 	/** Collection to set theremoval data on */
@@ -1126,6 +1131,7 @@ struct FSetAnchorStateDataflowNode : public FDataflowNode
 {
 	GENERATED_USTRUCT_BODY()
 	DATAFLOW_NODE_DEFINE_INTERNAL(FSetAnchorStateDataflowNode, "SetAnchorState", "GeometryCollection|Utilities", "")
+	DATAFLOW_NODE_RENDER_TYPE("SurfaceRender", FGeometryCollection::StaticType(), "Collection")
 
 private:
 	/** What anchor state to set on selected bones */
@@ -1146,6 +1152,11 @@ private:
 
 	virtual void Evaluate(UE::Dataflow::FContext& Context, const FDataflowOutput* Out) const override;
 
+#if WITH_EDITOR
+	virtual bool CanDebugDrawViewMode(const FName& ViewModeName) const override;
+	virtual void DebugDraw(UE::Dataflow::FContext& Context, IDataflowDebugDrawInterface& DataflowRenderingInterface, const FDebugDrawParameters& DebugDrawParameters) const override;
+#endif //WITH_EDITOR
+
 public:
 	FSetAnchorStateDataflowNode(const UE::Dataflow::FNodeParameters& InParam, FGuid InGuid = FGuid::NewGuid());
 };
@@ -1169,6 +1180,7 @@ struct FSetDynamicStateDataflowNode : public FDataflowNode
 {
 	GENERATED_USTRUCT_BODY()
 	DATAFLOW_NODE_DEFINE_INTERNAL(FSetDynamicStateDataflowNode, "SetDynamicState", "GeometryCollection|Utilities", "")
+	DATAFLOW_NODE_RENDER_TYPE("SurfaceRender", FGeometryCollection::StaticType(), "Collection")
 
 private:
 	/** Dynamic state to set on selected bones */
@@ -1322,6 +1334,7 @@ struct FCollectionSetPivotDataflowNode : public FDataflowNode
 {
 	GENERATED_USTRUCT_BODY()
 	DATAFLOW_NODE_DEFINE_INTERNAL(FCollectionSetPivotDataflowNode, "SetPivot", "GeometryCollection|Utilities", "")
+	DATAFLOW_NODE_RENDER_TYPE("SurfaceRender", FGeometryCollection::StaticType(), "Collection")
 
 public:
 	/** Collection for the pivot change */
@@ -1403,6 +1416,7 @@ struct FAddCustomCollectionAttributeDataflowNode : public FDataflowNode
 {
 	GENERATED_USTRUCT_BODY()
 	DATAFLOW_NODE_DEFINE_INTERNAL(FAddCustomCollectionAttributeDataflowNode, "AddCustomCollectionAttribute", "GeometryCollection|Utilities", "")
+	DATAFLOW_NODE_RENDER_TYPE("SurfaceRender", FGeometryCollection::StaticType(), "Collection")
 
 public:
 	/** Collection for the custom attribute */
@@ -1574,6 +1588,7 @@ struct FGetCollectionAttributeDataTypedDataflowNode_v2 : public FDataflowNode
 {
 	GENERATED_USTRUCT_BODY()
 	DATAFLOW_NODE_DEFINE_INTERNAL(FGetCollectionAttributeDataTypedDataflowNode_v2, "GetCollectionAttributeDataTyped", "GeometryCollection|Utilities", "")
+	DATAFLOW_NODE_RENDER_TYPE("SurfaceRender", FGeometryCollection::StaticType(), "Collection")
 
 private:
 	/** Collection for the custom attribute */
@@ -1638,6 +1653,7 @@ struct FSetCollectionAttributeDataTypedDataflowNode : public FDataflowNode
 {
 	GENERATED_USTRUCT_BODY()
 	DATAFLOW_NODE_DEFINE_INTERNAL(FSetCollectionAttributeDataTypedDataflowNode, "SetCollectionAttributeDataTyped", "GeometryCollection|Utilities", "")
+	DATAFLOW_NODE_RENDER_TYPE("SurfaceRender", FGeometryCollection::StaticType(), "Collection")
 
 public:
 	/** Collection for the custom attribute */
diff --git a/Engine/Source/Runtime/Experimental/Dataflow/Engine/Private/Dataflow/DataflowDebugDraw.cpp b/Engine/Source/Runtime/Experimental/Dataflow/Engine/Private/Dataflow/DataflowDebugDraw.cpp
--- a/Engine/Source/Runtime/Experimental/Dataflow/Engine/Private/Dataflow/DataflowDebugDraw.cpp
+++ b/Engine/Source/Runtime/Experimental/Dataflow/Engine/Private/Dataflow/DataflowDebugDraw.cpp
@@ -126,6 +126,7 @@ void FDataflowDebugDraw::DrawPoint(const FVector& Position)
 		.Position = Position,
 		.Size = PointSize,
 		.Color = ColorWithTranslucency.ToFColor(true),
+		.Priority = PriorityGroup,
 	};
 	DebugRenderSceneProxy->AddPoint(NewPoint);
 }
diff --git a/Engine/Source/Runtime/Experimental/Dataflow/Engine/Private/Dataflow/DataflowDebugDrawComponent.cpp b/Engine/Source/Runtime/Experimental/Dataflow/Engine/Private/Dataflow/DataflowDebugDrawComponent.cpp
--- a/Engine/Source/Runtime/Experimental/Dataflow/Engine/Private/Dataflow/DataflowDebugDrawComponent.cpp
+++ b/Engine/Source/Runtime/Experimental/Dataflow/Engine/Private/Dataflow/DataflowDebugDrawComponent.cpp
@@ -158,7 +158,7 @@ void FDataflowDebugRenderSceneProxy::GetDynamicMeshElementsForView(const FSceneV
 	// Draw Points
 	for (const FDebugPoint& Point : Points)
 	{
-		PDI->DrawPoint(Point.Position, Point.Color, Point.Size, SDPG_World);
+		PDI->DrawPoint(Point.Position, Point.Color, Point.Size, Point.Priority);
 	}
 
 	// Draw Objects
diff --git a/Engine/Source/Runtime/Experimental/Dataflow/Engine/Public/Dataflow/DataflowDebugDrawComponent.h b/Engine/Source/Runtime/Experimental/Dataflow/Engine/Public/Dataflow/DataflowDebugDrawComponent.h
--- a/Engine/Source/Runtime/Experimental/Dataflow/Engine/Public/Dataflow/DataflowDebugDrawComponent.h
+++ b/Engine/Source/Runtime/Experimental/Dataflow/Engine/Public/Dataflow/DataflowDebugDrawComponent.h
@@ -32,6 +32,7 @@ public:
 		FVector Position;
 		float	Size = 0;
 		FLinearColor Color = FLinearColor::White;
+		ESceneDepthPriorityGroup Priority = ESceneDepthPriorityGroup::SDPG_World;
 	};
 	
 	/** Add point to the scene proxy */
