disable the 'copy convex hulls from root' dataflownode

diff --git a/Engine/Plugins/Experimental/GeometryCollectionPlugin/Source/GeometryCollectionNodes/Private/Dataflow/GeometryCollectionUtilityNodes.cpp b/Engine/Plugins/Experimental/GeometryCollectionPlugin/Source/GeometryCollectionNodes/Private/Dataflow/GeometryCollectionUtilityNodes.cpp
--- a/Engine/Plugins/Experimental/GeometryCollectionPlugin/Source/GeometryCollectionNodes/Private/Dataflow/GeometryCollectionUtilityNodes.cpp
+++ b/Engine/Plugins/Experimental/GeometryCollectionPlugin/Source/GeometryCollectionNodes/Private/Dataflow/GeometryCollectionUtilityNodes.cpp
@@ -36,7 +36,8 @@ namespace UE::Dataflow
 		DATAFLOW_NODE_REGISTER_CREATION_FACTORY(FGenerateClusterConvexHullsFromLeafHullsDataflowNode);
 		DATAFLOW_NODE_REGISTER_CREATION_FACTORY(FGenerateClusterConvexHullsFromChildrenHullsDataflowNode);
 		DATAFLOW_NODE_REGISTER_CREATION_FACTORY(FClearConvexHullsDataflowNode);
-		DATAFLOW_NODE_REGISTER_CREATION_FACTORY(FCopyConvexHullsFromRootDataflowNode);
+		// Note: FCopyConvexHullsFromRootDataflowNode is temporarily disabled as we rework its functionality
+		//DATAFLOW_NODE_REGISTER_CREATION_FACTORY(FCopyConvexHullsFromRootDataflowNode);
 		DATAFLOW_NODE_REGISTER_CREATION_FACTORY(FMergeConvexHullsDataflowNode);
 		DATAFLOW_NODE_REGISTER_CREATION_FACTORY(FUpdateVolumeAttributesDataflowNode);
 		DATAFLOW_NODE_REGISTER_CREATION_FACTORY(FGetConvexHullVolumeDataflowNode);
