Dataflow : add collection rendering to pruning geometry collection node

diff --git a/Engine/Plugins/Experimental/GeometryCollectionPlugin/Source/GeometryCollectionNodes/Public/Dataflow/GeometryCollectionEditNodes.h b/Engine/Plugins/Experimental/GeometryCollectionPlugin/Source/GeometryCollectionNodes/Public/Dataflow/GeometryCollectionEditNodes.h
--- a/Engine/Plugins/Experimental/GeometryCollectionPlugin/Source/GeometryCollectionNodes/Public/Dataflow/GeometryCollectionEditNodes.h
+++ b/Engine/Plugins/Experimental/GeometryCollectionPlugin/Source/GeometryCollectionNodes/Public/Dataflow/GeometryCollectionEditNodes.h
@@ -27,6 +27,7 @@ struct FPruneInCollectionDataflowNode : public FDataflowNode
 {
 	GENERATED_USTRUCT_BODY()
 	DATAFLOW_NODE_DEFINE_INTERNAL(FPruneInCollectionDataflowNode, "PruneInCollection", "GeometryCollection|Edit", "")
+	DATAFLOW_NODE_RENDER_TYPE("SurfaceRender", FGeometryCollection::StaticType(), "Collection")
 
 public:
 	/** Fractured GeometryCollection to prune */
@@ -70,6 +71,7 @@ struct FSetVisibilityInCollectionDataflowNode : public FDataflowNode
 {
 	GENERATED_USTRUCT_BODY()
 	DATAFLOW_NODE_DEFINE_INTERNAL(FSetVisibilityInCollectionDataflowNode, "SetVisibilityInCollection", "GeometryCollection|Edit", "")
+	DATAFLOW_NODE_RENDER_TYPE("SurfaceRender", FGeometryCollection::StaticType(), "Collection")
 
 public:
 	/** What to set the visibility of the selected bones */
@@ -112,6 +114,7 @@ struct FMergeInCollectionDataflowNode : public FDataflowNode
 {
 	GENERATED_USTRUCT_BODY()
 	DATAFLOW_NODE_DEFINE_INTERNAL(FMergeInCollectionDataflowNode, "MergeInCollection", "GeometryCollection|Edit", "")
+	DATAFLOW_NODE_RENDER_TYPE("SurfaceRender", FGeometryCollection::StaticType(), "Collection")
 
 public:
 	/** Fractured GeometryCollection to merge */
