Dataflow : fix geometry collection generation nodes to output root proxy mesh array
- static mesh to collection
- blueprint to collection
- mesh to collection
- source to collection

diff --git a/Engine/Plugins/Experimental/GeometryCollectionPlugin/Source/GeometryCollectionNodes/Private/Dataflow/GeometryCollectionAssetNodes.cpp b/Engine/Plugins/Experimental/GeometryCollectionPlugin/Source/GeometryCollectionNodes/Private/Dataflow/GeometryCollectionAssetNodes.cpp
--- a/Engine/Plugins/Experimental/GeometryCollectionPlugin/Source/GeometryCollectionNodes/Private/Dataflow/GeometryCollectionAssetNodes.cpp
+++ b/Engine/Plugins/Experimental/GeometryCollectionPlugin/Source/GeometryCollectionNodes/Private/Dataflow/GeometryCollectionAssetNodes.cpp
@@ -10,6 +10,7 @@
 #include "GeometryCollection/GeometryCollectionClusteringUtility.h"
 #include "GeometryCollection/Facades/CollectionHierarchyFacade.h"
 #include "GeometryCollection/Facades/CollectionInstancedMeshFacade.h"
+#include "Components/StaticMeshComponent.h"
 #include "Materials/Material.h"
 #include "Materials/MaterialInterface.h"
 #include "PreviewScene.h"
@@ -290,7 +291,7 @@ FCreateGeometryCollectionFromSourcesDataflowNode::FCreateGeometryCollectionFromS
 
 void FCreateGeometryCollectionFromSourcesDataflowNode::Evaluate(UE::Dataflow::FContext& Context, const FDataflowOutput* Out) const
 {
-	ensure(Out->IsA(&Collection) || Out->IsA(&Materials) || Out->IsA(&InstancedMeshes));
+	ensure(Out->IsA(&Collection) || Out->IsA(&Materials) || Out->IsA(&MaterialInstances) || Out->IsA(&InstancedMeshes));
 	
 	const TArray<FGeometryCollectionSource>& InSources = GetValue(Context, &Sources);
 
@@ -375,17 +376,19 @@ FCreateGeometryCollectionFromSourcesDataflowNode_v2::FCreateGeometryCollectionFr
 	RegisterOutputConnection(&Collection);
 	RegisterOutputConnection(&Materials);
 	RegisterOutputConnection(&InstancedMeshes);
+	RegisterOutputConnection(&RootProxyMeshes);
 }
 
 void FCreateGeometryCollectionFromSourcesDataflowNode_v2::Evaluate(UE::Dataflow::FContext& Context, const FDataflowOutput* Out) const
 {
-	ensure(Out->IsA(&Collection) || Out->IsA(&Materials) || Out->IsA(&InstancedMeshes));
+	ensure(Out->IsA(&Collection) || Out->IsA(&Materials) || Out->IsA(&InstancedMeshes) || Out->IsA(&RootProxyMeshes));
 
 	const TArray<FGeometryCollectionSource>& InSources = GetValue(Context, &Sources);
 
 	FGeometryCollection OutCollection;
 	TArray<TObjectPtr<UMaterialInterface>> OutMaterials;
 	TArray<FGeometryCollectionAutoInstanceMesh> OutInstancedMeshes;
+	TArray<FDataflowRootProxyMesh> OutRootProxyMeshes;
 
 	// make sure we have an attribute for instanced meshes
 	GeometryCollection::Facades::FCollectionInstancedMeshFacade InstancedMeshFacade(OutCollection);
@@ -400,9 +403,11 @@ void FCreateGeometryCollectionFromSourcesDataflowNode_v2::Evaluate(UE::Dataflow:
 		// todo: change AppendGeometryCollectionSource to take a FManagedArrayCollection so we could move the collection when assigning it to the output
 		FGeometryCollectionEngineConversion::AppendGeometryCollectionSource(Source, OutCollection, MutableView(OutMaterials), bReindexMaterialsInLoop);
 
+		UStaticMesh* StaticMeshObject = Cast<UStaticMesh>(Source.SourceGeometryObject.TryLoad());
+
 		// todo(chaos) if the source is a geometry collection this will not work properly 
 		FGeometryCollectionAutoInstanceMesh InstancedMesh;
-		InstancedMesh.Mesh = Cast<UStaticMesh>(Source.SourceGeometryObject.TryLoad());
+		InstancedMesh.Mesh = StaticMeshObject;
 		InstancedMesh.Materials = Source.SourceMaterial;
 		InstancedMesh.NumInstances = 0;
 		InstancedMesh.CustomData.Reset();
@@ -422,6 +427,13 @@ void FCreateGeometryCollectionFromSourcesDataflowNode_v2::Evaluate(UE::Dataflow:
 		{
 			InstancedMeshFacade.SetIndex(TransformIndex, InstancedMeshIndex);
 		}
+
+		// Root proxy meshes
+		FDataflowRootProxyMesh RootProxyMesh;
+		RootProxyMesh.Mesh = StaticMeshObject;
+		RootProxyMesh.Transform = Source.LocalTransform;
+		OutRootProxyMeshes.Add(RootProxyMesh);
+
 	}
 	if (bReindexMaterialsInLoop == false)
 	{
@@ -449,6 +461,7 @@ void FCreateGeometryCollectionFromSourcesDataflowNode_v2::Evaluate(UE::Dataflow:
 	SetValue(Context, static_cast<const FManagedArrayCollection&>(OutCollection), &Collection);
 	SetValue(Context, MoveTemp(OutMaterials), &Materials);
 	SetValue(Context, MoveTemp(OutInstancedMeshes), &InstancedMeshes);
+	SetValue(Context, MoveTemp(OutRootProxyMeshes), &RootProxyMeshes);
 }
 
 // ===========================================================================================================================
@@ -465,7 +478,7 @@ FGeometryCollectionToCollectionDataflowNode::FGeometryCollectionToCollectionData
 
 void FGeometryCollectionToCollectionDataflowNode::Evaluate(UE::Dataflow::FContext& Context, const FDataflowOutput* Out) const
 {
-	ensure(Out->IsA(&Collection) || Out->IsA(&Materials) || Out->IsA(&InstancedMeshes));
+	ensure(Out->IsA(&Collection) || Out->IsA(&Materials) || Out->IsA(&MaterialInstances) || Out->IsA(&InstancedMeshes));
 
 	FManagedArrayCollection OutCollection;
 	TArray<TObjectPtr<UMaterialInterface>> OutMaterialInstances;
@@ -494,25 +507,36 @@ FGeometryCollectionToCollectionDataflowNode_v2::FGeometryCollectionToCollectionD
 	RegisterOutputConnection(&Collection);
 	RegisterOutputConnection(&Materials);
 	RegisterOutputConnection(&InstancedMeshes);
+	RegisterOutputConnection(&RootProxyMeshes);
 }
 
 void FGeometryCollectionToCollectionDataflowNode_v2::Evaluate(UE::Dataflow::FContext& Context, const FDataflowOutput* Out) const
 {
-	ensure(Out->IsA(&Collection) || Out->IsA(&Materials) || Out->IsA(&InstancedMeshes));
+	ensure(Out->IsA(&Collection) || Out->IsA(&Materials) || Out->IsA(&InstancedMeshes) || Out->IsA(&RootProxyMeshes));
 
 	FManagedArrayCollection OutCollection;
 	TArray<TObjectPtr<UMaterialInterface>> OutMaterials;
 	TArray<FGeometryCollectionAutoInstanceMesh> OutInstancedMeshes;
+	TArray<FDataflowRootProxyMesh> OutRootProxyMeshes;
 
 	if (GeometryCollection)
 	{
 		FGeometryCollectionEngineConversion::ConvertGeometryCollectionToGeometryCollection(GeometryCollection, OutCollection, OutMaterials, OutInstancedMeshes);
+
+		for (int32 MeshIndex = 0; MeshIndex < GeometryCollection->RootProxyData.ProxyMeshes.Num(); MeshIndex++)
+		{
+			FDataflowRootProxyMesh RootProxyMesh;
+			RootProxyMesh.Mesh = GeometryCollection->RootProxyData.ProxyMeshes[MeshIndex];
+			RootProxyMesh.Transform = FTransform(GeometryCollection->RootProxyData.GetMeshTransform(MeshIndex));
+			OutRootProxyMeshes.Emplace(RootProxyMesh);
+		}
 	}
 
 	// Set Outputs
 	SetValue(Context, MoveTemp(OutCollection), &Collection);
 	SetValue(Context, MoveTemp(OutMaterials), &Materials);
 	SetValue(Context, MoveTemp(OutInstancedMeshes), &InstancedMeshes);
+	SetValue(Context, MoveTemp(OutRootProxyMeshes), &RootProxyMeshes);
 }
 
 // ===========================================================================================================================
@@ -529,7 +553,7 @@ FBlueprintToCollectionDataflowNode::FBlueprintToCollectionDataflowNode(const UE:
 
 void FBlueprintToCollectionDataflowNode::Evaluate(UE::Dataflow::FContext& Context, const FDataflowOutput* Out) const
 {
-	ensure(Out->IsA(&Collection) || Out->IsA(&Materials) || Out->IsA(&InstancedMeshes));
+	ensure(Out->IsA(&Collection) || Out->IsA(&Materials) || Out->IsA(&MaterialInstances) || Out->IsA(&InstancedMeshes));
 
 	FManagedArrayCollection OutCollection;
 	TArray<TObjectPtr<UMaterialInterface>> OutMaterialInstances;
@@ -573,15 +597,17 @@ FBlueprintToCollectionDataflowNode_v2::FBlueprintToCollectionDataflowNode_v2(con
 	RegisterOutputConnection(&Collection);
 	RegisterOutputConnection(&Materials);
 	RegisterOutputConnection(&InstancedMeshes);
+	RegisterOutputConnection(&RootProxyMeshes);
 }
 
 void FBlueprintToCollectionDataflowNode_v2::Evaluate(UE::Dataflow::FContext& Context, const FDataflowOutput* Out) const
 {
-	ensure(Out->IsA(&Collection) || Out->IsA(&Materials) || Out->IsA(&InstancedMeshes));
+	ensure(Out->IsA(&Collection) || Out->IsA(&Materials) || Out->IsA(&InstancedMeshes) || Out->IsA(&RootProxyMeshes));
 
 	FManagedArrayCollection OutCollection;
 	TArray<TObjectPtr<UMaterialInterface>> OutMaterials;
 	TArray<FGeometryCollectionAutoInstanceMesh> OutInstancedMeshes;
+	TArray<FDataflowRootProxyMesh> OutRootProxyMeshes;
 
 	if (Blueprint)
 	{
@@ -598,6 +624,27 @@ void FBlueprintToCollectionDataflowNode_v2::Evaluate(UE::Dataflow::FContext& Con
 				{
 					FGeometryCollectionEngineConversion::FSkeletalMeshToCollectionConversionParameters ConversionParameters;
 					FGeometryCollectionEngineConversion::ConvertActorToGeometryCollection(PreviewActor, OutCollection, OutMaterials, OutInstancedMeshes, ConversionParameters, bSplitComponents);
+
+					// root proxies
+					const FTransform ActorTransform(PreviewActor->GetTransform());
+
+					TInlineComponentArray<UStaticMeshComponent*> StaticMeshComponents(PreviewActor);
+					for (UStaticMeshComponent* StaticMeshComponent : StaticMeshComponents)
+					{
+						if (StaticMeshComponent)
+						{
+							if (UStaticMesh* StaticMeshObject = StaticMeshComponent->GetStaticMesh())
+							{
+								FTransform MeshTransform(StaticMeshComponent->GetComponentTransform());
+								MeshTransform.SetTranslation((MeshTransform.GetTranslation() - ActorTransform.GetTranslation()));
+
+								FDataflowRootProxyMesh RootProxyMesh;
+								RootProxyMesh.Mesh = StaticMeshObject;
+								RootProxyMesh.Transform = MeshTransform;
+								OutRootProxyMeshes.Add(RootProxyMesh);
+							}
+						}
+					}
 				}
 			}
 		}
@@ -607,4 +654,5 @@ void FBlueprintToCollectionDataflowNode_v2::Evaluate(UE::Dataflow::FContext& Con
 	SetValue(Context, MoveTemp(OutCollection), &Collection);
 	SetValue(Context, MoveTemp(OutMaterials), &Materials);
 	SetValue(Context, MoveTemp(OutInstancedMeshes), &InstancedMeshes);
+	SetValue(Context, MoveTemp(OutRootProxyMeshes), &RootProxyMeshes);
 }
\ No newline at end of file
diff --git a/Engine/Plugins/Experimental/GeometryCollectionPlugin/Source/GeometryCollectionNodes/Private/Dataflow/GeometryCollectionStaticMeshToCollectionNode.cpp b/Engine/Plugins/Experimental/GeometryCollectionPlugin/Source/GeometryCollectionNodes/Private/Dataflow/GeometryCollectionStaticMeshToCollectionNode.cpp
--- a/Engine/Plugins/Experimental/GeometryCollectionPlugin/Source/GeometryCollectionNodes/Private/Dataflow/GeometryCollectionStaticMeshToCollectionNode.cpp
+++ b/Engine/Plugins/Experimental/GeometryCollectionPlugin/Source/GeometryCollectionNodes/Private/Dataflow/GeometryCollectionStaticMeshToCollectionNode.cpp
@@ -32,27 +32,35 @@ FStaticMeshToCollectionDataflowNode_v2::FStaticMeshToCollectionDataflowNode_v2(c
 	RegisterOutputConnection(&Collection);
 	RegisterOutputConnection(&Materials);
 	RegisterOutputConnection(&InstancedMeshes);
+	RegisterOutputConnection(&RootProxyMeshes);
 }
 
 void FStaticMeshToCollectionDataflowNode_v2::Evaluate(UE::Dataflow::FContext& Context, const FDataflowOutput* Out) const
 {
-	ensure(Out->IsA(&Collection) || Out->IsA(&Materials) || Out->IsA(&InstancedMeshes));
+	ensure(Out->IsA(&Collection) || Out->IsA(&Materials) || Out->IsA(&InstancedMeshes) || Out->IsA(&RootProxyMeshes));
 
 	FManagedArrayCollection OutCollection;
 	TArray<TObjectPtr<UMaterialInterface>> OutMaterials;
 	TArray<FGeometryCollectionAutoInstanceMesh> OutInstancedMeshes;
+	TArray<FDataflowRootProxyMesh> OutRootProxyMeshes;
 
 	TObjectPtr<UStaticMesh> StaticMeshVal = GetValue(Context, &StaticMesh, StaticMesh);
 	const FTransform& InMeshTransform = GetValue(Context, &MeshTransform, MeshTransform);
 	if (StaticMeshVal)
 	{
 		FGeometryCollectionEngineConversion::ConvertStaticMeshToGeometryCollection(StaticMeshVal, InMeshTransform, OutCollection, OutMaterials, OutInstancedMeshes, bSetInternalFromMaterialIndex, bSplitComponents);
+
+		FDataflowRootProxyMesh RootProxyMesh;
+		RootProxyMesh.Mesh = StaticMeshVal;
+		RootProxyMesh.Transform = InMeshTransform;
+		OutRootProxyMeshes.Add(RootProxyMesh);
 	}
 
 	// Set Outputs
 	SetValue(Context, MoveTemp(OutCollection), &Collection);
 	SetValue(Context, MoveTemp(OutMaterials), &Materials);
 	SetValue(Context, MoveTemp(OutInstancedMeshes), &InstancedMeshes);
+	SetValue(Context, MoveTemp(OutRootProxyMeshes), &RootProxyMeshes);
 }
 
 
diff --git a/Engine/Plugins/Experimental/GeometryCollectionPlugin/Source/GeometryCollectionNodes/Public/Dataflow/GeometryCollectionAssetNodes.h b/Engine/Plugins/Experimental/GeometryCollectionPlugin/Source/GeometryCollectionNodes/Public/Dataflow/GeometryCollectionAssetNodes.h
--- a/Engine/Plugins/Experimental/GeometryCollectionPlugin/Source/GeometryCollectionNodes/Public/Dataflow/GeometryCollectionAssetNodes.h
+++ b/Engine/Plugins/Experimental/GeometryCollectionPlugin/Source/GeometryCollectionNodes/Public/Dataflow/GeometryCollectionAssetNodes.h
@@ -280,6 +280,10 @@ private:
 	/** array of instanced meshes*/
 	UPROPERTY(meta = (DataflowOutput))
 	TArray<FGeometryCollectionAutoInstanceMesh> InstancedMeshes;
+
+	/** corresponding source proxies */
+	UPROPERTY(meta = (DataflowOutput))
+	TArray<FDataflowRootProxyMesh> RootProxyMeshes;
 };
 
 /**
@@ -348,6 +352,10 @@ private:
 	UPROPERTY(meta = (DataflowOutput))
 	TArray<FGeometryCollectionAutoInstanceMesh> InstancedMeshes;
 
+	/** corresponding source proxies */
+	UPROPERTY(meta = (DataflowOutput))
+	TArray<FDataflowRootProxyMesh> RootProxyMeshes;
+
 	virtual void Evaluate(UE::Dataflow::FContext& Context, const FDataflowOutput* Out) const override;
 };
 
@@ -426,6 +434,10 @@ private:
 	UPROPERTY(meta = (DataflowOutput))
 	TArray<FGeometryCollectionAutoInstanceMesh> InstancedMeshes;
 
+	/** corresponding source proxies */
+	UPROPERTY(meta = (DataflowOutput))
+	TArray<FDataflowRootProxyMesh> RootProxyMeshes;
+
 	virtual void Evaluate(UE::Dataflow::FContext& Context, const FDataflowOutput* Out) const override;
 };
 
diff --git a/Engine/Plugins/Experimental/GeometryCollectionPlugin/Source/GeometryCollectionNodes/Public/Dataflow/GeometryCollectionStaticMeshToCollectionNode.h b/Engine/Plugins/Experimental/GeometryCollectionPlugin/Source/GeometryCollectionNodes/Public/Dataflow/GeometryCollectionStaticMeshToCollectionNode.h
--- a/Engine/Plugins/Experimental/GeometryCollectionPlugin/Source/GeometryCollectionNodes/Public/Dataflow/GeometryCollectionStaticMeshToCollectionNode.h
+++ b/Engine/Plugins/Experimental/GeometryCollectionPlugin/Source/GeometryCollectionNodes/Public/Dataflow/GeometryCollectionStaticMeshToCollectionNode.h
@@ -4,6 +4,7 @@
 
 #include "CoreMinimal.h"
 #include "Dataflow/DataflowEngine.h"
+#include "Dataflow/GeometryCollectionAssetNodes.h"
 #include "GeometryCollection/GeometryCollectionObject.h"
 
 #include "GeometryCollectionStaticMeshToCollectionNode.generated.h"
@@ -49,6 +50,10 @@ private:
 	UPROPERTY(meta = (DataflowOutput))
 	TArray<FGeometryCollectionAutoInstanceMesh> InstancedMeshes;
 
+	/** corresponding source proxies */
+	UPROPERTY(meta = (DataflowOutput))
+	TArray<FDataflowRootProxyMesh> RootProxyMeshes;
+
 	/** Set the internal faces from material index */
 	UPROPERTY(EditAnywhere, Category = "Asset");
 	bool bSetInternalFromMaterialIndex = true;
