set limits on the grid point scatter in dataflow to fix crash if the number of points requested is large

diff --git a/Engine/Plugins/Experimental/GeometryCollectionPlugin/Source/GeometryCollectionNodes/Private/Dataflow/GeometryCollectionFracturingNodes.cpp b/Engine/Plugins/Experimental/GeometryCollectionPlugin/Source/GeometryCollectionNodes/Private/Dataflow/GeometryCollectionFracturingNodes.cpp
--- a/Engine/Plugins/Experimental/GeometryCollectionPlugin/Source/GeometryCollectionNodes/Private/Dataflow/GeometryCollectionFracturingNodes.cpp
+++ b/Engine/Plugins/Experimental/GeometryCollectionPlugin/Source/GeometryCollectionNodes/Private/Dataflow/GeometryCollectionFracturingNodes.cpp
@@ -355,11 +355,14 @@ void FGridScatterPointsDataflowNode::Evaluate(UE::Dataflow::FContext& Context, c
 		if (BBox.GetVolume() > 0.f)
 		{
 			const FVector Extent(BBox.Max - BBox.Min);
-			const int32 NumPointsInX = GetValue<int32>(Context, &NumberOfPointsInX);
-			const int32 NumPointsInY = GetValue<int32>(Context, &NumberOfPointsInY);
-			const int32 NumPointsInZ = GetValue<int32>(Context, &NumberOfPointsInZ);
+			// Note: Should match ClampMax in the UI. Do not raise above 1290 to avoid overflowing the TArray.
+			// (A smaller limit is preferrable because the rendering / later processing will likely not want to handle that many points.)
+			constexpr int32 MaxPointsPerDim = 200;
+			const int32 NumPointsInX = FMath::Clamp(GetValue<int32>(Context, &NumberOfPointsInX), 0, MaxPointsPerDim);
+			const int32 NumPointsInY = FMath::Clamp(GetValue<int32>(Context, &NumberOfPointsInY), 0, MaxPointsPerDim);
+			const int32 NumPointsInZ = FMath::Clamp(GetValue<int32>(Context, &NumberOfPointsInZ), 0, MaxPointsPerDim);
 			
-			if (NumPointsInX >= 2 && NumPointsInY >= 2 && NumPointsInZ >= 2)
+			if (NumPointsInX >= 1 && NumPointsInY >= 1 && NumPointsInZ >= 1)
 			{
 				const int32 NumPoints = NumPointsInX * NumPointsInY * NumPointsInZ;
 				const float dX = Extent.X / (float)NumPointsInX;
diff --git a/Engine/Plugins/Experimental/GeometryCollectionPlugin/Source/GeometryCollectionNodes/Public/Dataflow/GeometryCollectionFracturingNodes.h b/Engine/Plugins/Experimental/GeometryCollectionPlugin/Source/GeometryCollectionNodes/Public/Dataflow/GeometryCollectionFracturingNodes.h
--- a/Engine/Plugins/Experimental/GeometryCollectionPlugin/Source/GeometryCollectionNodes/Public/Dataflow/GeometryCollectionFracturingNodes.h
+++ b/Engine/Plugins/Experimental/GeometryCollectionPlugin/Source/GeometryCollectionNodes/Public/Dataflow/GeometryCollectionFracturingNodes.h
@@ -317,15 +317,15 @@ struct FGridScatterPointsDataflowNode : public FDataflowNode
 
 public:
 	/** Number of points in X direction */
-	UPROPERTY(EditAnywhere, Category = "Scatter", meta = (DataflowInput, UIMin = "0"));
+	UPROPERTY(EditAnywhere, Category = "Scatter", meta = (DataflowInput, ClampMin = "0", UIMax = "100", ClampMax = "200"));
 	int32 NumberOfPointsInX = 5;
 
 	/** Number of points in Y direction */
-	UPROPERTY(EditAnywhere, Category = "Scatter", meta = (DataflowInput, UIMin = "0"));
+	UPROPERTY(EditAnywhere, Category = "Scatter", meta = (DataflowInput, ClampMin = "0", UIMax = "100", ClampMax = "200"));
 	int32 NumberOfPointsInY = 5;
 
 	/** Number of points in Z direction */
-	UPROPERTY(EditAnywhere, Category = "Scatter", meta = (DataflowInput, UIMin = "0"));
+	UPROPERTY(EditAnywhere, Category = "Scatter", meta = (DataflowInput, ClampMin = "0", UIMax = "100", ClampMax = "200"));
 	int32 NumberOfPointsInZ = 5;
 
 	/** Seed for random */
