add missing transforms to dataflow geometry collection proximity rendering

diff --git a/Engine/Plugins/Experimental/GeometryCollectionPlugin/Source/GeometryCollectionNodes/Private/Dataflow/GeometryCollectionUtils.cpp b/Engine/Plugins/Experimental/GeometryCollectionPlugin/Source/GeometryCollectionNodes/Private/Dataflow/GeometryCollectionUtils.cpp
--- a/Engine/Plugins/Experimental/GeometryCollectionPlugin/Source/GeometryCollectionNodes/Private/Dataflow/GeometryCollectionUtils.cpp
+++ b/Engine/Plugins/Experimental/GeometryCollectionPlugin/Source/GeometryCollectionNodes/Private/Dataflow/GeometryCollectionUtils.cpp
@@ -2,6 +2,8 @@
 
 #include "Dataflow/GeometryCollectionUtils.h"
 
+#include "GeometryCollection/GeometryCollectionAlgo.h"
+
 namespace UE::Dataflow
 {
 	namespace Utils
@@ -26,6 +28,25 @@ namespace UE::Dataflow
 			{
 				const TManagedArray<TSet<int32>>& Proximity = Collection.GetAttribute<TSet<int32>>("Proximity", FGeometryCollection::GeometryGroup);
 				const TManagedArray<FBox>& BoundingBox = Collection.GetAttribute<FBox>("BoundingBox", FGeometryCollection::GeometryGroup);
+				const TManagedArray<FTransform3f>* Transform = Collection.FindAttributeTyped<FTransform3f>(FGeometryCollection::TransformAttribute, FGeometryCollection::TransformGroup);
+				const TManagedArray<int32>* Parent = Collection.FindAttributeTyped<int32>(FGeometryCollection::ParentAttribute, FGeometryCollection::TransformGroup);
+				const TManagedArray<int32>* TransformIndex = Collection.FindAttributeTyped<int32>(FGeometryCollection::TransformIndexAttribute, FGeometryCollection::GeometryGroup);
+				TArray<FTransform> GlobalTransformArray;
+				if (Transform && Parent && TransformIndex)
+				{
+					GeometryCollectionAlgo::GlobalMatrices(*Transform, *Parent, GlobalTransformArray);
+				}
+				auto TransformFromGeometry = [&GlobalTransformArray, &TransformIndex](int32 GeometryIdx, FVector Point) -> FVector
+				{
+					if (GlobalTransformArray.IsEmpty() || !TransformIndex || !TransformIndex->IsValidIndex(GeometryIdx))
+					{
+						return Point;
+					}
+					else
+					{
+						return GlobalTransformArray[(*TransformIndex)[GeometryIdx]].TransformPosition(Point);
+					}
+				};
 
 				const int32 NumGeometry = Collection.NumElements(FGeometryCollection::GeometryGroup);
 				TArray<FVector> Centers;
@@ -33,7 +54,7 @@ namespace UE::Dataflow
 				{
 					if (Proximity[Idx].Num() > 0)
 					{
-						const FVector Center = BoundingBox[Idx].GetCenter();
+						const FVector Center = TransformFromGeometry(Idx, BoundingBox[Idx].GetCenter());
 
 						Centers.Add(Center);
 
@@ -48,7 +69,7 @@ namespace UE::Dataflow
 
 						for (int32 IdxOther : Proximity[Idx])
 						{
-							const FVector CenterOther = BoundingBox[IdxOther].GetCenter();
+							const FVector CenterOther = TransformFromGeometry(IdxOther, BoundingBox[IdxOther].GetCenter());
 
 							DataflowRenderingInterface.DrawLine(Center, CenterOther);
 						}
@@ -60,7 +81,7 @@ namespace UE::Dataflow
 							DataflowRenderingInterface.DrawPoint(Point);
 						}
 
-						Centers.Empty();
+						Centers.Reset();
 					}
 				}
 			}
