fix planecutter dataflow visualization missing the collection transform's effect for the plane rendering

diff --git a/Engine/Plugins/Experimental/GeometryCollectionPlugin/Source/GeometryCollectionNodes/Private/Dataflow/GeometryCollectionFracturingNodes.cpp b/Engine/Plugins/Experimental/GeometryCollectionPlugin/Source/GeometryCollectionNodes/Private/Dataflow/GeometryCollectionFracturingNodes.cpp
--- a/Engine/Plugins/Experimental/GeometryCollectionPlugin/Source/GeometryCollectionNodes/Private/Dataflow/GeometryCollectionFracturingNodes.cpp
+++ b/Engine/Plugins/Experimental/GeometryCollectionPlugin/Source/GeometryCollectionNodes/Private/Dataflow/GeometryCollectionFracturingNodes.cpp
@@ -803,6 +803,8 @@ void FPlaneCutterDataflowNode_v2::DebugDraw(UE::Dataflow::FContext& Context, IDa
 
 		TArray<FSimpleDebugDrawMesh> DebugMeshes;
 		DebugMeshes.SetNum(PlaneTransforms.Num());
+		
+		FTransform CollectionTransform = GetValue(Context, &Transform);
 
 		for (int32 PlaneIdx = 0; PlaneIdx < PlaneTransforms.Num(); ++PlaneIdx)
 		{
@@ -819,7 +821,7 @@ void FPlaneCutterDataflowNode_v2::DebugDraw(UE::Dataflow::FContext& Context, IDa
 			const int32 HeightVertexCount = FMath::Min(MaxCountPerDim, Height / NoiseSettings.PointSpacing);
 
 			const FTransform& PlaneTransform = PlaneTransforms[PlaneIdx];
-			const FVector NoisePivot = GetValue(Context, &Transform).GetLocation();
+			const FVector NoisePivot = CollectionTransform.GetLocation();
 			FVector Normal = PlaneTransform.GetUnitAxis(EAxis::Z);
 
 			FSimpleDebugDrawMesh Mesh;
@@ -832,7 +834,7 @@ void FPlaneCutterDataflowNode_v2::DebugDraw(UE::Dataflow::FContext& Context, IDa
 				FVector WorldPos = PlaneTransforms[PlaneIdx].TransformPosition(DebugMeshes[PlaneIdx].Vertices[VertexIdx]);
 				FVector NewWorldPos = WorldPos + NoiseSettings.NoiseVector(WorldPos - NoisePivot, NoiseOffset).Dot(Normal) * Normal;
 				
-				DebugMeshes[PlaneIdx].Vertices[VertexIdx] = NewWorldPos;
+				DebugMeshes[PlaneIdx].Vertices[VertexIdx] = CollectionTransform.InverseTransformPosition(NewWorldPos);
 			}
 
 			DataflowRenderingInterface.DrawMesh(DebugMeshes[PlaneIdx]);
