Require opt-in of UV resizing from within cloth asset.

Added a new Dataflow Boolean Operation node to make it easier to && conditions in the graph.

Updated ResizeOutfitTemplate with this change and changed all of the "SkipXYZ" variables to "XYZ" and switched the logic for them.

diff --git a/Engine/Plugins/ChaosClothAssetEditor/Source/ChaosClothAssetDataflowNodes/Private/ChaosClothAsset/ChaosClothAssetDataflowNodesModule.cpp b/Engine/Plugins/ChaosClothAssetEditor/Source/ChaosClothAssetDataflowNodes/Private/ChaosClothAsset/ChaosClothAssetDataflowNodesModule.cpp
--- a/Engine/Plugins/ChaosClothAssetEditor/Source/ChaosClothAssetDataflowNodes/Private/ChaosClothAsset/ChaosClothAssetDataflowNodesModule.cpp
+++ b/Engine/Plugins/ChaosClothAssetEditor/Source/ChaosClothAssetDataflowNodes/Private/ChaosClothAsset/ChaosClothAssetDataflowNodesModule.cpp
@@ -16,6 +16,7 @@
 #include "ChaosClothAsset/CustomRegionResizingNode.h"
 #include "ChaosClothAsset/DatasmithImportNode.h"
 #include "ChaosClothAsset/DeleteElementNode.h"
+#include "ChaosClothAsset/EnableUVResizingNode.h"
 #include "ChaosClothAsset/GenerateSimMorphTargetNode.h"
 #include "ChaosClothAsset/ImportFilePathCustomization.h"
 #include "ChaosClothAsset/ImportNode.h"
@@ -111,6 +112,7 @@ namespace UE::Chaos::ClothAsset
 			DATAFLOW_NODE_REGISTER_CREATION_FACTORY(FChaosClothAssetCustomRegionResizingNode);
 			DATAFLOW_NODE_REGISTER_CREATION_FACTORY(FChaosClothAssetDatasmithImportNode);
 			DATAFLOW_NODE_REGISTER_CREATION_FACTORY(FChaosClothAssetDeleteElementNode);
+			DATAFLOW_NODE_REGISTER_CREATION_FACTORY(FChaosClothAssetEnableUVResizingNode);
 			DATAFLOW_NODE_REGISTER_CREATION_FACTORY(FChaosClothAssetExtractWeightMapNode);
 			DATAFLOW_NODE_REGISTER_CREATION_FACTORY(FChaosClothAssetExtractSelectionSetNode);
 			DATAFLOW_NODE_REGISTER_CREATION_FACTORY(FChaosClothAssetGenerateSimMorphTargetNode);
diff --git a/Engine/Plugins/ChaosClothAssetEditor/Source/ChaosClothAssetDataflowNodes/Private/ChaosClothAsset/ClothCollectionQueryNode.cpp b/Engine/Plugins/ChaosClothAssetEditor/Source/ChaosClothAssetDataflowNodes/Private/ChaosClothAsset/ClothCollectionQueryNode.cpp
--- a/Engine/Plugins/ChaosClothAssetEditor/Source/ChaosClothAssetDataflowNodes/Private/ChaosClothAsset/ClothCollectionQueryNode.cpp
+++ b/Engine/Plugins/ChaosClothAssetEditor/Source/ChaosClothAssetDataflowNodes/Private/ChaosClothAsset/ClothCollectionQueryNode.cpp
@@ -3,6 +3,7 @@
 #include "ChaosClothAsset/ClothCollectionQueryNode.h"
 #include "ChaosClothAsset/CollectionClothFacade.h"
 #include "ChaosClothAsset/ClothGeometryTools.h"
+#include "Chaos/CollectionPropertyFacade.h"
 
 #include UE_INLINE_GENERATED_CPP_BY_NAME(ClothCollectionQueryNode)
 
@@ -17,6 +18,7 @@ FChaosClothAssetCollectionQueryNode::FChaosClothAssetCollectionQueryNode(const U
 	RegisterOutputConnection(&bHasClothSimMesh);
 	RegisterOutputConnection(&bHasClothRenderMesh);
 	RegisterOutputConnection(&bHasClothProxyDeformer);
+	RegisterOutputConnection(&bBooleanPropertyValue);
 }
 
 void FChaosClothAssetCollectionQueryNode::Evaluate(UE::Dataflow::FContext& Context, const FDataflowOutput* Out) const
@@ -29,7 +31,8 @@ void FChaosClothAssetCollectionQueryNode::Evaluate(UE::Dataflow::FContext& Conte
 	else if (Out->IsA<bool>(&bIsClothCollection)
 		|| Out->IsA<bool>(&bHasClothSimMesh)
 		|| Out->IsA<bool>(&bHasClothRenderMesh)
-		|| Out->IsA<bool>(&bHasClothProxyDeformer))
+		|| Out->IsA<bool>(&bHasClothProxyDeformer)
+		|| Out->IsA<bool>(&bBooleanPropertyValue))
 	{
 		// Setting all of these outputs once since it involves making a temporary copy of the whole collection.
 		FManagedArrayCollection InCollection = GetValue<FManagedArrayCollection>(Context, &Collection);
@@ -40,6 +43,14 @@ void FChaosClothAssetCollectionQueryNode::Evaluate(UE::Dataflow::FContext& Conte
 		SetValue(Context, FClothGeometryTools::HasSimMesh(ClothCollection), &bHasClothSimMesh);
 		SetValue(Context, FClothGeometryTools::HasRenderMesh(ClothCollection), &bHasClothRenderMesh);
 		SetValue(Context, ClothFacade.IsValid(EClothCollectionOptionalSchemas::RenderDeformer), &bHasClothProxyDeformer);
+
+		Chaos::Softs::FCollectionPropertyConstFacade PropertyFacade(ClothCollection);
+		bool bBooleanPropertyResult = bBooleanPropertyValue;
+		if (PropertyFacade.IsValid())
+		{
+			bBooleanPropertyResult = PropertyFacade.GetValue(BooleanPropertyName, bBooleanPropertyValue);
+		}
+		SetValue(Context, bBooleanPropertyResult, &bBooleanPropertyValue);
 	}
 }
 
diff --git a/Engine/Plugins/ChaosClothAssetEditor/Source/ChaosClothAssetDataflowNodes/Private/ChaosClothAsset/EnableUVResizingNode.cpp b/Engine/Plugins/ChaosClothAssetEditor/Source/ChaosClothAssetDataflowNodes/Private/ChaosClothAsset/EnableUVResizingNode.cpp
new file mode 100644
--- /dev/null
+++ b/Engine/Plugins/ChaosClothAssetEditor/Source/ChaosClothAssetDataflowNodes/Private/ChaosClothAsset/EnableUVResizingNode.cpp
@@ -0,0 +1,17 @@
+// Copyright Epic Games, Inc. All Rights Reserved.
+
+#include "ChaosClothAsset/EnableUVResizingNode.h"
+#include "Chaos/CollectionPropertyFacade.h"
+
+#include UE_INLINE_GENERATED_CPP_BY_NAME(EnableUVResizingNode)
+
+FChaosClothAssetEnableUVResizingNode::FChaosClothAssetEnableUVResizingNode(const UE::Dataflow::FNodeParameters& InParam, FGuid InGuid)
+	: FChaosClothAssetSimulationBaseConfigNode(InParam, InGuid)
+{
+	RegisterCollectionConnections();
+}
+
+void FChaosClothAssetEnableUVResizingNode::AddProperties(FPropertyHelper& PropertyHelper) const
+{
+	PropertyHelper.SetPropertyBool(TEXT("EnableUVResizing"), true, {}, ECollectionPropertyFlags::None);
+}
diff --git a/Engine/Plugins/ChaosClothAssetEditor/Source/ChaosClothAssetDataflowNodes/Public/ChaosClothAsset/ClothCollectionQueryNode.h b/Engine/Plugins/ChaosClothAssetEditor/Source/ChaosClothAssetDataflowNodes/Public/ChaosClothAsset/ClothCollectionQueryNode.h
--- a/Engine/Plugins/ChaosClothAssetEditor/Source/ChaosClothAssetDataflowNodes/Public/ChaosClothAsset/ClothCollectionQueryNode.h
+++ b/Engine/Plugins/ChaosClothAssetEditor/Source/ChaosClothAssetDataflowNodes/Public/ChaosClothAsset/ClothCollectionQueryNode.h
@@ -46,4 +46,12 @@ private:
 	/** Does this collection have proxy deformer data */
 	UPROPERTY(Meta = (DataflowOutput))
 	bool bHasClothProxyDeformer = false;
+
+	/** Name of boolean property to query */
+	UPROPERTY(EditAnywhere, Category = "Cloth Collection Query")
+	FString BooleanPropertyName = TEXT("PropertyName");
+
+	/** Result of querying BooleanPropertyName. Default value if property doesn't exist matches node value.*/
+	UPROPERTY(EditAnywhere, Category = "Cloth Collection Query", Meta =(DataflowOutput))
+	bool bBooleanPropertyValue = false;
 };
diff --git a/Engine/Plugins/ChaosClothAssetEditor/Source/ChaosClothAssetDataflowNodes/Public/ChaosClothAsset/EnableUVResizingNode.h b/Engine/Plugins/ChaosClothAssetEditor/Source/ChaosClothAssetDataflowNodes/Public/ChaosClothAsset/EnableUVResizingNode.h
new file mode 100644
--- /dev/null
+++ b/Engine/Plugins/ChaosClothAssetEditor/Source/ChaosClothAssetDataflowNodes/Public/ChaosClothAsset/EnableUVResizingNode.h
@@ -0,0 +1,22 @@
+// Copyright Epic Games, Inc. All Rights Reserved.
+
+#pragma once 
+
+#include "ChaosClothAsset/SimulationBaseConfigNode.h"
+#include "EnableUVResizingNode.generated.h"
+
+/** Node for enabling UV Resizing used by the ChaosOutfitAsset's Resizeable Outfit.*/
+USTRUCT(Meta = (DataflowCloth, Experimental))
+struct FChaosClothAssetEnableUVResizingNode final : public FChaosClothAssetSimulationBaseConfigNode
+{
+	GENERATED_USTRUCT_BODY()
+	DATAFLOW_NODE_DEFINE_INTERNAL(FChaosClothAssetEnableUVResizingNode, "EnableUVResizing", "Cloth", "Cloth Enable UV Resizing")
+	DATAFLOW_NODE_RENDER_TYPE("SurfaceRender", FName("FClothCollection"), "Collection")
+
+public:
+
+	FChaosClothAssetEnableUVResizingNode(const UE::Dataflow::FNodeParameters& InParam, FGuid InGuid = FGuid::NewGuid());
+
+private:
+	virtual void AddProperties(FPropertyHelper& PropertyHelper) const override;
+};
diff --git a/Engine/Plugins/Experimental/GeometryCollectionPlugin/Source/GeometryCollectionNodes/Private/Dataflow/GeometryCollectionNodes.cpp b/Engine/Plugins/Experimental/GeometryCollectionPlugin/Source/GeometryCollectionNodes/Private/Dataflow/GeometryCollectionNodes.cpp
--- a/Engine/Plugins/Experimental/GeometryCollectionPlugin/Source/GeometryCollectionNodes/Private/Dataflow/GeometryCollectionNodes.cpp
+++ b/Engine/Plugins/Experimental/GeometryCollectionPlugin/Source/GeometryCollectionNodes/Private/Dataflow/GeometryCollectionNodes.cpp
@@ -70,6 +70,7 @@ namespace UE::Dataflow
 		DATAFLOW_NODE_REGISTER_CREATION_FACTORY(FTransformMeshDataflowNode);
 		DATAFLOW_NODE_REGISTER_CREATION_FACTORY(FCompareIntDataflowNode);
 		DATAFLOW_NODE_REGISTER_CREATION_FACTORY(FCompareFloatDataflowNode);
+		DATAFLOW_NODE_REGISTER_CREATION_FACTORY(FBooleanOperationDataflowNode);
 		DATAFLOW_NODE_REGISTER_CREATION_FACTORY(FBranchMeshDataflowNode);
 		DATAFLOW_NODE_REGISTER_CREATION_FACTORY(FBranchCollectionDataflowNode);
 		DATAFLOW_NODE_REGISTER_CREATION_FACTORY(FGetSchemaDataflowNode);
@@ -643,6 +644,8 @@ namespace
 			return A > B;
 		case ECompareOperationEnum::Dataflow_Compare_GreaterOrEqual:
 			return A >= B;
+		case ECompareOperationEnum::Dataflow_Compare_NotEqual:
+			return A != B;
 		default:
 			ensureMsgf(false, TEXT("Invalid ECompareOperationEnum value: %u"), (uint8)Operation);
 		}
@@ -678,6 +681,31 @@ void FCompareFloatDataflowNode::Evaluate(UE::Dataflow::FContext& Context, const
 }
 
 
+void FBooleanOperationDataflowNode::Evaluate(UE::Dataflow::FContext& Context, const FDataflowOutput* Out) const
+{
+	if (Out->IsA(&bResult))
+	{
+		bool bResultValue = false;
+		switch (Operation)
+		{
+		case EBooleanOperationEnum::Dataflow_And:
+			bResultValue = GetValue(Context, &bBoolA) && GetValue(Context, &bBoolB);
+			break;
+		case EBooleanOperationEnum::Dataflow_Or:
+			bResultValue = GetValue(Context, &bBoolA) || GetValue(Context, &bBoolB);
+			break;
+		case EBooleanOperationEnum::Dataflow_Not:
+			bResultValue = !GetValue(Context, &bBoolA);
+			break;
+		default:
+			ensureMsgf(false, TEXT("Invalid EBooleanOperationEnum value: %u"), (uint8)Operation);
+		}
+
+		SetValue(Context, bResultValue, &bResult);
+	}
+}
+
+
 void FBranchMeshDataflowNode::Evaluate(UE::Dataflow::FContext& Context, const FDataflowOutput* Out) const
 {
 	if (Out->IsA<TObjectPtr<UDynamicMesh>>(&Mesh))
diff --git a/Engine/Plugins/Experimental/GeometryCollectionPlugin/Source/GeometryCollectionNodes/Public/Dataflow/GeometryCollectionNodes.h b/Engine/Plugins/Experimental/GeometryCollectionPlugin/Source/GeometryCollectionNodes/Public/Dataflow/GeometryCollectionNodes.h
--- a/Engine/Plugins/Experimental/GeometryCollectionPlugin/Source/GeometryCollectionNodes/Public/Dataflow/GeometryCollectionNodes.h
+++ b/Engine/Plugins/Experimental/GeometryCollectionPlugin/Source/GeometryCollectionNodes/Public/Dataflow/GeometryCollectionNodes.h
@@ -854,6 +854,7 @@ enum class ECompareOperationEnum : uint8
 	Dataflow_Compare_SmallerOrEqual UMETA(DisplayName = "<="),
 	Dataflow_Compare_Greater UMETA(DisplayName = ">"),
 	Dataflow_Compare_GreaterOrEqual UMETA(DisplayName = ">="),
+	Dataflow_Compare_NotEqual UMETA(DisplayName = "!="),
 	//~~~
 	//256th entry
 	Dataflow_Max                UMETA(Hidden)
@@ -909,7 +910,7 @@ USTRUCT()
 struct FCompareFloatDataflowNode : public FDataflowNode
 {
 	GENERATED_USTRUCT_BODY()
-		DATAFLOW_NODE_DEFINE_INTERNAL(FCompareFloatDataflowNode, "CompareFloat", "Math|Compare", "")
+	DATAFLOW_NODE_DEFINE_INTERNAL(FCompareFloatDataflowNode, "CompareFloat", "Math|Compare", "")
 
 public:
 	/** Comparison operation */
@@ -941,7 +942,56 @@ public:
 };
 
 
+UENUM(BlueprintType)
+enum class EBooleanOperationEnum : uint8
+{
+	Dataflow_And UMETA(DisplayName = "&&"),
+	Dataflow_Or UMETA(DisplayName = "||"),
+	Dataflow_Not UMETA(DisplayName = "!"),
+	//~~~
+	//256th entry
+	Dataflow_Max                UMETA(Hidden)
+};
+
+/**
+ *
+ * Boolean operations
+ *
+ */
+USTRUCT()
+struct FBooleanOperationDataflowNode : public FDataflowNode
+{
+	GENERATED_USTRUCT_BODY()
+	DATAFLOW_NODE_DEFINE_INTERNAL(FBooleanOperationDataflowNode, "BooleanOperation", "Math|Boolean", "")
 
+public:
+	/** Boolean operation */
+	UPROPERTY(EditAnywhere, Category = "Boolean");
+	EBooleanOperationEnum Operation = EBooleanOperationEnum::Dataflow_And;
+
+	/** Boolean input */
+	UPROPERTY(EditAnywhere, Category = "Boolean", Meta = (DataflowInput));
+	bool bBoolA = false;
+
+	/** Boolean input */
+	UPROPERTY(EditAnywhere, Category = "Boolean", Meta = (EditCondition = "Operation != EBooleanOperationEnum::Dataflow_Not", DataflowInput));
+	bool bBoolB = false;
+
+	/** Boolean result of the operator */
+	UPROPERTY(meta = (DataflowOutput));
+	bool bResult = false;
+
+	FBooleanOperationDataflowNode(const UE::Dataflow::FNodeParameters& InParam, FGuid InGuid = FGuid::NewGuid())
+		: FDataflowNode(InParam, InGuid)
+	{
+		RegisterInputConnection(&bBoolA);
+		RegisterInputConnection(&bBoolB);
+		RegisterOutputConnection(&bResult);
+	}
+
+	virtual void Evaluate(UE::Dataflow::FContext& Context, const FDataflowOutput* Out) const override;
+
+};
 
 /**
  *
