Dataflow : fix geometry collection terminal assert generation causing an assert if the input collection is empty

diff --git a/Engine/Plugins/Experimental/GeometryCollectionPlugin/Source/GeometryCollectionNodes/Private/Dataflow/GeometryCollectionAssetNodes.cpp b/Engine/Plugins/Experimental/GeometryCollectionPlugin/Source/GeometryCollectionNodes/Private/Dataflow/GeometryCollectionAssetNodes.cpp
--- a/Engine/Plugins/Experimental/GeometryCollectionPlugin/Source/GeometryCollectionNodes/Private/Dataflow/GeometryCollectionAssetNodes.cpp
+++ b/Engine/Plugins/Experimental/GeometryCollectionPlugin/Source/GeometryCollectionNodes/Private/Dataflow/GeometryCollectionAssetNodes.cpp
@@ -10,6 +10,7 @@
 #include "GeometryCollection/GeometryCollectionClusteringUtility.h"
 #include "GeometryCollection/Facades/CollectionHierarchyFacade.h"
 #include "GeometryCollection/Facades/CollectionInstancedMeshFacade.h"
+#include "GeometryCollection/Facades/CollectionTransformFacade.h"
 #include "Components/StaticMeshComponent.h"
 #include "Materials/Material.h"
 #include "Materials/MaterialInterface.h"
@@ -121,6 +122,13 @@ void FGeometryCollectionTerminalDataflowNode_v2::SetAssetValue(TObjectPtr<UObjec
 
 	if (UGeometryCollection* CollectionAsset = Cast<UGeometryCollection>(Asset.Get()))
 	{
+		// make sure to always reset the root proxies even if the collection is empty 
+		if (IsConnected(&RootProxyMeshes))
+		{
+			CollectionAsset->RootProxyData.ProxyMeshes.Reset();
+			CollectionAsset->RootProxyData.MeshTransforms.Reset();
+		}
+
 		if (FGeometryCollectionPtr GeometryCollection = CollectionAsset->GetGeometryCollection())
 		{
 			const FManagedArrayCollection& InCollection = GetValue(Context, &Collection);
@@ -135,15 +143,15 @@ void FGeometryCollectionTerminalDataflowNode_v2::SetAssetValue(TObjectPtr<UObjec
 
 			if (IsConnected(&RootProxyMeshes))
 			{
+				const GeometryCollection::Facades::FCollectionTransformFacade TransformFacade(*GeometryCollection);
+				const FTransform RootTransform = TransformFacade.GetRootTransform();
+
 				FGeometryCollectionProxyMeshData& Data = CollectionAsset->RootProxyData;
-				Data.ProxyMeshes.Reset();
-				Data.ProxyMeshes.Reserve(InRootProxyMeshes.Num());
-				Data.MeshTransforms.Reset();
-				Data.MeshTransforms.Reserve(InRootProxyMeshes.Num());
 				for (const FDataflowRootProxyMesh& ProxyMesh : InRootProxyMeshes)
 				{
 					Data.ProxyMeshes.Add(ProxyMesh.Mesh);
-					Data.MeshTransforms.Add(FTransform3f(ProxyMesh.Transform));
+					// root proxy transforms are relative to the root transform but FDataflowRootProxyMesh is in component space 
+					Data.MeshTransforms.Add(FTransform3f(ProxyMesh.Transform.GetRelativeTransform(RootTransform)));
 				}
 			}
 
@@ -193,7 +201,15 @@ void FGeometryCollectionTerminalDataflowNode::SetAssetValue(TObjectPtr<UObject>
 	{
 		if (FGeometryCollectionPtr GeometryCollection = CollectionAsset->GetGeometryCollection())
 		{
-			const FManagedArrayCollection& InCollection = GetValue(Context, &Collection);
+			// need to make a copy since we may need to add attributes
+			FManagedArrayCollection InCollection = GetValue(Context, &Collection);
+			
+			// for now make sure we have the right interfaces setup on the source collection so that it does not trigger an ensure during the call to ResetFrom 
+			// (because of the discrepancy between convex properties attributes as they are created by default in FGeometryCollection but may be missing from the input collection )
+			// todo(dataflow) we should make this more automatic in the future
+			FGeometryCollectionConvexPropertiesInterface ConvexPropertiesInterface(&InCollection);
+			ConvexPropertiesInterface.InitializeInterface();
+
 			const FMaterialArray& InMaterials = GetValue(Context, &Materials);
 			const FMaterialInstanceArray& InMaterialInstances = GetValue(Context, &MaterialInstances);
 			const FInstancedMeshesArray& InInstancedMeshes = GetValue(Context, &InstancedMeshes);
@@ -523,11 +539,18 @@ void FGeometryCollectionToCollectionDataflowNode_v2::Evaluate(UE::Dataflow::FCon
 	{
 		FGeometryCollectionEngineConversion::ConvertGeometryCollectionToGeometryCollection(GeometryCollection, OutCollection, OutMaterials, OutInstancedMeshes);
 
+		FTransform RootTransform = FTransform::Identity;
+		if (GeometryCollection->GetGeometryCollection())
+		{
+			const GeometryCollection::Facades::FCollectionTransformFacade TransformFacade(*GeometryCollection->GetGeometryCollection());
+			RootTransform = TransformFacade.GetRootTransform();
+		}
+
 		for (int32 MeshIndex = 0; MeshIndex < GeometryCollection->RootProxyData.ProxyMeshes.Num(); MeshIndex++)
 		{
 			FDataflowRootProxyMesh RootProxyMesh;
 			RootProxyMesh.Mesh = GeometryCollection->RootProxyData.ProxyMeshes[MeshIndex];
-			RootProxyMesh.Transform = FTransform(GeometryCollection->RootProxyData.GetMeshTransform(MeshIndex));
+			RootProxyMesh.Transform = FTransform(GeometryCollection->RootProxyData.GetMeshTransform(MeshIndex)) * RootTransform;
 			OutRootProxyMeshes.Emplace(RootProxyMesh);
 		}
 	}
