Modify TransferVertexAttribute node default method, add error message

diff --git a/Engine/Plugins/Experimental/GeometryCollectionPlugin/Source/GeometryCollectionNodes/Private/Dataflow/GeometryCollectionTransferVertexAttributeNode.cpp b/Engine/Plugins/Experimental/GeometryCollectionPlugin/Source/GeometryCollectionNodes/Private/Dataflow/GeometryCollectionTransferVertexAttributeNode.cpp
--- a/Engine/Plugins/Experimental/GeometryCollectionPlugin/Source/GeometryCollectionNodes/Private/Dataflow/GeometryCollectionTransferVertexAttributeNode.cpp
+++ b/Engine/Plugins/Experimental/GeometryCollectionPlugin/Source/GeometryCollectionNodes/Private/Dataflow/GeometryCollectionTransferVertexAttributeNode.cpp
@@ -581,17 +581,10 @@ void FGeometryCollectionTransferVertexAttributeNode::Evaluate(UE::Dataflow::FCon
 			TransformNameSuffix
 		};
 
-		if (Target.IsValid() && Source.IsValid())
+		if (TransferMethod != EDataflowTransferVertexAttributeNodeTransferMethod::None
+			&& Target.IsValid() && Source.IsValid())
 		{
 			TArray<FTransferData> TransferDataArray = ComputeTransferData(SourceCollection, TargetCollection, TransferProperties, TransferMethod);
-			if (TransferMethod == EDataflowTransferVertexAttributeNodeTransferMethod::Global)
-			{
-				Context.Warning(TEXT("Transfer method: global transfer"), this, Out);
-			}
-			else if (TransferMethod == EDataflowTransferVertexAttributeNodeTransferMethod::None)
-			{
-				Context.Warning(TEXT("Transfer failed"), this, Out);
-			}
 			if (const TManagedArray<float>* SourceAttributeFloatArray = Source.GetAttributeArray<float>(Key.Attribute, Key.Group))
 			{
 				TManagedArray<float>* TargetAttributeArray = Target.GetAttributeArray<float>(Key.Attribute, Key.Group);
@@ -628,6 +621,7 @@ void FGeometryCollectionTransferVertexAttributeNode::Evaluate(UE::Dataflow::FCon
 			}
 			else
 			{
+				Context.Error(TEXT("Source attribute is not found"), this, Out);
 				UE_LOG(LogChaosDataflow, Warning, TEXT("TransferVertexAttribute: attribute[%s] of group[%s] has type that is not float or FLinearColor, transfer failed."), *Key.Attribute, *Key.Group);
 			}
 		}
@@ -659,7 +653,8 @@ void FGeometryCollectionTransferVertexSkinWeightsNode::Evaluate(UE::Dataflow::FC
 			TransformNameSuffix
 		};
 
-		if (Target.IsValid() && Source.IsValid())
+		if (TransferMethod != EDataflowTransferVertexAttributeNodeTransferMethod::None
+			&& Target.IsValid() && Source.IsValid())
 		{
 			TArray<FTransferData> TransferDataArray = ComputeTransferData(SourceCollection, TargetCollection, TransferProperties, TransferMethod);
 			GeometryCollection::Facades::FVertexBoneWeightsFacade SourceVertexBoneWeightsFacade(SourceCollection);
@@ -718,6 +713,10 @@ void FGeometryCollectionTransferVertexSkinWeightsNode::Evaluate(UE::Dataflow::FC
 					}
 				}
 			}
+			else
+			{
+				Context.Error(TEXT("Source skin weight is not found"), this, Out);
+			}
 		}
 		SetValue(Context, MoveTemp(TargetCollection), &Collection);
 	}
diff --git a/Engine/Plugins/Experimental/GeometryCollectionPlugin/Source/GeometryCollectionNodes/Public/Dataflow/GeometryCollectionTransferVertexAttributeNode.h b/Engine/Plugins/Experimental/GeometryCollectionPlugin/Source/GeometryCollectionNodes/Public/Dataflow/GeometryCollectionTransferVertexAttributeNode.h
--- a/Engine/Plugins/Experimental/GeometryCollectionPlugin/Source/GeometryCollectionNodes/Public/Dataflow/GeometryCollectionTransferVertexAttributeNode.h
+++ b/Engine/Plugins/Experimental/GeometryCollectionPlugin/Source/GeometryCollectionNodes/Public/Dataflow/GeometryCollectionTransferVertexAttributeNode.h
@@ -74,7 +74,7 @@ enum class EDataflowTransferVertexAttributeNodeTransferMethod : uint8
 	/** Transfers vertex attribute globally. Warning: attributes on one geometry might bleed over to neighbor geometries.*/
 	Global UMETA(DisplayName = "Global Transfer"),
 
-	/** No distance falloff */
+	/** No transfer */
 	None UMETA(DisplayName = "None"),
 	//~~~
 	//256th entry
@@ -113,7 +113,7 @@ public:
 	
 	/* Transfer method [default: Paired Geometry Transfer] */
 	UPROPERTY(EditAnywhere, Category = "Method")
-	EDataflowTransferVertexAttributeNodeTransferMethod TransferMethod = EDataflowTransferVertexAttributeNodeTransferMethod::Component;
+	EDataflowTransferVertexAttributeNodeTransferMethod TransferMethod = EDataflowTransferVertexAttributeNodeTransferMethod::Global;
 
 	/* Bounding volume type for source assets[default: Triangle] */
 	UPROPERTY(EditAnywhere, Category = "Thresholds")
@@ -193,7 +193,7 @@ public:
 
 	/* Transfer method [default: Paired Geometry Transfer] */
 	UPROPERTY(EditAnywhere, Category = "Method")
-	EDataflowTransferVertexAttributeNodeTransferMethod TransferMethod = EDataflowTransferVertexAttributeNodeTransferMethod::Component;
+	EDataflowTransferVertexAttributeNodeTransferMethod TransferMethod = EDataflowTransferVertexAttributeNodeTransferMethod::Global;
 
 	/* Bounding volume type for source assets[default: Triangle] */
 	UPROPERTY(EditAnywhere, Category = "Thresholds")
