PR #13086: Support "Only Same Parent" option for TinyGeo "Merge Geometry" mode

diff --git a/Engine/Plugins/Experimental/ChaosEditor/Source/FractureEditor/Private/FractureToolFixTinyGeo.cpp b/Engine/Plugins/Experimental/ChaosEditor/Source/FractureEditor/Private/FractureToolFixTinyGeo.cpp
--- a/Engine/Plugins/Experimental/ChaosEditor/Source/FractureEditor/Private/FractureToolFixTinyGeo.cpp
+++ b/Engine/Plugins/Experimental/ChaosEditor/Source/FractureEditor/Private/FractureToolFixTinyGeo.cpp
@@ -306,7 +306,9 @@ void UFractureToolFixTinyGeo::Execute(TWeakPtr<FFractureEditorModeToolkit> InToo
 					MinVolume,
 					SmallIndices,
 					false, /*bUnionJoinedPieces*/ // Note: Union-ing the pieces is nicer in theory, but can leave cracks and non-manifold garbage
-					SelectionMethod
+					SelectionMethod,
+					true /*bUseCollectionProximity*/,
+					TinyGeoSettings->bOnlySameParent
 				);
 			}
 			else
diff --git a/Engine/Plugins/Experimental/ChaosEditor/Source/FractureEditor/Private/FractureToolFixTinyGeo.h b/Engine/Plugins/Experimental/ChaosEditor/Source/FractureEditor/Private/FractureToolFixTinyGeo.h
--- a/Engine/Plugins/Experimental/ChaosEditor/Source/FractureEditor/Private/FractureToolFixTinyGeo.h
+++ b/Engine/Plugins/Experimental/ChaosEditor/Source/FractureEditor/Private/FractureToolFixTinyGeo.h
@@ -67,7 +67,7 @@ public:
 	bool bOnlyClusters = false;
 
 	// Only merge to neighbors with the same parent in the hierarchy
-	UPROPERTY(EditAnywhere, Category = MergeSettings, meta = (EditCondition = "MergeType == EMergeType::MergeClusters", EditConditionHides))
+	UPROPERTY(EditAnywhere, Category = MergeSettings)
 	bool bOnlySameParent = true;
 
 	// Helper variable to let the EditConditions above check whether the Fracture Level is set to 'All'
diff --git a/Engine/Plugins/Experimental/Fracture/Source/FractureEngine/Private/FractureEngineUtility.cpp b/Engine/Plugins/Experimental/Fracture/Source/FractureEngine/Private/FractureEngineUtility.cpp
--- a/Engine/Plugins/Experimental/Fracture/Source/FractureEngine/Private/FractureEngineUtility.cpp
+++ b/Engine/Plugins/Experimental/Fracture/Source/FractureEngine/Private/FractureEngineUtility.cpp
@@ -279,7 +279,8 @@ void FFractureEngineUtility::FixTinyGeo(FManagedArrayCollection& InOutCollection
 					SmallIndices,
 					false, /*bUnionJoinedPieces*/ // Note: Union-ing the pieces is nicer in theory, but can leave cracks and non-manifold garbage
 					SelectionMethod,
-					bUseCollectionProximity);
+					bUseCollectionProximity,
+					InOnlySameParent);
 			}
 			else
 			{
diff --git a/Engine/Plugins/Experimental/GeometryCollectionPlugin/Source/GeometryCollectionNodes/Private/Dataflow/GeometryCollectionUtilityNodes.cpp b/Engine/Plugins/Experimental/GeometryCollectionPlugin/Source/GeometryCollectionNodes/Private/Dataflow/GeometryCollectionUtilityNodes.cpp
--- a/Engine/Plugins/Experimental/GeometryCollectionPlugin/Source/GeometryCollectionNodes/Private/Dataflow/GeometryCollectionUtilityNodes.cpp
+++ b/Engine/Plugins/Experimental/GeometryCollectionPlugin/Source/GeometryCollectionNodes/Private/Dataflow/GeometryCollectionUtilityNodes.cpp
@@ -1138,7 +1138,7 @@ void FFixTinyGeoDataflowNode::Evaluate(UE::Dataflow::FContext& Context, const FD
 				bOnlyClusters,
 				NeighborSelection,
 				bOnlyToConnected,
-				bOnlySameParent,
+				MergeType == EFixTinyGeoMergeType::MergeClusters ? bOnlySameParent : bGeometryOnlySameParent,
 				bUseCollectionProximityForConnections);
 
 			SetValue(Context, MoveTemp(InCollection), &Collection);
diff --git a/Engine/Plugins/Experimental/GeometryCollectionPlugin/Source/GeometryCollectionNodes/Public/Dataflow/GeometryCollectionUtilityNodes.h b/Engine/Plugins/Experimental/GeometryCollectionPlugin/Source/GeometryCollectionNodes/Public/Dataflow/GeometryCollectionUtilityNodes.h
--- a/Engine/Plugins/Experimental/GeometryCollectionPlugin/Source/GeometryCollectionNodes/Public/Dataflow/GeometryCollectionUtilityNodes.h
+++ b/Engine/Plugins/Experimental/GeometryCollectionPlugin/Source/GeometryCollectionNodes/Public/Dataflow/GeometryCollectionUtilityNodes.h
@@ -865,13 +865,17 @@ public:
 	UPROPERTY(EditAnywhere, Category = MergeSettings, meta = (EditCondition = "!bFractureLevelIsAll && MergeType == EFixTinyGeoMergeType::MergeClusters", EditConditionHides))
 	bool bOnFractureLevel = true;
 
-	/** Only auto - consider clusters for merging.Note that leaf nodes can still be consider if manually selected. */
+	/** Only auto-consider clusters for merging. Note that leaf nodes can still be consider if manually selected. */
 	UPROPERTY(EditAnywhere, Category = MergeSettings, meta = (EditCondition = "!bFractureLevelIsAll && MergeType == EFixTinyGeoMergeType::MergeClusters && bOnFractureLevel", EditConditionHides))
 	bool bOnlyClusters = false;
 
-	/** Only merge to neighbors with the same parent in the hierarchy */
-	UPROPERTY(EditAnywhere, Category = MergeSettings, meta = (EditCondition = "MergeType == EFixTinyGeoMergeType::MergeClusters", EditConditionHides))
+	/** Only merge clusters to neighbors with the same parent in the hierarchy */
+	UPROPERTY(EditAnywhere, Category = MergeSettings, meta = (DisplayName = "Only Merge Clusters w/ Same Parent", EditCondition = "MergeType == EFixTinyGeoMergeType::MergeClusters", EditConditionHides))
 	bool bOnlySameParent = true;
+	
+	/** Only merge geometry to neighbors with the same parent in the hierarchy */
+	UPROPERTY(EditAnywhere, Category = MergeSettings, meta = (DisplayName = "Only Merge Geometry w/ Same Parent", EditCondition = "MergeType == EFixTinyGeoMergeType::MergeGeometry", EditConditionHides))
+	bool bGeometryOnlySameParent = false;
 
 	/** Helper variable to let the EditConditions above check whether the Fracture Level is set to 'All' */
 	UPROPERTY()
diff --git a/Engine/Plugins/Experimental/PlanarCutPlugin/Source/PlanarCut/Private/PlanarCut.cpp b/Engine/Plugins/Experimental/PlanarCutPlugin/Source/PlanarCut/Private/PlanarCut.cpp
--- a/Engine/Plugins/Experimental/PlanarCutPlugin/Source/PlanarCut/Private/PlanarCut.cpp
+++ b/Engine/Plugins/Experimental/PlanarCutPlugin/Source/PlanarCut/Private/PlanarCut.cpp
@@ -1319,7 +1319,8 @@ int32 MergeBones(
 	const TArrayView<const int32>& SmallTransformIndices,
 	bool bUnionJoinedPieces,
 	UE::PlanarCut::ENeighborSelectionMethod NeighborSelectionMethod,
-	bool bUseCollectionProximity
+	bool bUseCollectionProximity,
+	bool bOnlySameParent
 )
 {
 	FTransform CellsToWorld = FTransform::Identity;
@@ -1443,6 +1444,7 @@ int32 MergeBones(
 		}
 	};
 	TMap<int32, double> GeomToVol;
+	TMap<int32, int32> GeomToParent;
 	TMap<int32, int32> GeomIdxToRemoveGroupIdx;
 	TArray<FRemoveGroup> RemoveGroups;
 	TSet<int32> TooSmalls;
@@ -1473,6 +1475,7 @@ int32 MergeBones(
 		{
 			CanMerge.Add(GeomIdx);
 			GeomToVol.Add(GeomIdx, Volumes[Idx]);
+			GeomToParent.Add(GeomIdx, Collection.Parent[Idx]);
 		}
 	}
 	for (int32 TransformIdx : SmallTransformIndices)
@@ -1506,6 +1509,10 @@ int32 MergeBones(
 		{
 			if (NbrIdx != SmallIdx && CanMerge.Contains(NbrIdx))
 			{
+				if (bOnlySameParent && GeomToParent[NbrIdx] != GeomToParent[SmallIdx])
+				{
+					continue;
+				}
 				double Score;
 				if (NeighborSelectionMethod == UE::PlanarCut::ENeighborSelectionMethod::LargestNeighbor)
 				{
diff --git a/Engine/Plugins/Experimental/PlanarCutPlugin/Source/PlanarCut/Public/PlanarCut.h b/Engine/Plugins/Experimental/PlanarCutPlugin/Source/PlanarCut/Public/PlanarCut.h
--- a/Engine/Plugins/Experimental/PlanarCutPlugin/Source/PlanarCut/Public/PlanarCut.h
+++ b/Engine/Plugins/Experimental/PlanarCutPlugin/Source/PlanarCut/Public/PlanarCut.h
@@ -403,6 +403,7 @@ namespace UE
  * @param bUnionJoinedPieces		Try to 'union' the merged pieces, removing internal triangles and connecting the shared cut boundary
  * @param NeighborSelectionMethod	How to choose which neighbor to merge to
  * @param bUseCollectionProximity	If true, will use the proximity data on the collection to decide which pieces are allowed to merge. Otherwise, a reasonable default proximity will be computed and used.
+ * @param bOnlySameParent			Whether to only merge nodes that have the same parent in the hierarchy
  */
 int32 PLANARCUT_API MergeBones(
 	FGeometryCollection& Collection,
@@ -412,7 +413,8 @@ int32 PLANARCUT_API MergeBones(
 	const TArrayView<const int32>& SmallTransformIndices,
 	bool bUnionJoinedPieces,
 	UE::PlanarCut::ENeighborSelectionMethod NeighborSelectionMethod,
-	bool bUseCollectionProximity = true
+	bool bUseCollectionProximity = true,
+	bool bOnlySameParent = false
 );
 
 /**
