Dataflow : rename output of Get Centroid node to Points

diff --git a/Engine/Plugins/Experimental/GeometryCollectionPlugin/Source/GeometryCollectionNodes/Public/Dataflow/GeometryCollectionNodes.h b/Engine/Plugins/Experimental/GeometryCollectionPlugin/Source/GeometryCollectionNodes/Public/Dataflow/GeometryCollectionNodes.h
--- a/Engine/Plugins/Experimental/GeometryCollectionPlugin/Source/GeometryCollectionNodes/Public/Dataflow/GeometryCollectionNodes.h
+++ b/Engine/Plugins/Experimental/GeometryCollectionPlugin/Source/GeometryCollectionNodes/Public/Dataflow/GeometryCollectionNodes.h
@@ -638,11 +638,11 @@ public:
 	FManagedArrayCollection Collection;
 
 	/** The centroids will be output for the bones selected in the TransformSelection  */
-	UPROPERTY(meta = (DataflowInput, DisplayName = "TransformSelection"))
+	UPROPERTY(meta = (DataflowInput))
 	FDataflowTransformSelection TransformSelection;
 
 	/** Output centroids */
-	UPROPERTY(meta = (DataflowOutput))
+	UPROPERTY(meta = (DataflowOutput, DisplayName = "Points"))
 	TArray<FVector> Centroids = TArray<FVector>();
 
 	FGetCentroidsFromCollectionDataflowNode(const UE::Dataflow::FNodeParameters& InParam, FGuid InGuid = FGuid::NewGuid())
