Dataflow add nodes to get bounding boxes from Static and Dynamic Meshes
- Also imporve the current collection bounding box node
- Add keywords for easier discovery
- nodes are returning the dimensions and Center of the bounding box as well

diff --git a/Engine/Plugins/Experimental/Dataflow/Source/DataflowNodes/Private/Dataflow/DataflowStaticMeshNodes.cpp b/Engine/Plugins/Experimental/Dataflow/Source/DataflowNodes/Private/Dataflow/DataflowStaticMeshNodes.cpp
--- a/Engine/Plugins/Experimental/Dataflow/Source/DataflowNodes/Private/Dataflow/DataflowStaticMeshNodes.cpp
+++ b/Engine/Plugins/Experimental/Dataflow/Source/DataflowNodes/Private/Dataflow/DataflowStaticMeshNodes.cpp
@@ -10,10 +10,20 @@ namespace UE::Dataflow
 	void RegisterStaticMeshNodes()
 	{
 		DATAFLOW_NODE_REGISTER_CREATION_FACTORY(FGetStaticMeshDataflowNode);
+		DATAFLOW_NODE_REGISTER_CREATION_FACTORY(FGetStaticMeshBoundingBoxDataflowNode);
+
 		DATAFLOW_NODE_REGISTER_GETTER_FOR_ASSET(UStaticMesh, FGetStaticMeshDataflowNode);
 	}
 }
 
+//////////////////////////////////////////////////////////////////////////////////////////////
+
+FGetStaticMeshDataflowNode::FGetStaticMeshDataflowNode(const UE::Dataflow::FNodeParameters& InParam, FGuid InGuid)
+	: FDataflowNode(InParam, InGuid)
+{
+	RegisterOutputConnection(&StaticMesh);
+}
+
 void FGetStaticMeshDataflowNode::Evaluate(UE::Dataflow::FContext& Context, const FDataflowOutput* Out) const
 {
 	typedef TObjectPtr<const UStaticMesh> DataType;
@@ -50,3 +60,32 @@ void FGetStaticMeshDataflowNode::SetAssetProperty(UObject* Asset)
 	}
 }
 
+//////////////////////////////////////////////////////////////////////////////////////////////
+
+FGetStaticMeshBoundingBoxDataflowNode::FGetStaticMeshBoundingBoxDataflowNode(const UE::Dataflow::FNodeParameters& InParam, FGuid InGuid)
+	: FDataflowNode(InParam, InGuid)
+{
+	RegisterInputConnection(&StaticMesh);
+
+	RegisterOutputConnection(&BoundingBox);
+	RegisterOutputConnection(&Center);
+	RegisterOutputConnection(&Dimensions);
+}
+
+void FGetStaticMeshBoundingBoxDataflowNode::Evaluate(UE::Dataflow::FContext& Context, const FDataflowOutput* Out) const
+{
+	if (Out->IsA(&BoundingBox) || Out->IsA(&Center) || Out->IsA(&Dimensions))
+	{
+		FBox OutBoundingBox(ForceInit);
+
+		if (const TObjectPtr<const UStaticMesh> InStaticMesh = GetValue(Context, &StaticMesh))
+		{
+			OutBoundingBox = InStaticMesh->GetBounds().GetBox();
+		}
+
+		SetValue(Context, OutBoundingBox, &BoundingBox);
+		SetValue(Context, OutBoundingBox.GetCenter(), &Center);
+		SetValue(Context, OutBoundingBox.GetSize(), &Dimensions);
+	}
+}
+
diff --git a/Engine/Plugins/Experimental/Dataflow/Source/DataflowNodes/Public/Dataflow/DataflowStaticMeshNodes.h b/Engine/Plugins/Experimental/Dataflow/Source/DataflowNodes/Public/Dataflow/DataflowStaticMeshNodes.h
--- a/Engine/Plugins/Experimental/Dataflow/Source/DataflowNodes/Public/Dataflow/DataflowStaticMeshNodes.h
+++ b/Engine/Plugins/Experimental/Dataflow/Source/DataflowNodes/Public/Dataflow/DataflowStaticMeshNodes.h
@@ -28,11 +28,7 @@ public:
 	UPROPERTY(EditAnywhere, Category = "Overrides")
 	FName PropertyName = "StaticMesh";
 
-	FGetStaticMeshDataflowNode(const UE::Dataflow::FNodeParameters& InParam, FGuid InGuid = FGuid::NewGuid())
-		: FDataflowNode(InParam, InGuid)
-	{
-		RegisterOutputConnection(&StaticMesh);
-	}
+	FGetStaticMeshDataflowNode(const UE::Dataflow::FNodeParameters& InParam, FGuid InGuid = FGuid::NewGuid());
 
 	UE_API virtual void Evaluate(UE::Dataflow::FContext& Context, const FDataflowOutput* Out) const override;
 
@@ -40,6 +36,40 @@ public:
 	UE_API virtual void SetAssetProperty(UObject* Asset) override;
 };
 
+/**
+ * Get the local geometric bounding box a static mesh
+ */
+USTRUCT(meta = (DataflowGeometryCollection))
+struct FGetStaticMeshBoundingBoxDataflowNode : public FDataflowNode
+{
+	GENERATED_USTRUCT_BODY()
+	DATAFLOW_NODE_DEFINE_INTERNAL(FGetStaticMeshBoundingBoxDataflowNode, "GetStaticMeshBoundingBox", "Static Mesh", "Bounds Size Dimensions Extents Center")
+	DATAFLOW_NODE_RENDER_TYPE("SurfaceRender", FName("FBox"), "BoundingBox")
+
+private:
+	/** Static Mesh to compute the bouning box from */
+	UPROPERTY(EditAnywhere, Category = "Dataflow", meta = (DataflowInput))
+	TObjectPtr<const UStaticMesh> StaticMesh = nullptr;
+
+	/** Geometric bounding box of the input Static Mesh */
+	UPROPERTY(meta = (DataflowOutput))
+	FBox BoundingBox = FBox(ForceInit);
+
+	/** Center of the resulting bounding box */
+	UPROPERTY(meta = (DataflowOutput))
+	FVector Center = FVector::ZeroVector;
+
+	/** Dimensions of the resulting bounding box in centimers */
+	UPROPERTY(meta = (DataflowOutput))
+	FVector Dimensions = FVector::ZeroVector;
+
+	virtual void Evaluate(UE::Dataflow::FContext& Context, const FDataflowOutput* Out) const override;
+
+public:
+	FGetStaticMeshBoundingBoxDataflowNode(const UE::Dataflow::FNodeParameters& InParam, FGuid InGuid = FGuid::NewGuid());
+
+};
+
 
 namespace UE::Dataflow
 {
diff --git a/Engine/Plugins/Experimental/GeometryCollectionPlugin/Source/GeometryCollectionNodes/Private/Dataflow/GeometryCollectionMeshNodes.cpp b/Engine/Plugins/Experimental/GeometryCollectionPlugin/Source/GeometryCollectionNodes/Private/Dataflow/GeometryCollectionMeshNodes.cpp
--- a/Engine/Plugins/Experimental/GeometryCollectionPlugin/Source/GeometryCollectionNodes/Private/Dataflow/GeometryCollectionMeshNodes.cpp
+++ b/Engine/Plugins/Experimental/GeometryCollectionPlugin/Source/GeometryCollectionNodes/Private/Dataflow/GeometryCollectionMeshNodes.cpp
@@ -50,6 +50,7 @@ namespace UE::Dataflow
 		DATAFLOW_NODE_REGISTER_CREATION_FACTORY(FMeshToCollectionDataflowNode);
 		DATAFLOW_NODE_REGISTER_CREATION_FACTORY(FCollectionToMeshDataflowNode);
 		DATAFLOW_NODE_REGISTER_CREATION_FACTORY(FStaticMeshToMeshDataflowNode);
+		DATAFLOW_NODE_REGISTER_CREATION_FACTORY(FGetMeshBoundingBoxDataflowNode);
 		DATAFLOW_NODE_REGISTER_CREATION_FACTORY(FMeshAppendDataflowNode);
 		DATAFLOW_NODE_REGISTER_CREATION_FACTORY(FDataflowMeshAppendDataflowNode);
 		DATAFLOW_NODE_REGISTER_CREATION_FACTORY(FMakeDataflowMeshDataflowNode);
@@ -337,7 +338,41 @@ void FStaticMeshToMeshDataflowNode::Evaluate(UE::Dataflow::FContext& Context, co
 #endif
 }
 
+///////////////////////////////////////////////////////////////////////////////////////
 
+FGetMeshBoundingBoxDataflowNode::FGetMeshBoundingBoxDataflowNode(const UE::Dataflow::FNodeParameters& InParam, FGuid InGuid) 
+	: FDataflowNode(InParam, InGuid)
+{
+	RegisterInputConnection(&Mesh);
+
+	RegisterOutputConnection(&BoundingBox);
+	RegisterOutputConnection(&Center);
+	RegisterOutputConnection(&Dimensions);
+}
+
+void FGetMeshBoundingBoxDataflowNode::Evaluate(UE::Dataflow::FContext& Context, const FDataflowOutput* Out) const
+{
+	if (Out->IsA(&BoundingBox) || Out->IsA(&Center) || Out->IsA(&Dimensions))
+	{
+		FBox OutBoundingBox(ForceInit);
+
+		if (const TObjectPtr<const UDynamicMesh> InMesh = GetValue(Context, &Mesh))
+		{
+			if (InMesh->GetMeshPtr())
+			{
+				const UE::Geometry::FAxisAlignedBox3d LocalBounds = InMesh->GetMeshPtr()->GetBounds();
+				OutBoundingBox.Min = LocalBounds.Min;
+				OutBoundingBox.Max = LocalBounds.Max;
+			}
+		}
+
+		SetValue(Context, OutBoundingBox, &BoundingBox);
+		SetValue(Context, OutBoundingBox.GetCenter(), &Center);
+		SetValue(Context, OutBoundingBox.GetSize(), &Dimensions);
+	}
+}
+
+///////////////////////////////////////////////////////////////////////////////////////
 void FMeshAppendDataflowNode::Evaluate(UE::Dataflow::FContext& Context, const FDataflowOutput* Out) const
 {
 	if (Out->IsA<TObjectPtr<UDynamicMesh>>(&Mesh))
diff --git a/Engine/Plugins/Experimental/GeometryCollectionPlugin/Source/GeometryCollectionNodes/Private/Dataflow/GeometryCollectionNodes.cpp b/Engine/Plugins/Experimental/GeometryCollectionPlugin/Source/GeometryCollectionNodes/Private/Dataflow/GeometryCollectionNodes.cpp
--- a/Engine/Plugins/Experimental/GeometryCollectionPlugin/Source/GeometryCollectionNodes/Private/Dataflow/GeometryCollectionNodes.cpp
+++ b/Engine/Plugins/Experimental/GeometryCollectionPlugin/Source/GeometryCollectionNodes/Private/Dataflow/GeometryCollectionNodes.cpp
@@ -214,11 +214,21 @@ void FLogStringDataflowNode::Evaluate(UE::Dataflow::FContext& Context, const FDa
 	}
 }
 
+///////////////////////////////////////////////////////////////////////////////////////////////
 
+FBoundingBoxDataflowNode::FBoundingBoxDataflowNode(const UE::Dataflow::FNodeParameters& InParam, FGuid InGuid)
+	: FDataflowNode(InParam, InGuid)
+{
+	RegisterInputConnection(&Collection);
+
+	RegisterOutputConnection(&BoundingBox);
+	RegisterOutputConnection(&Center);
+	RegisterOutputConnection(&Dimensions);
+}
 
 void FBoundingBoxDataflowNode::Evaluate(UE::Dataflow::FContext& Context, const FDataflowOutput* Out) const
 {
-	if (Out->IsA<FBox>(&BoundingBox))
+	if (Out->IsA(&BoundingBox) || Out->IsA(&Center) || Out->IsA(&Dimensions))
 	{
 		const FManagedArrayCollection& InCollection = GetValue<FManagedArrayCollection>(Context, &Collection);
 
@@ -226,9 +236,13 @@ void FBoundingBoxDataflowNode::Evaluate(UE::Dataflow::FContext& Context, const F
 		const FBox& BoundingBoxInCollectionSpace = BoundsFacade.GetBoundingBoxInCollectionSpace();
 
 		SetValue(Context, BoundingBoxInCollectionSpace, &BoundingBox);
+		SetValue(Context, BoundingBoxInCollectionSpace.GetCenter(), &Center);
+		SetValue(Context, BoundingBoxInCollectionSpace.GetSize(), &Dimensions);
 	}
 }
 
+///////////////////////////////////////////////////////////////////////////////////////////////
+
 void FGetBoxLengthsDataflowNode::Evaluate(UE::Dataflow::FContext& Context, const FDataflowOutput* Out) const
 {
 	if (Out->IsA(&Lengths))
diff --git a/Engine/Plugins/Experimental/GeometryCollectionPlugin/Source/GeometryCollectionNodes/Public/Dataflow/GeometryCollectionMeshNodes.h b/Engine/Plugins/Experimental/GeometryCollectionPlugin/Source/GeometryCollectionNodes/Public/Dataflow/GeometryCollectionMeshNodes.h
--- a/Engine/Plugins/Experimental/GeometryCollectionPlugin/Source/GeometryCollectionNodes/Public/Dataflow/GeometryCollectionMeshNodes.h
+++ b/Engine/Plugins/Experimental/GeometryCollectionPlugin/Source/GeometryCollectionNodes/Public/Dataflow/GeometryCollectionMeshNodes.h
@@ -459,6 +459,39 @@ public:
 
 };
 
+/**
+ * Get the local geometric bounding box a dynamic mesh
+ */
+USTRUCT(meta = (DataflowGeometryCollection))
+struct FGetMeshBoundingBoxDataflowNode : public FDataflowNode
+{
+	GENERATED_USTRUCT_BODY()
+	DATAFLOW_NODE_DEFINE_INTERNAL(FGetMeshBoundingBoxDataflowNode, "GetMeshBoundingBox", "Mesh", "Bounds Size Dimensions Extents Center")
+	DATAFLOW_NODE_RENDER_TYPE("SurfaceRender", FName("FBox"), "BoundingBox")
+
+private:
+	/** dynamic mesh to compute the bouning box from */
+	UPROPERTY(meta = (DataflowInput))
+	TObjectPtr<const UDynamicMesh> Mesh = nullptr;
+
+	/** Geometric bounding box of the input dynamic mesh */
+	UPROPERTY(meta = (DataflowOutput))
+	FBox BoundingBox = FBox(ForceInit);
+
+	/** Center of the resulting bounding box */
+	UPROPERTY(meta = (DataflowOutput))
+	FVector Center = FVector::ZeroVector;
+
+	/** Dimensions of the resulting bounding box in centimers */
+	UPROPERTY(meta = (DataflowOutput))
+	FVector Dimensions = FVector::ZeroVector;
+
+	virtual void Evaluate(UE::Dataflow::FContext& Context, const FDataflowOutput* Out) const override;
+
+public:
+	FGetMeshBoundingBoxDataflowNode(const UE::Dataflow::FNodeParameters& InParam, FGuid InGuid = FGuid::NewGuid());
+
+};
 
 /**
  *
diff --git a/Engine/Plugins/Experimental/GeometryCollectionPlugin/Source/GeometryCollectionNodes/Public/Dataflow/GeometryCollectionNodes.h b/Engine/Plugins/Experimental/GeometryCollectionPlugin/Source/GeometryCollectionNodes/Public/Dataflow/GeometryCollectionNodes.h
--- a/Engine/Plugins/Experimental/GeometryCollectionPlugin/Source/GeometryCollectionNodes/Public/Dataflow/GeometryCollectionNodes.h
+++ b/Engine/Plugins/Experimental/GeometryCollectionPlugin/Source/GeometryCollectionNodes/Public/Dataflow/GeometryCollectionNodes.h
@@ -165,38 +165,36 @@ public:
 
 };
 
-
-
-
-
 /**
- *
- * Description for this node
- *
+ * Get the geometric bounding box a collection in collection space
  */
 USTRUCT(meta = (DataflowGeometryCollection))
 struct FBoundingBoxDataflowNode : public FDataflowNode
 {
 	GENERATED_USTRUCT_BODY()
-	DATAFLOW_NODE_DEFINE_INTERNAL(FBoundingBoxDataflowNode, "BoundingBox", "Utilities|Box", "")
+	DATAFLOW_NODE_DEFINE_INTERNAL(FBoundingBoxDataflowNode, "GetCollectionBoundingBox", "Utilities|Box", "Collection Bounds Size Dimensions Extents Center")
 	DATAFLOW_NODE_RENDER_TYPE("SurfaceRender",FName("FBox"), "BoundingBox")
 
 public:
+	/** Collection to compute the bouning box from */
 	UPROPERTY(meta = (DataflowInput))
 	FManagedArrayCollection Collection;
 
+	/** Geometric bounding box of the input collection */
 	UPROPERTY(meta = (DataflowOutput))
 	FBox BoundingBox = FBox(ForceInit);
 
-	FBoundingBoxDataflowNode(const UE::Dataflow::FNodeParameters& InParam, FGuid InGuid = FGuid::NewGuid())
-		: FDataflowNode(InParam, InGuid)
-	{
-		RegisterInputConnection(&Collection);
-		RegisterOutputConnection(&BoundingBox);
-	}
+	/** Center of the resulting bounding box */
+	UPROPERTY(meta = (DataflowOutput))
+	FVector Center = FVector::ZeroVector;
 
-	virtual void Evaluate(UE::Dataflow::FContext& Context, const FDataflowOutput* Out) const override;
+	/** Dimensions of the resulting bounding box in centimers */
+	UPROPERTY(meta = (DataflowOutput))
+	FVector Dimensions = FVector::ZeroVector;
+
+	FBoundingBoxDataflowNode(const UE::Dataflow::FNodeParameters& InParam, FGuid InGuid = FGuid::NewGuid());
 
+	virtual void Evaluate(UE::Dataflow::FContext& Context, const FDataflowOutput* Out) const override;
 };
 
 /**
@@ -306,10 +304,10 @@ USTRUCT()
 struct FExpandBoundingBoxDataflowNode : public FDataflowNode
 {
 	GENERATED_USTRUCT_BODY()
-	DATAFLOW_NODE_DEFINE_INTERNAL(FExpandBoundingBoxDataflowNode, "ExpandBoundingBox", "Utilities|Box", "")
+	DATAFLOW_NODE_DEFINE_INTERNAL(FExpandBoundingBoxDataflowNode, "BreakBox", "Utilities|Box", "Bounds Bounding Expand Min Max Center Half Extents Volume")
 
 public:
-	UPROPERTY(meta = (DataflowInput))
+	UPROPERTY(meta = (DataflowInput, DisplayName = "Box"))
 	FBox BoundingBox = FBox(ForceInit);;
 
 	UPROPERTY(meta = (DataflowOutput))
@@ -351,7 +349,7 @@ USTRUCT()
 struct FExpandBoundingSphereDataflowNode : public FDataflowNode
 {
 	GENERATED_USTRUCT_BODY()
-	DATAFLOW_NODE_DEFINE_INTERNAL(FExpandBoundingSphereDataflowNode, "ExpandBoundingSphere", "Utilities|Sphere", "")
+	DATAFLOW_NODE_DEFINE_INTERNAL(FExpandBoundingSphereDataflowNode, "BreakBoundingSphere", "Utilities|Sphere", "Expand Radius Center Volume")
 	DATAFLOW_NODE_RENDER_TYPE("PointRender", FName("FVector"), "Center")
 
 public:
@@ -389,7 +387,7 @@ USTRUCT()
 struct FExpandVectorDataflowNode : public FDataflowNode
 {
 	GENERATED_USTRUCT_BODY()
-	DATAFLOW_NODE_DEFINE_INTERNAL(FExpandVectorDataflowNode, "ExpandVector", "Utilities|Vector", "")
+	DATAFLOW_NODE_DEFINE_INTERNAL(FExpandVectorDataflowNode, "BreakVector", "Utilities|Vector", "Expand X Y Z")
 
 public:
 	UPROPERTY(meta = (DataflowInput))
