Dataflow : Fix an issue in the get centroid dataflow node not working properly if some of the transform did not have a geometry associated with them

diff --git a/Engine/Plugins/Experimental/GeometryCollectionPlugin/Source/GeometryCollectionNodes/Private/Dataflow/GeometryCollectionNodes.cpp b/Engine/Plugins/Experimental/GeometryCollectionPlugin/Source/GeometryCollectionNodes/Private/Dataflow/GeometryCollectionNodes.cpp
--- a/Engine/Plugins/Experimental/GeometryCollectionPlugin/Source/GeometryCollectionNodes/Private/Dataflow/GeometryCollectionNodes.cpp
+++ b/Engine/Plugins/Experimental/GeometryCollectionPlugin/Source/GeometryCollectionNodes/Private/Dataflow/GeometryCollectionNodes.cpp
@@ -500,39 +500,39 @@ void FGetRootIndexFromCollectionDataflowNode::Evaluate(UE::Dataflow::FContext& C
 
 void FGetCentroidsFromCollectionDataflowNode::Evaluate(UE::Dataflow::FContext& Context, const FDataflowOutput* Out) const
 {
-	if (Out->IsA<TArray<FVector>>(&Centroids))
+	if (Out->IsA(&Centroids))
 	{
-		const FManagedArrayCollection& InCollection = GetValue<FManagedArrayCollection>(Context, &Collection);
-		const FDataflowTransformSelection& InTransformSelection = GetValue<FDataflowTransformSelection>(Context, &TransformSelection);
+		TArray<FVector> OutCentroids;
 
-		GeometryCollection::Facades::FBoundsFacade BoundsFacade(InCollection);
-		const TArray<FVector>& InCentroids = BoundsFacade.GetCentroids();
+		const FManagedArrayCollection& InCollection = GetValue(Context, &Collection);
+		const FDataflowTransformSelection& InTransformSelection = GetValue(Context, &TransformSelection);
 
-		GeometryCollection::Facades::FCollectionTransformFacade TransformFacade(InCollection);
+		const bool bIsSelectionConnected = IsConnected(&TransformSelection);
 
-		TArray<FVector> CentroidsArr;
-		for (int32 Idx = 0; Idx < InCentroids.Num(); ++Idx)
+		GeometryCollection::Facades::FBoundsFacade BoundsFacade(InCollection);
+		GeometryCollection::Facades::FCollectionTransformFacade TransformFacade(InCollection);
+		if (BoundsFacade.IsValid() && TransformFacade.IsValid())
 		{
-			const FVector PositionInBoneSpace(InCentroids[Idx]);
+			const TArray<FVector>& InCentroids = BoundsFacade.GetCentroids();
+			const TManagedArray<int32>& TransformToGeometryIndex = BoundsFacade.GetTransformToGeometryIndex();
 
-			// Transform from BoneSpace to CollectionSpace
-			const FTransform CollectionSpaceTransform = TransformFacade.ComputeCollectionSpaceTransform(Idx);
-			const FVector PositionInCollectionSpace = CollectionSpaceTransform.TransformPosition(PositionInBoneSpace);
-
-			if (IsConnected<FDataflowTransformSelection>(&TransformSelection))
+			const int32 NumTransforms = InCollection.NumElements(FTransformCollection::TransformGroup);
+			for (int32 TransformIdx = 0; TransformIdx < NumTransforms; ++TransformIdx)
 			{
-				if (InTransformSelection.IsSelected(Idx))
+				const int32 GeometryIndex = TransformToGeometryIndex.IsValidIndex(TransformIdx)? TransformToGeometryIndex[TransformIdx]: INDEX_NONE;
+				if (InCentroids.IsValidIndex(GeometryIndex))
 				{
-					CentroidsArr.Add(PositionInCollectionSpace);
+					const FVector PositionInBoneSpace(InCentroids[GeometryIndex]);
+					const FTransform CollectionSpaceTransform = TransformFacade.ComputeCollectionSpaceTransform(TransformIdx);
+					const FVector PositionInCollectionSpace = CollectionSpaceTransform.TransformPosition(PositionInBoneSpace);
+					if (!bIsSelectionConnected || TransformIdx >= InTransformSelection.Num() || InTransformSelection.IsSelected(TransformIdx))
+					{
+						OutCentroids.Add(PositionInCollectionSpace);
+					}
 				}
 			}
-			else
-			{
-				CentroidsArr.Add(PositionInCollectionSpace);
-			}
 		}
-
-		SetValue(Context, MoveTemp(CentroidsArr), &Centroids);
+		SetValue(Context, MoveTemp(OutCentroids), &Centroids);
 	}
 }
 
diff --git a/Engine/Source/Runtime/Experimental/Chaos/Public/GeometryCollection/Facades/CollectionBoundsFacade.h b/Engine/Source/Runtime/Experimental/Chaos/Public/GeometryCollection/Facades/CollectionBoundsFacade.h
--- a/Engine/Source/Runtime/Experimental/Chaos/Public/GeometryCollection/Facades/CollectionBoundsFacade.h
+++ b/Engine/Source/Runtime/Experimental/Chaos/Public/GeometryCollection/Facades/CollectionBoundsFacade.h
@@ -44,7 +44,10 @@ namespace GeometryCollection::Facades
 		/** BoundingBox access */
 		const TManagedArray< FBox >& GetBoundingBoxes() const { return BoundingBoxAttribute.Get(); }
 
-		/** Centroids (Centers of BoundingBoxes) access */
+		/** 
+		* Get center of bounding boxes 
+		* warning : this is indexed on the geometry group not the transform group
+		*/
 		CHAOS_API TArray<FVector> GetCentroids() const;
 
 		/** BoundingBox for the whole collection in Collection space */
