Dataflow:
- Added rendering output collection to AppendCollections dataflow node
- Added debugdraw of centroids to GetCentroidsFromCollection dataflow node

diff --git a/Engine/Plugins/Experimental/GeometryCollectionPlugin/Source/GeometryCollectionNodes/Private/Dataflow/GeometryCollectionNodes.cpp b/Engine/Plugins/Experimental/GeometryCollectionPlugin/Source/GeometryCollectionNodes/Private/Dataflow/GeometryCollectionNodes.cpp
--- a/Engine/Plugins/Experimental/GeometryCollectionPlugin/Source/GeometryCollectionNodes/Private/Dataflow/GeometryCollectionNodes.cpp
+++ b/Engine/Plugins/Experimental/GeometryCollectionPlugin/Source/GeometryCollectionNodes/Private/Dataflow/GeometryCollectionNodes.cpp
@@ -44,6 +44,7 @@
 #include "Dataflow/DataflowRenderingViewMode.h"
 #endif
 #include "Dataflow/GeometryCollectionUtils.h"
+#include "Dataflow/DataflowDebugDrawInterface.h"
 
 #include UE_INLINE_GENERATED_CPP_BY_NAME(GeometryCollectionNodes)
 
@@ -498,6 +499,16 @@ void FGetRootIndexFromCollectionDataflowNode::Evaluate(UE::Dataflow::FContext& C
 	}
 }
 
+//-----------------------------------------------------------------------------------------------
+
+FGetCentroidsFromCollectionDataflowNode::FGetCentroidsFromCollectionDataflowNode(const UE::Dataflow::FNodeParameters& InParam, FGuid InGuid) 
+	: FDataflowNode(InParam, InGuid)
+{
+	RegisterInputConnection(&Collection);
+	RegisterInputConnection(&TransformSelection);
+	RegisterOutputConnection(&Centroids);
+}
+
 void FGetCentroidsFromCollectionDataflowNode::Evaluate(UE::Dataflow::FContext& Context, const FDataflowOutput* Out) const
 {
 	if (Out->IsA(&Centroids))
@@ -532,10 +543,54 @@ void FGetCentroidsFromCollectionDataflowNode::Evaluate(UE::Dataflow::FContext& C
 				}
 			}
 		}
+
 		SetValue(Context, MoveTemp(OutCentroids), &Centroids);
 	}
 }
 
+#if WITH_EDITOR
+bool FGetCentroidsFromCollectionDataflowNode::CanDebugDraw() const
+{
+	return true;
+}
+
+bool FGetCentroidsFromCollectionDataflowNode::CanDebugDrawViewMode(const FName& ViewModeName) const
+{
+	return ViewModeName == UE::Dataflow::FDataflowConstruction3DViewMode::Name;
+}
+
+void FGetCentroidsFromCollectionDataflowNode::DebugDraw(UE::Dataflow::FContext& Context, IDataflowDebugDrawInterface& DataflowRenderingInterface, const FDebugDrawParameters& DebugDrawParameters) const
+{
+	if ((DebugDrawParameters.bNodeIsSelected || DebugDrawParameters.bNodeIsPinned))
+	{
+		if (IsConnected(&Collection))
+		{
+			const FDataflowOutput* CentroidsOutput = FindOutput(&Centroids);
+			if (CentroidsOutput)
+			{
+				const TArray<FVector> EmptyArray;
+				const TArray<FVector>& OutCentroids = CentroidsOutput->GetValue<TArray<FVector>>(Context, EmptyArray);
+
+				DataflowRenderingInterface.SetWireframe(true);
+				DataflowRenderingInterface.SetWorldPriority();
+				DataflowRenderingInterface.SetShaded(false);
+				DataflowRenderingInterface.SetColor(Color);
+				DataflowRenderingInterface.SetPointSize(PointSize);
+
+				if (OutCentroids.Num() > 0)
+				{
+					for (int32 Idx = 0; Idx < OutCentroids.Num(); ++Idx)
+					{
+						DataflowRenderingInterface.DrawPoint(FVector(OutCentroids[Idx]));
+					}
+				}
+			}
+		}
+	}
+}
+#endif
+
+//-----------------------------------------------------------------------------------------------
 
 void FTransformCollectionDataflowNode::Evaluate(UE::Dataflow::FContext& Context, const FDataflowOutput* Out) const
 {
diff --git a/Engine/Plugins/Experimental/GeometryCollectionPlugin/Source/GeometryCollectionNodes/Public/Dataflow/GeometryCollectionNodes.h b/Engine/Plugins/Experimental/GeometryCollectionPlugin/Source/GeometryCollectionNodes/Public/Dataflow/GeometryCollectionNodes.h
--- a/Engine/Plugins/Experimental/GeometryCollectionPlugin/Source/GeometryCollectionNodes/Public/Dataflow/GeometryCollectionNodes.h
+++ b/Engine/Plugins/Experimental/GeometryCollectionPlugin/Source/GeometryCollectionNodes/Public/Dataflow/GeometryCollectionNodes.h
@@ -11,6 +11,7 @@
 #include "Dataflow/DataflowSelection.h"
 #include "Math/MathFwd.h"
 #include "DynamicMesh/DynamicMesh3.h"
+#include "Dataflow/DataflowDebugDraw.h"
 
 #include "GeometryCollectionNodes.generated.h"
 
@@ -67,7 +68,7 @@ struct FAppendCollectionAssetsDataflowNode : public FDataflowNode
 {
 	GENERATED_USTRUCT_BODY()
 	DATAFLOW_NODE_DEFINE_INTERNAL(FAppendCollectionAssetsDataflowNode, "AppendCollections", "GeometryCollection", "")
-	DATAFLOW_NODE_RENDER_TYPE("SurfaceRender", FGeometryCollection::StaticType(), "Collection")
+	DATAFLOW_NODE_RENDER_TYPE("SurfaceRender", FGeometryCollection::StaticType(), "Collection1")
 
 public:
 	typedef FManagedArrayCollection DataType;
@@ -82,7 +83,7 @@ public:
 	TArray<FString> GeometryGroupGuidsOut1;
 
 	UPROPERTY(meta = (DataflowOutput, DisplayName = "GeometryGroupIndicesOut2"))
-		TArray<FString> GeometryGroupGuidsOut2;
+	TArray<FString> GeometryGroupGuidsOut2;
 
 	FAppendCollectionAssetsDataflowNode(const UE::Dataflow::FNodeParameters& InParam, FGuid InGuid = FGuid::NewGuid())
 		: FDataflowNode(InParam, InGuid)
@@ -630,7 +631,7 @@ struct FGetCentroidsFromCollectionDataflowNode : public FDataflowNode
 	GENERATED_USTRUCT_BODY()
 	DATAFLOW_NODE_DEFINE_INTERNAL(FGetCentroidsFromCollectionDataflowNode, "GetCentroidsFromCollection", "GeometryCollection|Utilities", "")
 
-public:
+private:
 	/** Input Collection */
 	UPROPERTY(meta = (DataflowInput))
 	FManagedArrayCollection Collection;
@@ -639,20 +640,26 @@ public:
 	UPROPERTY(meta = (DataflowInput))
 	FDataflowTransformSelection TransformSelection;
 
+	UPROPERTY(EditAnywhere, Category = "Debug Draw")
+	FLinearColor Color = FLinearColor::Blue;
+
+	UPROPERTY(EditAnywhere, Category = "Debug Draw")
+	float PointSize = 12.f;
+
 	/** Output centroids */
 	UPROPERTY(meta = (DataflowOutput, DisplayName = "Points"))
 	TArray<FVector> Centroids = TArray<FVector>();
 
-	FGetCentroidsFromCollectionDataflowNode(const UE::Dataflow::FNodeParameters& InParam, FGuid InGuid = FGuid::NewGuid())
-		: FDataflowNode(InParam, InGuid)
-	{
-		RegisterInputConnection(&Collection);
-		RegisterInputConnection(&TransformSelection);
-		RegisterOutputConnection(&Centroids);
-	}
-
 	virtual void Evaluate(UE::Dataflow::FContext& Context, const FDataflowOutput* Out) const override;
 
+#if WITH_EDITOR
+	virtual bool CanDebugDraw() const override;
+	virtual bool CanDebugDrawViewMode(const FName& ViewModeName) const override;
+	virtual void DebugDraw(UE::Dataflow::FContext& Context, IDataflowDebugDrawInterface& DataflowRenderingInterface, const FDebugDrawParameters& DebugDrawParameters) const override;
+#endif
+public:
+	FGetCentroidsFromCollectionDataflowNode(const UE::Dataflow::FNodeParameters& InParam, FGuid InGuid = FGuid::NewGuid());
+
 };
 
 
