Made blueprint dataflow nodes evaluate on the game thread, for safety

diff --git a/Engine/Plugins/Experimental/GeometryCollectionPlugin/Source/GeometryCollectionNodes/Public/Dataflow/GeometryCollectionMeshNodes.h b/Engine/Plugins/Experimental/GeometryCollectionPlugin/Source/GeometryCollectionNodes/Public/Dataflow/GeometryCollectionMeshNodes.h
--- a/Engine/Plugins/Experimental/GeometryCollectionPlugin/Source/GeometryCollectionNodes/Public/Dataflow/GeometryCollectionMeshNodes.h
+++ b/Engine/Plugins/Experimental/GeometryCollectionPlugin/Source/GeometryCollectionNodes/Public/Dataflow/GeometryCollectionMeshNodes.h
@@ -109,6 +109,12 @@ private:
 
 	virtual void OnPropertyChanged(UE::Dataflow::FContext& Context, const FPropertyChangedEvent& PropertyChangedEvent) override;
 
+	// Blueprint evaluation is not generally thread-safe
+	virtual bool EvaluateOnGameThreadOnly() const override
+	{
+		return true;
+	}
+
 	// Handling for the selected blueprint being changed under the node (e.g., recompiled)
 	FDelegateHandle BlueprintChangeDelegateHandle;
 	void TeardownBlueprintEvent();
