Added option for dataflow FilterPointsWithMesh node to filter by mesh min/max distances in addition to inside/outside winding number test

diff --git a/Engine/Plugins/Experimental/GeometryCollectionPlugin/Source/GeometryCollectionNodes/Private/Dataflow/GeometryCollectionSamplingNodes.cpp b/Engine/Plugins/Experimental/GeometryCollectionPlugin/Source/GeometryCollectionNodes/Private/Dataflow/GeometryCollectionSamplingNodes.cpp
--- a/Engine/Plugins/Experimental/GeometryCollectionPlugin/Source/GeometryCollectionNodes/Private/Dataflow/GeometryCollectionSamplingNodes.cpp
+++ b/Engine/Plugins/Experimental/GeometryCollectionPlugin/Source/GeometryCollectionNodes/Private/Dataflow/GeometryCollectionSamplingNodes.cpp
@@ -26,6 +26,9 @@ FFilterPointSetWithMeshDataflowNode::FFilterPointSetWithMeshDataflowNode(const U
 	RegisterInputConnection(&SamplePoints);
 	RegisterInputConnection(&bKeepInside).SetCanHidePin(true).SetPinIsHidden(true);
 	RegisterInputConnection(&WindingThreshold).SetCanHidePin(true).SetPinIsHidden(true);
+	RegisterInputConnection(&MinDistance).SetCanHidePin(true).SetPinIsHidden(true);
+	RegisterInputConnection(&MaxDistance).SetCanHidePin(true).SetPinIsHidden(true);
+	RegisterInputConnection(&bUseSignedDistance).SetCanHidePin(true).SetPinIsHidden(true);
 
 	RegisterOutputConnection(&SamplePoints);
 }
@@ -35,26 +38,125 @@ void FFilterPointSetWithMeshDataflowNode::Evaluate(UE::Dataflow::FContext& Conte
 	if (Out->IsA(&SamplePoints))
 	{
 		TArray<FVector> InSamplePoints = GetValue(Context, &SamplePoints);
-		double InWindingThreshold = (double)GetValue(Context, &WindingThreshold);
-		bool bInKeepInside = GetValue(Context, &bKeepInside);
+		const double InWindingThreshold = (double)GetValue(Context, &WindingThreshold);
+		const bool bInKeepInside = GetValue(Context, &bKeepInside);
+		const double InMinDistance = (double)GetValue(Context, &MinDistance);
+		const double InMaxDistance = (double)GetValue(Context, &MaxDistance);
+		const bool bInUseSignedDistance = GetValue(Context, &bUseSignedDistance);
+
 		if (TObjectPtr<const UDynamicMesh> InTargetMesh = GetValue(Context, &TargetMesh))
 		{
 			using namespace UE::Geometry;
-			InTargetMesh->ProcessMesh([&InSamplePoints, InWindingThreshold, bInKeepInside](const FDynamicMesh3& Mesh)
+			InTargetMesh->ProcessMesh(
+				[&InSamplePoints, InWindingThreshold, bInKeepInside, InMinDistance, InMaxDistance, bInUseSignedDistance, this]
+				(const FDynamicMesh3& Mesh)
 				{
+					const bool bFilterWinding = bool(FilterMethod & (uint8)EFilterPointSetWithMeshDataflowMethodFlags::Winding);
+					const bool bFilterMinDist = bool(FilterMethod & (uint8)EFilterPointSetWithMeshDataflowMethodFlags::MinDistance);
+					const bool bFilterMaxDist = bool(FilterMethod & (uint8)EFilterPointSetWithMeshDataflowMethodFlags::MaxDistance);
+
+					if (!bFilterWinding && !bFilterMinDist && !bFilterMaxDist)
+					{
+						// No filtering, early-out
+						return;
+					}
+
+					// Clamp unsigned distances to 0
+					double UseMinDistance = InMinDistance;
+					double UseMaxDistance = InMaxDistance;
+					if (!bInUseSignedDistance)
+					{
+						UseMinDistance = FMath::Max(0, UseMinDistance);
+						UseMaxDistance = FMath::Max(0, UseMaxDistance);
+					}
+
+					if (bFilterMinDist && bFilterMaxDist)
+					{
+						// If Min > Max, impossible for points to pass the filter, so early-out
+						if (UseMinDistance > UseMaxDistance)
+						{
+							InSamplePoints.Empty();
+							return;
+						}
+					}
+
+					const bool bNeedWinding = bFilterWinding || bInUseSignedDistance;
+					const bool bNeedDistance = bFilterMinDist || bFilterMaxDist;
+					const double MaxRelevantDistance = UE_DOUBLE_KINDA_SMALL_NUMBER + FMath::Max(
+						bFilterMinDist ? FMath::Abs(UseMinDistance) : 0,
+						bFilterMaxDist ? FMath::Abs(UseMaxDistance) : 0);
+
 					//  set up AABBTree and FWNTree lists
 					TMeshAABBTree3<FDynamicMesh3> Spatial(&Mesh);
-					TFastWindingTree<FDynamicMesh3> FastWinding(&Spatial);
-
-					// Filter points w/ the fast winding tree
-					TArray<bool> PointInside;
-					PointInside.SetNumZeroed(InSamplePoints.Num());
-					ParallelFor(PointInside.Num(), 
-						[&PointInside, &FastWinding, &InSamplePoints, InWindingThreshold]
+					TFastWindingTree<FDynamicMesh3> FastWinding(&Spatial, bNeedWinding);
+
+					// square threshold distances, but keep their signs, so we can do threshold testing w/out sqrt
+					const double UseSignSqMinDist = FMath::CopySign(UseMinDistance * UseMinDistance, UseMinDistance);
+					const double UseSignSqMaxDist = FMath::CopySign(UseMaxDistance * UseMaxDistance, UseMaxDistance);
+
+					// Filter points 
+					TArray<bool> KeepPoint;
+					KeepPoint.SetNumUninitialized(InSamplePoints.Num());
+					ParallelFor(KeepPoint.Num(),
+						[&KeepPoint, &Spatial, &FastWinding, &InSamplePoints, 
+						bFilterWinding, bFilterMinDist, bFilterMaxDist, bInUseSignedDistance, bInKeepInside,
+						bNeedWinding, InWindingThreshold, bNeedDistance, UseSignSqMinDist, UseSignSqMaxDist, MaxRelevantDistance]
 						(int32 PointIdx)
 						{
 							FVector Pt = InSamplePoints[PointIdx];
-							PointInside[PointIdx] = FastWinding.IsInside(Pt, InWindingThreshold);
+
+							bool bWindingInside = false;
+							if (bNeedWinding)
+							{
+								bWindingInside = FastWinding.IsInside(Pt, InWindingThreshold);
+								// test if we fail the winding filter
+								if (bFilterWinding && bInKeepInside != bWindingInside)
+								{
+									KeepPoint[PointIdx] = false;
+									return;
+								}
+							}
+							double FoundDistSq = 0;
+							if (bNeedDistance)
+							{
+								int32 FoundTID = Spatial.FindNearestTriangle(Pt, FoundDistSq, UE::Geometry::IMeshSpatial::FQueryOptions(MaxRelevantDistance));
+								if (FoundTID == INDEX_NONE)
+								{
+									// we have a max dist, lack of closest point -> we fail the max dist filter
+									// or it's inside the shape w/ signed distances -> it's too far inside, fail the min filter
+									if (bFilterMaxDist || (bInUseSignedDistance && bWindingInside))
+									{
+										KeepPoint[PointIdx] = false;
+										return;
+									}
+									else
+									{
+										// point at least passes the min distance threshold
+										FoundDistSq = UseSignSqMinDist;
+									}
+								}
+								else
+								{
+									if (bInUseSignedDistance && bWindingInside)
+									{
+										// sign the squared distance
+										FoundDistSq = -FoundDistSq;
+									}
+								}
+								if (bFilterMinDist && FoundDistSq < UseSignSqMinDist)
+								{
+									KeepPoint[PointIdx] = false;
+									return;
+								}
+								else if (bFilterMaxDist && FoundDistSq > UseSignSqMaxDist)
+								{
+									KeepPoint[PointIdx] = false;
+									return;
+								}
+							}
+
+							// passed all filters, keep the point
+							KeepPoint[PointIdx] = true;
 						}
 					);
 
@@ -62,7 +164,7 @@ void FFilterPointSetWithMeshDataflowNode::Evaluate(UE::Dataflow::FContext& Conte
 					int32 FoundPoints = 0;
 					for (int32 Idx = 0; ; ++Idx)
 					{
-						while (Idx < InSamplePoints.Num() && PointInside[Idx] != bInKeepInside)
+						while (Idx < InSamplePoints.Num() && !KeepPoint[Idx])
 						{
 							Idx++;
 						}
diff --git a/Engine/Plugins/Experimental/GeometryCollectionPlugin/Source/GeometryCollectionNodes/Public/Dataflow/GeometryCollectionSamplingNodes.h b/Engine/Plugins/Experimental/GeometryCollectionPlugin/Source/GeometryCollectionNodes/Public/Dataflow/GeometryCollectionSamplingNodes.h
--- a/Engine/Plugins/Experimental/GeometryCollectionPlugin/Source/GeometryCollectionNodes/Public/Dataflow/GeometryCollectionSamplingNodes.h
+++ b/Engine/Plugins/Experimental/GeometryCollectionPlugin/Source/GeometryCollectionNodes/Public/Dataflow/GeometryCollectionSamplingNodes.h
@@ -16,6 +16,20 @@ namespace UE_DEPRECATED(5.5, "Use UE::Dataflow instead.") Dataflow {}
 class FGeometryCollection;
 class UDynamicMesh;
 
+/** Flags to control which mesh point filtering method(s) are applied */
+UENUM(meta = (Bitflags, UseEnumValuesAsMaskValuesInEditor = "true"))
+enum class EFilterPointSetWithMeshDataflowMethodFlags : uint8
+{
+	None = 0 UMETA(Hidden),
+	// Use the winding number to filter inside or outside of the mesh
+	Winding = 1 << 0,
+	// Filter away points below a minimum mesh distance
+	MinDistance = 1 << 1,
+	// Filter away points above a maximum mesh distance
+	MaxDistance = 1 << 2
+};
+ENUM_CLASS_FLAGS(EFilterPointSetWithMeshDataflowMethodFlags)
+
 /**
  *
  * Filter a point set to only the points inside or outside of a given mesh
@@ -33,14 +47,32 @@ private:
 	UPROPERTY(meta = (DataflowInput, DataflowIntrinsic))
 	TObjectPtr<UDynamicMesh> TargetMesh;
 
-	/** Whether to keep the points inside or (if false) outside the mesh. */
-	UPROPERTY(EditAnywhere, Category=Options, meta = (DataflowInput))
+	UPROPERTY(EditAnywhere, Category=Options, meta=(Bitmask, BitmaskEnum="/Script/GeometryCollectionNodes.EFilterPointSetWithMeshDataflowMethodFlags"))
+	uint8 FilterMethod = (uint8)EFilterPointSetWithMeshDataflowMethodFlags::Winding;
+
+	/** Whether to keep the points inside or (if false) outside the mesh, when filtering by Winding Number. */
+	UPROPERTY(EditAnywhere, Category = WindingNumber, meta = (DataflowInput, EditCondition = "FilterMethod & \"/Script/GeometryCollectionNodes.EFilterPointSetWithMeshDataflowMethodFlags::Winding\""))
 	bool bKeepInside = true;
 
-	/** The winding number threshold to use for determining whether a point is inside or outside of the mesh */
-	UPROPERTY(EditAnywhere, Category = Options, meta = (DataflowInput))
+	/** The winding number threshold to use for determining whether a point is inside or outside of the mesh, if corresponding Filter Method is set  */
+	UPROPERTY(EditAnywhere, Category = WindingNumber, meta = (DataflowInput))
 	float WindingThreshold = .5f;
 
+	/** The min distance to surface to keep, if corresponding Filter Method is set  */
+	UPROPERTY(EditAnywhere, Category = Distance, meta = (DataflowInput, EditCondition = "FilterMethod & \"/Script/GeometryCollectionNodes.EFilterPointSetWithMeshDataflowMethodFlags::MinDistance\""))
+	float MinDistance = 0.f;
+
+	/** The max distance to surface to keep, if corresponding Filter Method is set */
+	UPROPERTY(EditAnywhere, Category = Distance, meta = (DataflowInput, EditCondition = "FilterMethod & \"/Script/GeometryCollectionNodes.EFilterPointSetWithMeshDataflowMethodFlags::MaxDistance\""))
+	float MaxDistance = 1000.f;
+
+	/**
+	 * Whether to use signed distances for the Min and Max Distance thresholds. Otherwise, unsigned distance is used.
+	 * Note: Signs are computed via the Winding Number. The sign is negative if the point's Winding Number is below the Winding Threshold. 
+	 */
+	UPROPERTY(EditAnywhere, Category = Distance, meta = (DataflowInput))
+	bool bUseSignedDistance = false;
+
 	/** Points to filter */
 	UPROPERTY(meta = (DataflowInput, DataflowOutput))
 	TArray<FVector> SamplePoints;
