minor cleanup of mesh filtering dataflow node (addressing earlier comments)

diff --git a/Engine/Plugins/Experimental/GeometryCollectionPlugin/Source/GeometryCollectionNodes/Private/Dataflow/GeometryCollectionSamplingNodes.cpp b/Engine/Plugins/Experimental/GeometryCollectionPlugin/Source/GeometryCollectionNodes/Private/Dataflow/GeometryCollectionSamplingNodes.cpp
--- a/Engine/Plugins/Experimental/GeometryCollectionPlugin/Source/GeometryCollectionNodes/Private/Dataflow/GeometryCollectionSamplingNodes.cpp
+++ b/Engine/Plugins/Experimental/GeometryCollectionPlugin/Source/GeometryCollectionNodes/Private/Dataflow/GeometryCollectionSamplingNodes.cpp
@@ -119,7 +119,7 @@ void FFilterPointSetWithMeshDataflowNode::Evaluate(UE::Dataflow::FContext& Conte
 							double FoundDistSq = 0;
 							if (bNeedDistance)
 							{
-								int32 FoundTID = Spatial.FindNearestTriangle(Pt, FoundDistSq, UE::Geometry::IMeshSpatial::FQueryOptions(MaxRelevantDistance));
+								const int32 FoundTID = Spatial.FindNearestTriangle(Pt, FoundDistSq, UE::Geometry::IMeshSpatial::FQueryOptions(MaxRelevantDistance));
 								if (FoundTID == INDEX_NONE)
 								{
 									// we have a max dist, lack of closest point -> we fail the max dist filter
