Made the collection->mesh dataflow node take an optional selection and perform the conversion more consistently with the collection selection->meshes dataflow node

diff --git a/Engine/Plugins/Experimental/GeometryCollectionPlugin/Source/GeometryCollectionNodes/Private/Dataflow/GeometryCollectionMeshNodes.cpp b/Engine/Plugins/Experimental/GeometryCollectionPlugin/Source/GeometryCollectionNodes/Private/Dataflow/GeometryCollectionMeshNodes.cpp
--- a/Engine/Plugins/Experimental/GeometryCollectionPlugin/Source/GeometryCollectionNodes/Private/Dataflow/GeometryCollectionMeshNodes.cpp
+++ b/Engine/Plugins/Experimental/GeometryCollectionPlugin/Source/GeometryCollectionNodes/Private/Dataflow/GeometryCollectionMeshNodes.cpp
@@ -35,6 +35,7 @@
 #include "FractureEngineSelection.h"
 #include "FractureEngineUtility.h"
 #include "DynamicMesh/DynamicMesh3.h"
+#include "DynamicMesh/MeshTransforms.h"
 #include "UDynamicMesh.h"
 
 #include UE_INLINE_GENERATED_CPP_BY_NAME(GeometryCollectionMeshNodes)
@@ -227,52 +228,57 @@ void FMeshToCollectionDataflowNode::Evaluate(UE::Dataflow::FContext& Context, co
 
 void FCollectionToMeshDataflowNode::Evaluate(UE::Dataflow::FContext& Context, const FDataflowOutput* Out) const
 {
-#if WITH_EDITORONLY_DATA
-	if (Out->IsA<TObjectPtr<UDynamicMesh>>(&Mesh))
+	if (Out->IsA(&Mesh))
 	{
-		const FManagedArrayCollection& InCollection = GetValue<FManagedArrayCollection>(Context, &Collection);
-		
-		if (InCollection.NumElements(FGeometryCollection::TransformGroup) > 0)
-		{
-			if (TUniquePtr<FGeometryCollection> GeomCollection = TUniquePtr<FGeometryCollection>(InCollection.NewCopy<FGeometryCollection>()))
-			{
-				const TManagedArray<FTransform3f>& BoneTransforms = InCollection.GetAttribute<FTransform3f>("Transform", FGeometryCollection::TransformGroup);
-
-				TArray<int32> TransformIndices;
-				TransformIndices.AddUninitialized(BoneTransforms.Num());
-
-				int32 Idx = 0;
-				for (int32& TransformIdx : TransformIndices)
-				{
-					TransformIdx = Idx++;
-				}
-
-				FMeshDescription MeshDescription;
-				FStaticMeshAttributes Attributes(MeshDescription);
-				Attributes.Register();
+		TObjectPtr<UDynamicMesh> NewMesh = NewObject<UDynamicMesh>();
+		FDataflowTransformSelection InTransformSelection = GetValue(Context, &TransformSelection);
 
-				FTransform TransformOut;
+		//
+		// If not connected select everything by default
+		//
+		if (!IsConnected(&TransformSelection))
+		{
+			const FManagedArrayCollection& InCollection = GetValue(Context, &Collection);
 
-				ConvertToMeshDescription(MeshDescription, TransformOut, bCenterPivot, *GeomCollection, BoneTransforms, TransformIndices);
+			InTransformSelection.Initialize(InCollection.NumElements(FGeometryCollection::TransformGroup), true);
+		}
 
-				TObjectPtr<UDynamicMesh> NewMesh = NewObject<UDynamicMesh>();
-				NewMesh->Reset();
+		if (InTransformSelection.AnySelected())
+		{
+			const FManagedArrayCollection& InCollection = GetValue(Context, &Collection);
 
-				UE::Geometry::FDynamicMesh3& DynMesh = NewMesh->GetMeshRef();
-				{
-					FMeshDescriptionToDynamicMesh ConverterToDynamicMesh;
-					ConverterToDynamicMesh.Convert(&MeshDescription, DynMesh);
-				}
+			TArray<int32> TransformSelectionArray = InTransformSelection.AsArray();
 
-				SetValue(Context, NewMesh, &Mesh);
+			GeometryCollection::Facades::FCollectionTransformSelectionFacade SelectionFacade(InCollection);
+			Chaos::Facades::FCollectionHierarchyFacade HierarchyFacade(InCollection);
+			TArray<int32> LeafSelectionArray = TransformSelectionArray;
+			SelectionFacade.ConvertSelectionToRigidNodes(LeafSelectionArray);
 
-				return;
+			UE::Geometry::FGeometryCollectionToDynamicMeshes CollectionToMeshes;
+			UE::Geometry::FGeometryCollectionToDynamicMeshes::FToMeshOptions ToMeshOptions;
+			ToMeshOptions.bWeldVertices = bWeldVertices;
+			ToMeshOptions.bSaveIsolatedVertices = bPreserveIsolatedVertices;
+			if (CollectionToMeshes.InitFromTransformSelection(InCollection, LeafSelectionArray, ToMeshOptions)
+				&& !CollectionToMeshes.Meshes.IsEmpty())
+			{
+				NewMesh->EditMesh([this, &CollectionToMeshes](UE::Geometry::FDynamicMesh3& CombinedMesh)
+					{
+						CombinedMesh = MoveTemp(*CollectionToMeshes.Meshes[0].Mesh);
+						for (int32 MeshIdx = 1; MeshIdx < CollectionToMeshes.Meshes.Num(); ++MeshIdx)
+						{
+							CombinedMesh.AppendWithOffsets(*CollectionToMeshes.Meshes[MeshIdx].Mesh);
+						}
+						if (bCenterPivot)
+						{
+							MeshTransforms::Translate(CombinedMesh, -CombinedMesh.GetBounds().Center());
+						}
+					}
+				);
 			}
 		}
 
-		SetValue(Context, TObjectPtr<UDynamicMesh>(NewObject<UDynamicMesh>()), &Mesh);
+		SetValue(Context, NewMesh, &Mesh);
 	}
-#endif
 }
 
 
diff --git a/Engine/Plugins/Experimental/GeometryCollectionPlugin/Source/GeometryCollectionNodes/Public/Dataflow/GeometryCollectionMeshNodes.h b/Engine/Plugins/Experimental/GeometryCollectionPlugin/Source/GeometryCollectionNodes/Public/Dataflow/GeometryCollectionMeshNodes.h
--- a/Engine/Plugins/Experimental/GeometryCollectionPlugin/Source/GeometryCollectionNodes/Public/Dataflow/GeometryCollectionMeshNodes.h
+++ b/Engine/Plugins/Experimental/GeometryCollectionPlugin/Source/GeometryCollectionNodes/Public/Dataflow/GeometryCollectionMeshNodes.h
@@ -401,7 +401,8 @@ public:
 	UPROPERTY(meta = (DataflowInput, DataflowIntrinsic))
 	FManagedArrayCollection Collection;
 
-	UPROPERTY(EditAnywhere, Category = "General", meta = (DisplayName = "Center Pivot"));
+	/** Whether to translate the mesh geometry to be centered around its bounding box. */
+	UPROPERTY(EditAnywhere, Category = Options)
 	bool bCenterPivot = false;
 
 	/** Output DynamicMesh */
@@ -412,10 +413,25 @@ public:
 		: FDataflowNode(InParam, InGuid)
 	{
 		RegisterInputConnection(&Collection);
+		RegisterInputConnection(&TransformSelection);
 		RegisterOutputConnection(&Mesh);
 	}
 
 	virtual void Evaluate(UE::Dataflow::FContext& Context, const FDataflowOutput* Out) const override;
+	
+private:
+	/** Geometry on or under selected bones will be appended to the output mesh. If not connected, all geometry will be processed. */
+	UPROPERTY(meta = (DataflowInput))
+	FDataflowTransformSelection TransformSelection;
+
+	/** Whether the processed mesh will have edges at normal/UV/color seams welded so they are treated as one edge during processing. */
+	UPROPERTY(EditAnywhere, Category = Options)
+	bool bWeldVertices = true;
+
+	/** Whether to preserve isolated vertices which aren't used by any triangles. */
+	UPROPERTY(EditAnywhere, Category = Options)
+	bool bPreserveIsolatedVertices = true;
+
 
 };
 
