Dataflow : Add string attribute read and write Dataflow nodes - by default the "BoneName" attribute will be used

diff --git a/Engine/Plugins/Experimental/GeometryCollectionPlugin/Source/GeometryCollectionNodes/Private/Dataflow/GeometryCollectionSelectionNodes.cpp b/Engine/Plugins/Experimental/GeometryCollectionPlugin/Source/GeometryCollectionNodes/Private/Dataflow/GeometryCollectionSelectionNodes.cpp
--- a/Engine/Plugins/Experimental/GeometryCollectionPlugin/Source/GeometryCollectionNodes/Private/Dataflow/GeometryCollectionSelectionNodes.cpp
+++ b/Engine/Plugins/Experimental/GeometryCollectionPlugin/Source/GeometryCollectionNodes/Private/Dataflow/GeometryCollectionSelectionNodes.cpp
@@ -77,6 +77,9 @@ namespace UE::Dataflow
 		DATAFLOW_NODE_REGISTER_CREATION_FACTORY(FCollectionSelectionByAttrDataflowNode);
 		DATAFLOW_NODE_REGISTER_CREATION_FACTORY(FGeometrySelectionToVertexSelectionDataflowNode);
 		DATAFLOW_NODE_REGISTER_CREATION_FACTORY(FCollectionSelectInternalFacesDataflowNode);
+		DATAFLOW_NODE_REGISTER_CREATION_FACTORY(FCollectionSelectTransformStringDataflowNode);
+		DATAFLOW_NODE_REGISTER_CREATION_FACTORY(FCollectionSetTransformStringValueDataflowNode);
+			;
 
 		// generic input nodes
 		DATAFLOW_NODE_REGISTER_CREATION_FACTORY(FCollectionSelectionSetOperationDataflowNode);
@@ -1844,4 +1847,122 @@ void FCollectionSelectInternalFacesDataflowNode::Evaluate(UE::Dataflow::FContext
 		}
 		SetValue(Context, OutFaceSelection, &FaceSelection);
 	}
+}
+
+/////////////////////////////////////////////////////////////////////////////////
+
+FCollectionSelectTransformStringDataflowNode::FCollectionSelectTransformStringDataflowNode(const UE::Dataflow::FNodeParameters& InParam, FGuid InGuid)
+	: FDataflowNode(InParam, InGuid)
+{
+	RegisterInputConnection(&Collection);
+	RegisterInputConnection(&SearchText);
+	RegisterInputConnection(&Attribute).SetCanHidePin(true).SetPinIsHidden(true);
+
+	RegisterOutputConnection(&Collection, &Collection);
+	RegisterOutputConnection(&TransformSelection);
+}
+
+void FCollectionSelectTransformStringDataflowNode::Evaluate(UE::Dataflow::FContext& Context, const FDataflowOutput* Out) const
+{
+	if (Out->IsA(&Collection))
+	{
+		SafeForwardInput(Context, &Collection, &Collection);
+	}
+	else if (Out->IsA(&TransformSelection))
+	{
+		const FManagedArrayCollection& InCollection = GetValue(Context, &Collection);
+		const FString& InSearchText = GetValue(Context, &SearchText);
+		const FString& InAttributeName = GetValue(Context, &Attribute);
+
+		FDataflowTransformSelection OutSelection;
+		OutSelection.InitFromArray(InCollection, {});
+
+		const TManagedArray<FString>* StringAttribute = InCollection.FindAttributeTyped<FString>(FName(InAttributeName), FTransformCollection::TransformGroup);
+		if (StringAttribute)
+		{
+			const int32 Num = StringAttribute->Num();
+			for (int32 Index = 0; Index < Num; Index++)
+			{
+				const FString& Value = (*StringAttribute)[Index];
+
+				bool bSelect = false;
+				switch (Method)
+				{
+				case EDataflowCollectionSelectionByNameMethod::Exact:
+					bSelect = (Value == InSearchText);
+					break;
+				case EDataflowCollectionSelectionByNameMethod::StartsWith:
+					bSelect = Value.StartsWith(InSearchText, ESearchCase::CaseSensitive);
+					break;
+				case EDataflowCollectionSelectionByNameMethod::EndsWith:
+					bSelect = Value.EndsWith(InSearchText, ESearchCase::CaseSensitive);
+					break;
+				case EDataflowCollectionSelectionByNameMethod::Contains:
+					bSelect = Value.Contains(InSearchText, ESearchCase::CaseSensitive);
+					break;
+				}
+				if (bSelect)
+				{
+					OutSelection.SetSelected(Index);
+				}
+				else
+				{
+					OutSelection.SetNotSelected(Index);
+				}
+			}
+		}
+		SetValue(Context, OutSelection, &TransformSelection);
+	}
+}
+
+//////////////////////////////////////////////////////////////////////////////////////
+FCollectionSetTransformStringValueDataflowNode::FCollectionSetTransformStringValueDataflowNode(const UE::Dataflow::FNodeParameters& InParam, FGuid InGuid)
+	: FDataflowNode(InParam, InGuid)
+{
+	RegisterInputConnection(&Collection);
+	RegisterInputConnection(&TextToSet);
+	RegisterInputConnection(&TransformSelection);
+	RegisterInputConnection(&Attribute).SetCanHidePin(true).SetPinIsHidden(true);
+
+	RegisterOutputConnection(&Collection, &Collection);
+}
+
+void FCollectionSetTransformStringValueDataflowNode::Evaluate(UE::Dataflow::FContext& Context, const FDataflowOutput* Out) const
+{
+	if (Out->IsA(&Collection))
+	{
+		FManagedArrayCollection CollectionToUpdate = GetValue(Context, &Collection);
+		FString InTextToSet = GetValue(Context, &TextToSet);
+		const FString& InAttributeName = GetValue(Context, &Attribute);
+
+		FDataflowTransformSelection InSelection = GetValue(Context, &TransformSelection);
+		if (!IsConnected(&TransformSelection))
+		{
+			InSelection.InitFromArray(CollectionToUpdate, {});
+			InSelection.Invert();
+		}
+
+		const FString StringFormat = InTextToSet
+			.Replace(TEXT("{Current}"), TEXT("{0}"))
+			.Replace(TEXT("{Index}"), TEXT("{1}"));
+
+		// make sure we have the right format
+		// todo : add the ability to choose the attribute ?
+		TManagedArray<FString>* StringAttribute = CollectionToUpdate.FindAttributeTyped<FString>(FName(InAttributeName), FTransformCollection::TransformGroup);
+		if (StringAttribute)
+		{
+			const int32 Num = StringAttribute->Num();
+			int32 SelectionIndex = 0;
+			for (int32 Index = 0; Index < Num; Index++)
+			{
+				if (InSelection.IsSelected(Index))
+				{
+					const FString& CurrentValue = (*StringAttribute)[Index];
+					(*StringAttribute)[Index] = FString::Format(*StringFormat, { CurrentValue, SelectionIndex });
+					SelectionIndex++;
+				}
+			}
+		}
+		SetValue(Context, MoveTemp(CollectionToUpdate), &Collection);
+	}
 }
\ No newline at end of file
diff --git a/Engine/Plugins/Experimental/GeometryCollectionPlugin/Source/GeometryCollectionNodes/Public/Dataflow/GeometryCollectionSelectionNodes.h b/Engine/Plugins/Experimental/GeometryCollectionPlugin/Source/GeometryCollectionNodes/Public/Dataflow/GeometryCollectionSelectionNodes.h
--- a/Engine/Plugins/Experimental/GeometryCollectionPlugin/Source/GeometryCollectionNodes/Public/Dataflow/GeometryCollectionSelectionNodes.h
+++ b/Engine/Plugins/Experimental/GeometryCollectionPlugin/Source/GeometryCollectionNodes/Public/Dataflow/GeometryCollectionSelectionNodes.h
@@ -1677,6 +1677,92 @@ public:
 	FCollectionSelectInternalFacesDataflowNode(const UE::Dataflow::FNodeParameters& InParam, FGuid InGuid = FGuid::NewGuid());
 };
 
+UENUM()
+enum class EDataflowCollectionSelectionByNameMethod: uint8
+{
+	/** name must match exactly the input text */
+	Exact,
+	/** name must start with the input text */
+	StartsWith, 
+	/** name must end with the input text */
+	EndsWith,
+	/** name must contain the input text */
+	Contains,
+};
+
+/**
+ * Selects transforms by name using a the BoneName attributeor other Transform group string typed attributes 
+ */
+USTRUCT(meta = (DataflowGeometryCollection))
+struct FCollectionSelectTransformStringDataflowNode : public FDataflowNode
+{
+	GENERATED_USTRUCT_BODY()
+	DATAFLOW_NODE_DEFINE_INTERNAL(FCollectionSelectTransformStringDataflowNode, "CollectionSelectTransformString", "GeometryCollection|Selection|Transform", "Name, Bone, Attribute")
+	DATAFLOW_NODE_RENDER_TYPE("SurfaceRender", FGeometryCollection::StaticType(), "Collection")
+
+private:
+	/** Collection for the selection */
+	UPROPERTY(meta = (DataflowInput, DataflowOutput, DataflowPassthrough = "Collection", DataflowIntrinsic))
+	FManagedArrayCollection Collection;
+
+	/** Text to serach in the name */
+	UPROPERTY(EditAnywhere, Category = Options, meta = (DataflowInput))
+	FString Attribute = "BoneName";
+
+	/** Text to serach in the name */
+	UPROPERTY(EditAnywhere, Category = Options, meta=(DataflowInput))
+	FString SearchText;
+
+	/** Method to use to match the name */
+	UPROPERTY(EditAnywhere, Category = "Volume")
+	EDataflowCollectionSelectionByNameMethod Method = EDataflowCollectionSelectionByNameMethod::Contains;
+
+	/** output selection of the matching transforms */
+	UPROPERTY(meta = (DataflowOutput))
+	FDataflowTransformSelection TransformSelection;
+
+	virtual void Evaluate(UE::Dataflow::FContext& Context, const FDataflowOutput* Out) const override;
+
+public:
+	FCollectionSelectTransformStringDataflowNode(const UE::Dataflow::FNodeParameters& InParam, FGuid InGuid = FGuid::NewGuid());
+};
+
+/**
+ * Set selected transform string value
+ * the string format can use the following predefined value : 
+ * - {Current} current value of the attribute
+ * - {Index} index in the selection passed as input
+ */
+USTRUCT(meta = (DataflowGeometryCollection))
+struct FCollectionSetTransformStringValueDataflowNode : public FDataflowNode
+{
+	GENERATED_USTRUCT_BODY()
+	DATAFLOW_NODE_DEFINE_INTERNAL(FCollectionSetTransformStringValueDataflowNode, "CollectionSetTransformString", "GeometryCollection|Selection|Transform", "Bone Name Attribute")
+	DATAFLOW_NODE_RENDER_TYPE("SurfaceRender", FGeometryCollection::StaticType(), "Collection")
+
+private:
+	/** Collection for the selection */
+	UPROPERTY(meta = (DataflowInput, DataflowOutput, DataflowPassthrough = "Collection", DataflowIntrinsic))
+	FManagedArrayCollection Collection;
+
+	/** input selection of the transforms to set - if not connected use all */
+	UPROPERTY(meta = (DataflowInput))
+	FDataflowTransformSelection TransformSelection;
+
+	/** Text to serach in the name */
+	UPROPERTY(EditAnywhere, Category = Options, meta = (DataflowInput))
+	FString Attribute = "BoneName";
+
+	/** Text to set  */
+	UPROPERTY(EditAnywhere, Category = Options, meta = (DataflowInput))
+	FString TextToSet;
+
+	virtual void Evaluate(UE::Dataflow::FContext& Context, const FDataflowOutput* Out) const override;
+
+public:
+	FCollectionSetTransformStringValueDataflowNode(const UE::Dataflow::FNodeParameters& InParam, FGuid InGuid = FGuid::NewGuid());
+};
+
 namespace UE::Dataflow
 {
 	void GeometryCollectionSelectionNodes();
