Dataflow : fixed plane cutter node not properly accounting for the NumPlane input when connected

diff --git a/Engine/Plugins/Experimental/GeometryCollectionPlugin/Source/GeometryCollectionNodes/Private/Dataflow/GeometryCollectionFracturingNodes.cpp b/Engine/Plugins/Experimental/GeometryCollectionPlugin/Source/GeometryCollectionNodes/Private/Dataflow/GeometryCollectionFracturingNodes.cpp
--- a/Engine/Plugins/Experimental/GeometryCollectionPlugin/Source/GeometryCollectionNodes/Private/Dataflow/GeometryCollectionFracturingNodes.cpp
+++ b/Engine/Plugins/Experimental/GeometryCollectionPlugin/Source/GeometryCollectionNodes/Private/Dataflow/GeometryCollectionFracturingNodes.cpp
@@ -639,11 +639,11 @@ static FLinearColor GetRandomColor(const int32 RandomSeed, int32 Idx)
 
 void FPlaneCutterDataflowNode_v2::Evaluate(UE::Dataflow::FContext& Context, const FDataflowOutput* Out) const
 {
-	if (Out->IsA<FManagedArrayCollection>(&Collection) ||
-		Out->IsA<FDataflowTransformSelection>(&TransformSelection) ||
-		Out->IsA<FDataflowTransformSelection>(&NewGeometryTransformSelection))
+	if (Out->IsA(&Collection) ||
+		Out->IsA(&TransformSelection) ||
+		Out->IsA(&NewGeometryTransformSelection))
 	{
-		FManagedArrayCollection InCollection = GetValue<FManagedArrayCollection>(Context, &Collection);
+		FManagedArrayCollection InCollection = GetValue(Context, &Collection);
 		FBox InBoundingBox = GetValue(Context, &BoundingBox);
 
 		//
@@ -657,7 +657,7 @@ void FPlaneCutterDataflowNode_v2::Evaluate(UE::Dataflow::FContext& Context, cons
 			InBoundingBox = BoundingBoxInCollectionSpace;
 		}
 
-		FDataflowTransformSelection InTransformSelection = GetValue<FDataflowTransformSelection>(Context, &TransformSelection);
+		FDataflowTransformSelection InTransformSelection = GetValue(Context, &TransformSelection);
 
 		//
 		// If not connected select everything by default
@@ -674,13 +674,15 @@ void FPlaneCutterDataflowNode_v2::Evaluate(UE::Dataflow::FContext& Context, cons
 			InTransformSelection = NewTransformSelection;
 		}
 
+		const int32 InNumPlanes = GetValue(Context, &NumPlanes);
+
 		if (InTransformSelection.AnySelected())
 		{
 			int32 ResultGeometryIndex = FFractureEngineFracturing::PlaneCutter(InCollection,
 				InTransformSelection,
 				InBoundingBox,
 				GetValue(Context, &Transform),
-				NumPlanes,
+				InNumPlanes,
 				GetValue(Context, &RandomSeed),
 				GetValue(Context, &ChanceToFracture),
 				SplitIslands,
@@ -774,7 +776,7 @@ void FPlaneCutterDataflowNode_v2::DebugDraw(UE::Dataflow::FContext& Context, IDa
 
 		TArray<FTransform> PlaneTransforms;
 
-		FFractureEngineFracturing::GenerateSliceTransforms(InBoundingBox, InRandomSeed, NumPlanes, PlaneTransforms);
+		FFractureEngineFracturing::GenerateSliceTransforms(InBoundingBox, InRandomSeed, InNumPlanes, PlaneTransforms);
 
 		if (PlaneTransforms.IsEmpty())
 		{
