fixed crashes due to incorrect selection init for some dataflow fracture nodes

diff --git a/Engine/Plugins/Experimental/GeometryCollectionPlugin/Source/GeometryCollectionNodes/Private/Dataflow/GeometryCollectionFracturingNodes.cpp b/Engine/Plugins/Experimental/GeometryCollectionPlugin/Source/GeometryCollectionNodes/Private/Dataflow/GeometryCollectionFracturingNodes.cpp
--- a/Engine/Plugins/Experimental/GeometryCollectionPlugin/Source/GeometryCollectionNodes/Private/Dataflow/GeometryCollectionFracturingNodes.cpp
+++ b/Engine/Plugins/Experimental/GeometryCollectionPlugin/Source/GeometryCollectionNodes/Private/Dataflow/GeometryCollectionFracturingNodes.cpp
@@ -705,8 +705,9 @@ void FPlaneCutterDataflowNode_v2::Evaluate(UE::Dataflow::FContext& Context, cons
 				{
 					const TManagedArray<int32>& TransformIndices = InCollection.GetAttribute<int32>("TransformIndex", FGeometryCollection::GeometryGroup);
 
-					NewSelection.Initialize(TransformIndices.Num(), false);
-					OriginalSelection.Initialize(TransformIndices.Num(), false);
+					const int32 NumTransforms = InCollection.NumElements(FGeometryCollection::TransformGroup);
+					NewSelection.Initialize(NumTransforms, false);
+					OriginalSelection.Initialize(NumTransforms, false);
 
 					// The newly fractured pieces are added to the end of the transform array (starting position is ResultGeometryIndex)
 					for (int32 Idx = ResultGeometryIndex; Idx < TransformIndices.Num(); ++Idx)
@@ -929,8 +930,9 @@ void FSliceCutterDataflowNode::Evaluate(UE::Dataflow::FContext& Context, const F
 				{
 					const TManagedArray<int32>& TransformIndices = InCollection.GetAttribute<int32>("TransformIndex", FGeometryCollection::GeometryGroup);
 
-					NewSelection.Initialize(TransformIndices.Num(), false);
-					OriginalSelection.Initialize(TransformIndices.Num(), false);
+					const int32 NumTransforms = InCollection.NumElements(FGeometryCollection::TransformGroup);
+					NewSelection.Initialize(NumTransforms, false);
+					OriginalSelection.Initialize(NumTransforms, false);
 
 					// The newly fractured pieces are added to the end of the transform array (starting position is ResultGeometryIndex)
 					for (int32 Idx = ResultGeometryIndex; Idx < TransformIndices.Num(); ++Idx)
@@ -1034,8 +1036,9 @@ void FBrickCutterDataflowNode::Evaluate(UE::Dataflow::FContext& Context, const F
 				{
 					const TManagedArray<int32>& TransformIndices = InCollection.GetAttribute<int32>("TransformIndex", FGeometryCollection::GeometryGroup);
 
-					NewSelection.Initialize(TransformIndices.Num(), false);
-					OriginalSelection.Initialize(TransformIndices.Num(), false);
+					const int32 NumTransforms = InCollection.NumElements(FGeometryCollection::TransformGroup);
+					NewSelection.Initialize(NumTransforms, false);
+					OriginalSelection.Initialize(NumTransforms, false);
 
 					// The newly fractured pieces are added to the end of the transform array (starting position is ResultGeometryIndex)
 					for (int32 Idx = ResultGeometryIndex; Idx < TransformIndices.Num(); ++Idx)
@@ -1210,8 +1213,9 @@ void FMeshCutterDataflowNode::Evaluate(UE::Dataflow::FContext& Context, const FD
 					{
 						const TManagedArray<int32>& TransformIndices = InCollection.GetAttribute<int32>("TransformIndex", FGeometryCollection::GeometryGroup);
 
-						NewSelection.Initialize(TransformIndices.Num(), false);
-						OriginalSelection.Initialize(TransformIndices.Num(), false);
+						const int32 NumTransforms = InCollection.NumElements(FGeometryCollection::TransformGroup);
+						NewSelection.Initialize(NumTransforms, false);
+						OriginalSelection.Initialize(NumTransforms, false);
 
 						// The newly fractured pieces are added to the end of the transform array (starting position is ResultGeometryIndex)
 						for (int32 Idx = ResultGeometryIndex; Idx < TransformIndices.Num(); ++Idx)
