Code improvements from code submiited at CL 43675006

diff --git a/Engine/Plugins/Experimental/GeometryCollectionPlugin/Source/GeometryCollectionEditor/Private/GeometryCollection/AssetDefinition_GeometryCollection.cpp b/Engine/Plugins/Experimental/GeometryCollectionPlugin/Source/GeometryCollectionEditor/Private/GeometryCollection/AssetDefinition_GeometryCollection.cpp
--- a/Engine/Plugins/Experimental/GeometryCollectionPlugin/Source/GeometryCollectionEditor/Private/GeometryCollection/AssetDefinition_GeometryCollection.cpp
+++ b/Engine/Plugins/Experimental/GeometryCollectionPlugin/Source/GeometryCollectionEditor/Private/GeometryCollection/AssetDefinition_GeometryCollection.cpp
@@ -136,9 +136,11 @@ namespace MenuExtension_GeometryCollection
 
 		// Create a frame in the package
 		T* NewAsset = NewObject<T>(OuterForAsset, T::StaticClass(), *UniqueAssetName, RF_Public | RF_Standalone | RF_Transactional);
-		FAssetRegistryModule::AssetCreated(NewAsset);
-
-		NewAsset->Modify();
+		if (NewAsset)
+		{
+			FAssetRegistryModule::AssetCreated(NewAsset);
+			NewAsset->Modify();
+		}
 		return NewAsset;
 	}
 
@@ -156,34 +158,36 @@ namespace MenuExtension_GeometryCollection
 					if (GeometryCollection)
 					{
 						// find path to reference asset
-						UPackage* AssetOuterPackage = CastChecked<UPackage>(GeometryCollection->GetOuter());
-						FString AssetPackageName = AssetOuterPackage->GetName();
-						FString AssetFolderPath = FPackageName::GetLongPackagePath(AssetPackageName);
-						FString AssetBaseName = FPaths::GetBaseFilename(AssetPackageName);
-						AssetBaseName.RemoveFromStart(TEXT("GC_"));
+						if (UPackage* AssetOuterPackage = Cast<UPackage>(GeometryCollection->GetOuter()))
+						{
+							FString AssetPackageName = AssetOuterPackage->GetName();
+							FString AssetFolderPath = FPackageName::GetLongPackagePath(AssetPackageName);
+							FString AssetBaseName = FPaths::GetBaseFilename(AssetPackageName);
+							AssetBaseName.RemoveFromStart(TEXT("GC_"));
 
-						FString SkeletalMeshAssetName = FString::Printf(TEXT("SK_%s"), *AssetBaseName);
-						FString SkeletonAssetName = FString::Printf(TEXT("SKEL_%s"), *AssetBaseName);
+							FString SkeletalMeshAssetName = FString::Printf(TEXT("SK_%s"), *AssetBaseName);
+							FString SkeletonAssetName = FString::Printf(TEXT("SKEL_%s"), *AssetBaseName);
 
-						USkeletalMesh* SkeletalMesh = GetOrCreateAsset<USkeletalMesh>(AssetFolderPath, SkeletalMeshAssetName);
-						USkeleton* Skeleton = GetOrCreateAsset<USkeleton>(AssetFolderPath, SkeletonAssetName);
+							USkeletalMesh* SkeletalMesh = GetOrCreateAsset<USkeletalMesh>(AssetFolderPath, SkeletalMeshAssetName);
+							USkeleton* Skeleton = GetOrCreateAsset<USkeleton>(AssetFolderPath, SkeletonAssetName);
 
-						if (SkeletalMesh && Skeleton)
-						{
-
-							if (GeometryCollection->GetGeometryCollection())
+							if (SkeletalMesh && Skeleton)
 							{
-								FGeometryCollectionEngineConversion::ConvertCollectionToSkeletalMesh(*GeometryCollection->GetGeometryCollection(), GeometryCollection->Materials, *SkeletalMesh, *Skeleton);
-							}
 
-							SkeletalMesh->PostEditChange();
-							Skeleton->PostEditChange();
+								if (GeometryCollection->GetGeometryCollection())
+								{
+									FGeometryCollectionEngineConversion::ConvertCollectionToSkeletalMesh(*GeometryCollection->GetGeometryCollection(), GeometryCollection->Materials, *SkeletalMesh, *Skeleton);
+								}
 
-							SkeletalMesh->MarkPackageDirty();
-							Skeleton->MarkPackageDirty();
+								SkeletalMesh->PostEditChange();
+								Skeleton->PostEditChange();
 
-							FAssetRegistryModule::AssetCreated(SkeletalMesh);
-							FAssetRegistryModule::AssetCreated(Skeleton);
+								SkeletalMesh->MarkPackageDirty();
+								Skeleton->MarkPackageDirty();
+
+								FAssetRegistryModule::AssetCreated(SkeletalMesh);
+								FAssetRegistryModule::AssetCreated(Skeleton);
+							}
 						}
 					}
 				}
diff --git a/Engine/Plugins/Experimental/GeometryCollectionPlugin/Source/GeometryCollectionNodes/Private/Dataflow/GeometryCollectionSkeletalMeshToCollectionNode.cpp b/Engine/Plugins/Experimental/GeometryCollectionPlugin/Source/GeometryCollectionNodes/Private/Dataflow/GeometryCollectionSkeletalMeshToCollectionNode.cpp
--- a/Engine/Plugins/Experimental/GeometryCollectionPlugin/Source/GeometryCollectionNodes/Private/Dataflow/GeometryCollectionSkeletalMeshToCollectionNode.cpp
+++ b/Engine/Plugins/Experimental/GeometryCollectionPlugin/Source/GeometryCollectionNodes/Private/Dataflow/GeometryCollectionSkeletalMeshToCollectionNode.cpp
@@ -48,7 +48,7 @@ if (Out->IsA(&SkeletalMesh) || Out->IsA(&Skeleton))
 		USkeletalMesh* OutSkeletalMesh = NewObject<USkeletalMesh>();
 		USkeleton* OutSkeleton = NewObject<USkeleton>();
 		
-		if (OutSkeletalMesh || OutSkeleton)
+		if (OutSkeletalMesh && OutSkeleton)
 		{
 			// make a copy so that we can ensure the level array is created
 			FManagedArrayCollection InCollection = GetValue(Context, &Collection);
diff --git a/Engine/Source/Runtime/Experimental/GeometryCollectionEngine/Private/GeometryCollection/GeometryCollectionEngineConversion.cpp b/Engine/Source/Runtime/Experimental/GeometryCollectionEngine/Private/GeometryCollection/GeometryCollectionEngineConversion.cpp
--- a/Engine/Source/Runtime/Experimental/GeometryCollectionEngine/Private/GeometryCollection/GeometryCollectionEngineConversion.cpp
+++ b/Engine/Source/Runtime/Experimental/GeometryCollectionEngine/Private/GeometryCollection/GeometryCollectionEngineConversion.cpp
@@ -2089,18 +2089,21 @@ void FGeometryCollectionEngineConversion::ConvertCollectionToSkeletalMesh(const
 		// find corresponding bone from skeleton 
 		const FName BoneName = FName((*Names)[TransformIndex]);
 		const int32 BoneIndex = OutSkeleton.GetReferenceSkeleton().FindBoneIndex(BoneName);
-		const int32 ParentBoneIndex = (BoneIndex != INDEX_NONE)? OutSkeleton.GetReferenceSkeleton().GetParentIndex(BoneIndex): INDEX_NONE;
+		if (BoneIndex != INDEX_NONE)
+		{
+			const int32 ParentBoneIndex = OutSkeleton.GetReferenceSkeleton().GetParentIndex(BoneIndex);
 
-		const TArray<FTransform> BoneTransforms = OutSkeleton.GetReferenceSkeleton().GetRefBonePose();
-		const FTransform BoneTransform = (BoneTransforms.IsValidIndex(BoneIndex)) ? BoneTransforms[BoneIndex] : FTransform::Identity;
+			const TArray<FTransform> BoneTransforms = OutSkeleton.GetReferenceSkeleton().GetRefBonePose();
+			const FTransform BoneTransform = (BoneTransforms.IsValidIndex(BoneIndex)) ? BoneTransforms[BoneIndex] : FTransform::Identity;
 
-		const FBoneID BoneID(BoneIndex);
+			const FBoneID BoneID(BoneIndex);
 
-		Attributes.CreateBone(BoneIndex);
-		{
-			BoneNames.Set(BoneID, BoneName);
-			BoneParentIndices.Set(BoneID, ParentBoneIndex);
-			BonePoses.Set(BoneID, BoneTransform);
+			Attributes.CreateBone(BoneIndex);
+			{
+				BoneNames.Set(BoneID, BoneName);
+				BoneParentIndices.Set(BoneID, ParentBoneIndex);
+				BonePoses.Set(BoneID, BoneTransform);
+			}
 		}
 	}
 
