Dataflow : Expose Geometry script variables as input to the mesh porocessing nodes
- Updated Dynamic connection class to allow to be used along side of other connections

diff --git a/Engine/Plugins/Experimental/GeometryCollectionPlugin/Source/GeometryCollectionNodes/Private/Dataflow/GeometryCollectionMeshNodes.cpp b/Engine/Plugins/Experimental/GeometryCollectionPlugin/Source/GeometryCollectionNodes/Private/Dataflow/GeometryCollectionMeshNodes.cpp
--- a/Engine/Plugins/Experimental/GeometryCollectionPlugin/Source/GeometryCollectionNodes/Private/Dataflow/GeometryCollectionMeshNodes.cpp
+++ b/Engine/Plugins/Experimental/GeometryCollectionPlugin/Source/GeometryCollectionNodes/Private/Dataflow/GeometryCollectionMeshNodes.cpp
@@ -41,6 +41,8 @@
 
 #include UE_INLINE_GENERATED_CPP_BY_NAME(GeometryCollectionMeshNodes)
 
+DEFINE_LOG_CATEGORY_STATIC(LogGeometryCollectionMeshNodes, Warning, All);
+
 namespace UE::Dataflow
 {
 
@@ -783,20 +785,33 @@ void FGetMeshDataDataflowNode::Evaluate(UE::Dataflow::FContext& Context, const F
 	}
 }
 
+/////////////////////////////////////////////////////////////////////////////////////////////////////////////
+
+void FMeshProcessorDataflowNodeBase::PostSerialize(const FArchive& Ar)
+{
+	Super::PostSerialize(Ar);
+	if (Ar.IsLoading())
+	{
+		TeardownBlueprintEvent();
+		SetupBlueprintEvent();
+	}
+}
+
 void FMeshProcessorDataflowNodeBase::OnPropertyChanged(UE::Dataflow::FContext& Context, const FPropertyChangedEvent& PropertyChangedEvent)
 {
 	if (PropertyChangedEvent.ChangeType == EPropertyChangeType::ValueSet 
 		&& PropertyChangedEvent.GetPropertyName() == GET_MEMBER_NAME_CHECKED(FMeshProcessorDataflowNodeBase, MeshProcessor))
 	{
+		TeardownBlueprintEvent();
 		if (MeshProcessor)
 		{
 			MeshProcessorInstance = NewObject<UDynamicMeshProcessorBlueprint>(OwningObject, MeshProcessor, NAME_None, RF_Transactional);
-			TeardownBlueprintEvent();
 			SetupBlueprintEvent();
 		}
 		else
 		{
 			MeshProcessorInstance = nullptr;
+			RefreshConnectionsFromBlueprint();
 		}
 	}
 }
@@ -814,10 +829,12 @@ void FMeshProcessorDataflowNodeBase::SetupBlueprintEvent()
 			}
 			BlueprintChangeDelegateHandle = Blueprint->OnChanged().AddLambda([this](UBlueprint* BP)
 			{
+				RefreshConnectionsFromBlueprint();
 				Invalidate();
 			});
 		}
 	}
+	RefreshConnectionsFromBlueprint();
 #endif
 }
 
@@ -835,6 +852,71 @@ void FMeshProcessorDataflowNodeBase::TeardownBlueprintEvent()
 #endif
 }
 
+void FMeshProcessorDataflowNodeBase::RefreshConnectionsFromBlueprint()
+{
+#if WITH_EDITOR
+	PropertyBag.Reset();
+	if (MeshProcessor)
+	{
+		if (UBlueprint* Blueprint = Cast<UBlueprint>(MeshProcessor->ClassGeneratedBy))
+		{
+			if (Blueprint->GeneratedClass != nullptr)
+			{
+				for (TFieldIterator<FProperty> PropertyIt(Blueprint->GeneratedClass, EFieldIteratorFlags::IncludeSuper); PropertyIt; ++PropertyIt)
+				{
+					FProperty* Property = *PropertyIt;
+					if (!Property->HasAnyPropertyFlags(CPF_Parm) && Property->HasAllPropertyFlags(CPF_BlueprintVisible))
+					{
+						PropertyBag.AddProperty(Property->GetFName(), Property, false);
+					}
+				}
+			}
+		}
+	}
+	// Make sure the node is up to date with the serialized data
+	DynamicConnections.Refresh();
+#endif
+}
+
+void FMeshProcessorDataflowNodeBase::ApplyParametersToBlueprintInstance(UE::Dataflow::FContext& Context) const
+{
+#if WITH_EDITOR
+	if (UBlueprint* Blueprint = Cast<UBlueprint>(MeshProcessor->ClassGeneratedBy))
+	{
+		if (Blueprint->GeneratedClass != nullptr && MeshProcessorInstance)
+		{
+			TConstArrayView<FPropertyBagPropertyDesc> Descs = PropertyBag.GetPropertyBagStruct()->GetPropertyDescs();
+			for (const FPropertyBagPropertyDesc& Desc : Descs)
+			{
+				if (FProperty* BPProperty = Blueprint->GeneratedClass->FindPropertyByName(Desc.Name))
+				{
+					if (BPProperty->GetClass() == Desc.CachedProperty->GetClass())
+					{
+						if (const FDataflowInput* Input = FindInput(Desc.Name))
+						{
+							// only send the data is the input is connected, other wise keep what was there before 
+							if (const FDataflowOutput* Output = Input->GetConnection())
+							{
+								Input->PullValue(Context);
+								if (const void* CachedValueAddress = Context.GetUntypedData(Output->CacheKey(), nullptr))
+								{
+									void* TargetAddress = BPProperty->ContainerPtrToValuePtr<void>(MeshProcessorInstance);
+									BPProperty->CopyCompleteValue(TargetAddress, CachedValueAddress);
+								}
+							}
+						}
+					}
+				}
+			}
+		}
+	}
+#else
+	UE_LOG(LogGeometryCollectionMeshNodes, Error, TEXT("FMeshProcessorDataflowNodeBase - trying to execute a mesh processing node ( running a Geometry script ) in a non-editor build - connected input parameters won't be passed to the script"));
+#endif
+}
+
+///////////////////////////////////////////////////////////////////////////////////////////////////////////
+
 void FApplyMeshProcessorToMeshDataflowNode::Evaluate(UE::Dataflow::FContext& Context, const FDataflowOutput* Out) const
 {
 	if (Out->IsA<TObjectPtr<UDynamicMesh>>(&Mesh))
@@ -851,6 +933,8 @@ void FApplyMeshProcessorToMeshDataflowNode::Evaluate(UE::Dataflow::FContext& Con
 			TObjectPtr<UDynamicMesh> NewMesh = NewObject<UDynamicMesh>();
 			NewMesh->SetMesh(InMesh->GetMeshRef());
 
+			ApplyParametersToBlueprintInstance(Context);
+
 			bool bFailed = false;
 			MeshProcessorInstance->ProcessDynamicMesh(NewMesh, bFailed);
 
@@ -913,6 +997,8 @@ void FApplyMeshProcessorToGeometryCollectionDataflowNode::Evaluate(UE::Dataflow:
 				{
 					NewMesh->SetMesh(MoveTemp(*MeshInfo.Mesh));
 
+					ApplyParametersToBlueprintInstance(Context);
+
 					bool bFailed = false;
 					MeshProcessorInstance->ProcessDynamicMesh(NewMesh, bFailed);
 					if (!bFailed) // on success, move the mesh back
diff --git a/Engine/Plugins/Experimental/GeometryCollectionPlugin/Source/GeometryCollectionNodes/Public/Dataflow/GeometryCollectionMeshNodes.h b/Engine/Plugins/Experimental/GeometryCollectionPlugin/Source/GeometryCollectionNodes/Public/Dataflow/GeometryCollectionMeshNodes.h
--- a/Engine/Plugins/Experimental/GeometryCollectionPlugin/Source/GeometryCollectionNodes/Public/Dataflow/GeometryCollectionMeshNodes.h
+++ b/Engine/Plugins/Experimental/GeometryCollectionPlugin/Source/GeometryCollectionNodes/Public/Dataflow/GeometryCollectionMeshNodes.h
@@ -5,6 +5,7 @@
 #include "CoreMinimal.h"
 #include "Dataflow/DataflowEngine.h"
 #include "Dataflow/DataflowMesh.h"
+#include "Dataflow/DataflowDynamicConnections.h"
 #include "GeometryCollection/ManagedArrayCollection.h"
 #include "Dataflow/DataflowSelection.h"
 #include "DynamicMesh/DynamicMesh3.h"
@@ -69,7 +70,11 @@ public:
  * Base class for nodes that applying Geometry Script Mesh Processors
  */
 USTRUCT()
-struct FMeshProcessorDataflowNodeBase : public FDataflowNode
+struct FMeshProcessorDataflowNodeBase 
+	: public FDataflowNode
+#if CPP
+	, public FDataflowDynamicConnections::IOwnerInterface
+#endif
 {
 	GENERATED_USTRUCT_BODY()
 
@@ -79,6 +84,7 @@ public:
 
 	FMeshProcessorDataflowNodeBase(const UE::Dataflow::FNodeParameters& InParam, FGuid InGuid = FGuid::NewGuid())
 		: FDataflowNode(InParam, InGuid), OwningObject(InParam.OwningObject)
+		, DynamicConnections(UE::Dataflow::FPin::EDirection::INPUT, this, Cast<UDataflow>(InParam.OwningObject))
 	{
 	}
 
@@ -95,18 +101,17 @@ protected:
 	UPROPERTY(EditAnywhere, Instanced, NoClear, Category = "Instance", meta = (ShowOnlyInnerProperties))
 	TObjectPtr<UDynamicMeshProcessorBlueprint> MeshProcessorInstance;
 
-	virtual void PostSerialize(const FArchive& Ar) override final
-	{
-		Super::PostSerialize(Ar);
-		if (Ar.IsLoading())
-		{
-			TeardownBlueprintEvent();
-			SetupBlueprintEvent();
-		}
-	}
+	virtual void PostSerialize(const FArchive& Ar) override final;
+
+	void ApplyParametersToBlueprintInstance(UE::Dataflow::FContext& Context) const;
 
 private:
 
+	/* Begin IOwnerInterface */
+	virtual FDataflowNode* GetOwner(const FDataflowDynamicConnections* Caller) override { return this; }
+	virtual const FInstancedPropertyBag& GetPropertyBag(const FDataflowDynamicConnections* Caller) override { return PropertyBag; }
+	/* End IOwnerInterface */
+
 	virtual void OnPropertyChanged(UE::Dataflow::FContext& Context, const FPropertyChangedEvent& PropertyChangedEvent) override;
 
 	// Blueprint evaluation is not generally thread-safe
@@ -119,10 +124,17 @@ private:
 	FDelegateHandle BlueprintChangeDelegateHandle;
 	void TeardownBlueprintEvent();
 	void SetupBlueprintEvent();
+	void RefreshConnectionsFromBlueprint();
 
 	// Remember the parent UObject so that we can properly parent MeshProcessorInstance when the MeshProcessor changes 
 	UPROPERTY(Transient)
 	TObjectPtr<UObject> OwningObject;
+
+	UPROPERTY()
+	FDataflowDynamicConnections DynamicConnections;
+
+	UPROPERTY()
+	FInstancedPropertyBag PropertyBag;
 };
 
 /**
diff --git a/Engine/Source/Runtime/Experimental/Dataflow/Engine/Private/Dataflow/DataflowDynamicConnections.cpp b/Engine/Source/Runtime/Experimental/Dataflow/Engine/Private/Dataflow/DataflowDynamicConnections.cpp
--- a/Engine/Source/Runtime/Experimental/Dataflow/Engine/Private/Dataflow/DataflowDynamicConnections.cpp
+++ b/Engine/Source/Runtime/Experimental/Dataflow/Engine/Private/Dataflow/DataflowDynamicConnections.cpp
@@ -49,6 +49,8 @@ void FDataflowDynamicConnections::Refresh()
 		return;
 	}
 
+	OwnerNode->PauseInvalidations();
+
 	// save the remote connections from the dynamic connections
 	using FRemoteConnections = TArray<FDataflowConnection*>;
 	TMap<FGuid, FRemoteConnections> RemoteConnectionsByPropertyId;
@@ -69,8 +71,7 @@ void FDataflowDynamicConnections::Refresh()
 	// clear all the dynamic connections
 	ClearDynamicConnections();
 
-	ConnectionNameToPropertyId.Reset();
-	DynamicProperties.Reset();
+	OwnerNode->ResumeInvalidations();
 
 	// Go through all the properties of the property bag recreate the output and reconnect them if needed
 	const TConnectionReference<FDataflowAllTypes> NullReference(nullptr);
@@ -96,12 +97,13 @@ void FDataflowDynamicConnections::Refresh()
 				}
 			}
 		}
-		// Todo : need to also remove the one that are no longer necessary
-		// ( need to disconnect if needed ? or keep the one that are still connected ?)
-		OwnerNode->Invalidate();
-
-		DataflowAsset->RefreshEdNodeByGuid(OwnerNode->GetGuid());
 	}
+
+	// Todo : need to also remove the one that are no longer necessary
+	// ( need to disconnect if needed ? or keep the one that are still connected ?)
+	OwnerNode->Invalidate();
+
+	DataflowAsset->RefreshEdNodeByGuid(OwnerNode->GetGuid());
 }
 
 FDataflowConnection* FDataflowDynamicConnections::AddNewConnectionFromPropertyDesc(const FPropertyBagPropertyDesc& PropertyDesc)
@@ -247,13 +249,18 @@ void FDataflowDynamicConnections::ClearDynamicConnections()
 {
 	if (FDataflowNode* OwnerNode = OwnerInterface ? OwnerInterface->GetOwner(this) : nullptr)
 	{
-		if (PinDirection == UE::Dataflow::FPin::EDirection::INPUT)
-		{
-			OwnerNode->ClearInputs();
-		}
-		else if (PinDirection == UE::Dataflow::FPin::EDirection::OUTPUT)
+		for (int32 ConnectionIndex = 0; ConnectionIndex < DynamicProperties.Num(); ++ConnectionIndex)
 		{
-			OwnerNode->ClearOutputs();
+			if (PinDirection == UE::Dataflow::FPin::EDirection::INPUT)
+			{
+				OwnerNode->UnregisterInputConnection(GetConnectionReference(ConnectionIndex));
+			}
+			else if (PinDirection == UE::Dataflow::FPin::EDirection::OUTPUT)
+			{
+				OwnerNode->UnregisterOutputConnection(GetConnectionReference(ConnectionIndex));
+			}
 		}
 	}
+	ConnectionNameToPropertyId.Reset();
+	DynamicProperties.Reset();
 }
\ No newline at end of file
diff --git a/Engine/Source/Runtime/Experimental/Dataflow/Engine/Public/Dataflow/DataflowDynamicConnections.h b/Engine/Source/Runtime/Experimental/Dataflow/Engine/Public/Dataflow/DataflowDynamicConnections.h
--- a/Engine/Source/Runtime/Experimental/Dataflow/Engine/Public/Dataflow/DataflowDynamicConnections.h
+++ b/Engine/Source/Runtime/Experimental/Dataflow/Engine/Public/Dataflow/DataflowDynamicConnections.h
@@ -26,10 +26,10 @@ struct FDataflowDynamicConnections
 	};
 
 public:
-	FDataflowDynamicConnections();
-	FDataflowDynamicConnections(UE::Dataflow::FPin::EDirection PInDirection, IOwnerInterface* OwnerInterface, UDataflow* DataflowAsset);
+	DATAFLOWENGINE_API FDataflowDynamicConnections();
+	DATAFLOWENGINE_API FDataflowDynamicConnections(UE::Dataflow::FPin::EDirection PInDirection, IOwnerInterface* OwnerInterface, UDataflow* DataflowAsset);
 
-	void Refresh();
+	DATAFLOWENGINE_API void Refresh();
 
 private:
 	using FConnectionReference = UE::Dataflow::TConnectionReference<FDataflowAllTypes>;
