Transition dataflow selection Initialize usage to the new InitializeFromCollection method where applicable + minor related cleanup

diff --git a/Engine/Plugins/Experimental/GeometryCollectionPlugin/Source/GeometryCollectionNodes/Private/Dataflow/GeometryCollectionFracturingNodes.cpp b/Engine/Plugins/Experimental/GeometryCollectionPlugin/Source/GeometryCollectionNodes/Private/Dataflow/GeometryCollectionFracturingNodes.cpp
--- a/Engine/Plugins/Experimental/GeometryCollectionPlugin/Source/GeometryCollectionNodes/Private/Dataflow/GeometryCollectionFracturingNodes.cpp
+++ b/Engine/Plugins/Experimental/GeometryCollectionPlugin/Source/GeometryCollectionNodes/Private/Dataflow/GeometryCollectionFracturingNodes.cpp
@@ -513,7 +513,7 @@ void FVoronoiFractureDataflowNode_v2::Evaluate(UE::Dataflow::FContext& Context,
 			const TArray<int32>& SelectionArr = TransformSelectionFacade.SelectAll();
 
 			FDataflowTransformSelection NewTransformSelection;
-			NewTransformSelection.Initialize(InCollection.NumElements(FGeometryCollection::TransformGroup), false);
+			NewTransformSelection.InitializeFromCollection(InCollection, false);
 			NewTransformSelection.SetFromArray(SelectionArr);
 
 			InTransformSelection = NewTransformSelection;
@@ -549,9 +549,8 @@ void FVoronoiFractureDataflowNode_v2::Evaluate(UE::Dataflow::FContext& Context,
 				{
 					const TManagedArray<int32>& GeometryToTransformIndices = InCollection.GetAttribute<int32>("TransformIndex", FGeometryCollection::GeometryGroup);
 
-					const int32 NumTransforms = InCollection.NumElements(FGeometryCollection::TransformGroup);
-					NewSelection.Initialize(NumTransforms, false);
-					OriginalSelection.Initialize(NumTransforms, false);
+					NewSelection.InitializeFromCollection(InCollection, false);
+					OriginalSelection.InitializeFromCollection(InCollection, false);
 
 					// The newly fractured pieces are added to the end of the transform array (starting position is ResultGeometryIndex)
 					for (int32 GeometryIdx = ResultGeometryIndex; GeometryIdx < GeometryToTransformIndices.Num(); ++GeometryIdx)
@@ -668,7 +667,7 @@ void FPlaneCutterDataflowNode_v2::Evaluate(UE::Dataflow::FContext& Context, cons
 			const TArray<int32>& SelectionArr = TransformSelectionFacade.SelectAll();
 
 			FDataflowTransformSelection NewTransformSelection;
-			NewTransformSelection.Initialize(InCollection.NumElements(FGeometryCollection::TransformGroup), false);
+			NewTransformSelection.InitializeFromCollection(InCollection, false);
 			NewTransformSelection.SetFromArray(SelectionArr);
 
 			InTransformSelection = NewTransformSelection;
@@ -705,9 +704,8 @@ void FPlaneCutterDataflowNode_v2::Evaluate(UE::Dataflow::FContext& Context, cons
 				{
 					const TManagedArray<int32>& TransformIndices = InCollection.GetAttribute<int32>("TransformIndex", FGeometryCollection::GeometryGroup);
 
-					const int32 NumTransforms = InCollection.NumElements(FGeometryCollection::TransformGroup);
-					NewSelection.Initialize(NumTransforms, false);
-					OriginalSelection.Initialize(NumTransforms, false);
+					NewSelection.InitializeFromCollection(InCollection, false);
+					OriginalSelection.InitializeFromCollection(InCollection, false);
 
 					// The newly fractured pieces are added to the end of the transform array (starting position is ResultGeometryIndex)
 					for (int32 Idx = ResultGeometryIndex; Idx < TransformIndices.Num(); ++Idx)
@@ -890,7 +888,7 @@ void FSliceCutterDataflowNode::Evaluate(UE::Dataflow::FContext& Context, const F
 			const TArray<int32>& SelectionArr = TransformSelectionFacade.SelectAll();
 
 			FDataflowTransformSelection NewTransformSelection;
-			NewTransformSelection.Initialize(InCollection.NumElements(FGeometryCollection::TransformGroup), false);
+			NewTransformSelection.InitializeFromCollection(InCollection, false);
 			NewTransformSelection.SetFromArray(SelectionArr);
 
 			InTransformSelection = NewTransformSelection;
@@ -930,9 +928,8 @@ void FSliceCutterDataflowNode::Evaluate(UE::Dataflow::FContext& Context, const F
 				{
 					const TManagedArray<int32>& TransformIndices = InCollection.GetAttribute<int32>("TransformIndex", FGeometryCollection::GeometryGroup);
 
-					const int32 NumTransforms = InCollection.NumElements(FGeometryCollection::TransformGroup);
-					NewSelection.Initialize(NumTransforms, false);
-					OriginalSelection.Initialize(NumTransforms, false);
+					NewSelection.InitializeFromCollection(InCollection, false);
+					OriginalSelection.InitializeFromCollection(InCollection, false);
 
 					// The newly fractured pieces are added to the end of the transform array (starting position is ResultGeometryIndex)
 					for (int32 Idx = ResultGeometryIndex; Idx < TransformIndices.Num(); ++Idx)
@@ -982,7 +979,7 @@ void FBrickCutterDataflowNode::Evaluate(UE::Dataflow::FContext& Context, const F
 			const TArray<int32>& SelectionArr = TransformSelectionFacade.SelectAll();
 
 			FDataflowTransformSelection NewTransformSelection;
-			NewTransformSelection.Initialize(InCollection.NumElements(FGeometryCollection::TransformGroup), false);
+			NewTransformSelection.InitializeFromCollection(InCollection, false);
 			NewTransformSelection.SetFromArray(SelectionArr);
 
 			InTransformSelection = NewTransformSelection;
@@ -1036,9 +1033,8 @@ void FBrickCutterDataflowNode::Evaluate(UE::Dataflow::FContext& Context, const F
 				{
 					const TManagedArray<int32>& TransformIndices = InCollection.GetAttribute<int32>("TransformIndex", FGeometryCollection::GeometryGroup);
 
-					const int32 NumTransforms = InCollection.NumElements(FGeometryCollection::TransformGroup);
-					NewSelection.Initialize(NumTransforms, false);
-					OriginalSelection.Initialize(NumTransforms, false);
+					NewSelection.InitializeFromCollection(InCollection, false);
+					OriginalSelection.InitializeFromCollection(InCollection, false);
 
 					// The newly fractured pieces are added to the end of the transform array (starting position is ResultGeometryIndex)
 					for (int32 Idx = ResultGeometryIndex; Idx < TransformIndices.Num(); ++Idx)
@@ -1088,7 +1084,7 @@ void FMeshCutterDataflowNode::Evaluate(UE::Dataflow::FContext& Context, const FD
 			const TArray<int32>& SelectionArr = TransformSelectionFacade.SelectAll();
 
 			FDataflowTransformSelection NewTransformSelection;
-			NewTransformSelection.Initialize(InCollection.NumElements(FGeometryCollection::TransformGroup), false);
+			NewTransformSelection.InitializeFromCollection(InCollection, false);
 			NewTransformSelection.SetFromArray(SelectionArr);
 
 			InTransformSelection = NewTransformSelection;
@@ -1213,9 +1209,8 @@ void FMeshCutterDataflowNode::Evaluate(UE::Dataflow::FContext& Context, const FD
 					{
 						const TManagedArray<int32>& TransformIndices = InCollection.GetAttribute<int32>("TransformIndex", FGeometryCollection::GeometryGroup);
 
-						const int32 NumTransforms = InCollection.NumElements(FGeometryCollection::TransformGroup);
-						NewSelection.Initialize(NumTransforms, false);
-						OriginalSelection.Initialize(NumTransforms, false);
+						NewSelection.InitializeFromCollection(InCollection, false);
+						OriginalSelection.InitializeFromCollection(InCollection, false);
 
 						// The newly fractured pieces are added to the end of the transform array (starting position is ResultGeometryIndex)
 						for (int32 Idx = ResultGeometryIndex; Idx < TransformIndices.Num(); ++Idx)
@@ -1315,7 +1310,7 @@ void FUniformFractureDataflowNode::Evaluate(UE::Dataflow::FContext& Context, con
 			const TArray<int32>& SelectionArr = TransformSelectionFacade.SelectAll();
 
 			FDataflowTransformSelection NewTransformSelection;
-			NewTransformSelection.Initialize(InCollection.NumElements(FGeometryCollection::TransformGroup), false);
+			NewTransformSelection.InitializeFromCollection(InCollection, false);
 			NewTransformSelection.SetFromArray(SelectionArr);
 
 			InTransformSelection = NewTransformSelection;
@@ -1358,9 +1353,8 @@ void FUniformFractureDataflowNode::Evaluate(UE::Dataflow::FContext& Context, con
 				{
 					const TManagedArray<int32>& GeometryToTransformIndices = InCollection.GetAttribute<int32>("TransformIndex", FGeometryCollection::GeometryGroup);
 
-					const int32 NumTransforms = InCollection.NumElements(FGeometryCollection::TransformGroup);
-					NewSelection.Initialize(NumTransforms, false);
-					OriginalSelection.Initialize(NumTransforms, false);
+					NewSelection.InitializeFromCollection(InCollection, false);
+					OriginalSelection.InitializeFromCollection(InCollection, false);
 
 					// The newly fractured pieces are added to the end of the transform array (starting position is ResultGeometryIndex)
 					for (int32 GeometryIdx = ResultGeometryIndex; GeometryIdx < GeometryToTransformIndices.Num(); ++GeometryIdx)
diff --git a/Engine/Plugins/Experimental/GeometryCollectionPlugin/Source/GeometryCollectionNodes/Private/Dataflow/GeometryCollectionMaterialInterfaceNodes.cpp b/Engine/Plugins/Experimental/GeometryCollectionPlugin/Source/GeometryCollectionNodes/Private/Dataflow/GeometryCollectionMaterialInterfaceNodes.cpp
--- a/Engine/Plugins/Experimental/GeometryCollectionPlugin/Source/GeometryCollectionNodes/Private/Dataflow/GeometryCollectionMaterialInterfaceNodes.cpp
+++ b/Engine/Plugins/Experimental/GeometryCollectionPlugin/Source/GeometryCollectionNodes/Private/Dataflow/GeometryCollectionMaterialInterfaceNodes.cpp
@@ -303,7 +303,7 @@ void FAssignMaterialInterfaceToCollectionDataflowNode::Evaluate(UE::Dataflow::FC
 		if (!IsConnected(&FaceSelection))
 		{
 			// if not connected select all faces
-			InFaceSelection.Initialize(NumFaces, true);
+			InFaceSelection.InitializeFromCollection(InCollection, true);
 		}
 
 		static const FName MaterialIdAttributeName = "MaterialID";
diff --git a/Engine/Plugins/Experimental/GeometryCollectionPlugin/Source/GeometryCollectionNodes/Private/Dataflow/GeometryCollectionMeshNodes.cpp b/Engine/Plugins/Experimental/GeometryCollectionPlugin/Source/GeometryCollectionNodes/Private/Dataflow/GeometryCollectionMeshNodes.cpp
--- a/Engine/Plugins/Experimental/GeometryCollectionPlugin/Source/GeometryCollectionNodes/Private/Dataflow/GeometryCollectionMeshNodes.cpp
+++ b/Engine/Plugins/Experimental/GeometryCollectionPlugin/Source/GeometryCollectionNodes/Private/Dataflow/GeometryCollectionMeshNodes.cpp
@@ -249,7 +249,7 @@ void FCollectionToMeshDataflowNode::Evaluate(UE::Dataflow::FContext& Context, co
 		{
 			const FManagedArrayCollection& InCollection = GetValue(Context, &Collection);
 
-			InTransformSelection.Initialize(InCollection.NumElements(FGeometryCollection::TransformGroup), true);
+			InTransformSelection.InitializeFromCollection(InCollection, true);
 		}
 
 		if (InTransformSelection.AnySelected())
@@ -972,7 +972,7 @@ void FApplyMeshProcessorToGeometryCollectionDataflowNode::Evaluate(UE::Dataflow:
 			const TArray<int32>& SelectionArr = TransformSelectionFacade.SelectAll();
 
 			FDataflowTransformSelection NewTransformSelection;
-			NewTransformSelection.Initialize(InCollection.NumElements(FGeometryCollection::TransformGroup), false);
+			NewTransformSelection.InitializeFromCollection(InCollection, false);
 			NewTransformSelection.SetFromArray(SelectionArr);
 
 			InTransformSelection = NewTransformSelection;
@@ -1088,7 +1088,7 @@ void FAppendMeshesToCollectionDataflowNode::Evaluate(UE::Dataflow::FContext& Con
 			}
 			if (bModifiedCollection)
 			{
-				NewSelection.Initialize(GeomCollection->Transform.Num(), false);
+				NewSelection.InitializeFromCollection(*GeomCollection, false);
 				for (int32 Idx = FirstNewTransformIndex; Idx < NewSelection.Num(); ++Idx)
 				{
 					NewSelection.SetSelected(Idx);
@@ -1126,7 +1126,7 @@ void FCollectionSelectionToMeshesDataflowNode::Evaluate(UE::Dataflow::FContext&
 		{
 			const FManagedArrayCollection& InCollection = GetValue(Context, &Collection);
 
-			InTransformSelection.Initialize(InCollection.NumElements(FGeometryCollection::TransformGroup), true);
+			InTransformSelection.InitializeFromCollection(InCollection, true);
 		}
 
 		if (InTransformSelection.AnySelected())
diff --git a/Engine/Plugins/Experimental/GeometryCollectionPlugin/Source/GeometryCollectionNodes/Private/Dataflow/GeometryCollectionSelectionNodes.cpp b/Engine/Plugins/Experimental/GeometryCollectionPlugin/Source/GeometryCollectionNodes/Private/Dataflow/GeometryCollectionSelectionNodes.cpp
--- a/Engine/Plugins/Experimental/GeometryCollectionPlugin/Source/GeometryCollectionNodes/Private/Dataflow/GeometryCollectionSelectionNodes.cpp
+++ b/Engine/Plugins/Experimental/GeometryCollectionPlugin/Source/GeometryCollectionNodes/Private/Dataflow/GeometryCollectionSelectionNodes.cpp
@@ -103,7 +103,7 @@ void FCollectionTransformSelectionAllDataflowNode::Evaluate(UE::Dataflow::FConte
 		const TArray<int32>& SelectionArr = TransformSelectionFacade.SelectAll();
 
 		FDataflowTransformSelection NewTransformSelection;
-		NewTransformSelection.Initialize(InCollection.NumElements(FGeometryCollection::TransformGroup), false);
+		NewTransformSelection.InitializeFromCollection(InCollection, false);
 		NewTransformSelection.SetFromArray(SelectionArr);
 
 		SetValue(Context, MoveTemp(NewTransformSelection), &TransformSelection);
@@ -302,7 +302,7 @@ void FCollectionTransformSelectionNoneDataflowNode::Evaluate(UE::Dataflow::FCont
 		const TArray<int32>& SelectionArr = TransformSelectionFacade.SelectNone();
 
 		FDataflowTransformSelection NewTransformSelection;
-		NewTransformSelection.Initialize(InCollection.NumElements(FGeometryCollection::TransformGroup), false);
+		NewTransformSelection.InitializeFromCollection(InCollection, false);
 		NewTransformSelection.SetFromArray(SelectionArr);
 
 		SetValue(Context, MoveTemp(NewTransformSelection), &TransformSelection);
@@ -340,7 +340,7 @@ void FCollectionTransformSelectionRandomDataflowNode::Evaluate(UE::Dataflow::FCo
 		const TArray<int32>& SelectionArr = TransformSelectionFacade.SelectRandom(bDeterministic, RandomSeedVal, RandomThresholdVal);
 
 		FDataflowTransformSelection NewTransformSelection;
-		NewTransformSelection.Initialize(InCollection.NumElements(FGeometryCollection::TransformGroup), false);
+		NewTransformSelection.InitializeFromCollection(InCollection, false);
 		NewTransformSelection.SetFromArray(SelectionArr);
 
 		SetValue(Context, MoveTemp(NewTransformSelection), &TransformSelection);
@@ -362,7 +362,7 @@ void FCollectionTransformSelectionRootDataflowNode::Evaluate(UE::Dataflow::FCont
 		const TArray<int32>& SelectionArr = TransformSelectionFacade.SelectRootBones();
 
 		FDataflowTransformSelection NewTransformSelection;
-		NewTransformSelection.Initialize(InCollection.NumElements(FGeometryCollection::TransformGroup), false);
+		NewTransformSelection.InitializeFromCollection(InCollection, false);
 		NewTransformSelection.SetFromArray(SelectionArr);
 
 		SetValue(Context, MoveTemp(NewTransformSelection), &TransformSelection);
@@ -382,10 +382,8 @@ void FCollectionTransformSelectionCustomDataflowNode::Evaluate(UE::Dataflow::FCo
 
 		if (InCollection.HasGroup(FGeometryCollection::TransformGroup))
 		{
-			const int32 NumTransforms = InCollection.NumElements(FGeometryCollection::TransformGroup);
-
 			FDataflowTransformSelection NewTransformSelection;
-			NewTransformSelection.Initialize(NumTransforms, false);
+			NewTransformSelection.InitializeFromCollection(InCollection, false);
 
 			const FString InBoneIndices = GetValue<FString>(Context, &BoneIndicies);
 
@@ -397,15 +395,17 @@ void FCollectionTransformSelectionCustomDataflowNode::Evaluate(UE::Dataflow::FCo
 				if (IndexStr.IsNumeric())
 				{
 					int32 Index = FCString::Atoi(*IndexStr);
-					if (Index >= 0 && Index < NumTransforms)
+					if (NewTransformSelection.IsValidIndex(Index))
 					{
 						NewTransformSelection.SetSelected(Index);
 					}
 					else
 					{
 						// ERROR: INVALID INDEX
-						FString ErrorStr = "Invalid specified index found.";
-						UE_LOG(LogTemp, Error, TEXT("[Dataflow ERROR] %s"), *ErrorStr);
+						Context.Error(FString::Printf(
+							TEXT("Index %d is not a valid vertex index for transform group size %d."),
+							Index, NewTransformSelection.Num()),
+							this);
 					}
 				}
 			}
@@ -502,10 +502,8 @@ void FCollectionTransformSelectionCustomDataflowNode_v2::Evaluate(UE::Dataflow::
 
 		if (InCollection.HasGroup(FGeometryCollection::TransformGroup))
 		{
-			const int32 NumTransforms = InCollection.NumElements(FGeometryCollection::TransformGroup);
-
 			FDataflowTransformSelection NewTransformSelection;
-			NewTransformSelection.Initialize(NumTransforms, false);
+			NewTransformSelection.InitializeFromCollection(InCollection, false);
 
 			const FString InBoneIndices = GetValue(Context, &BoneIndices);
 
@@ -548,21 +546,22 @@ void FCollectionTransformSelectionFromIndexArrayDataflowNode::Evaluate(UE::Dataf
 
 		if (InCollection.HasGroup(FGeometryCollection::TransformGroup))
 		{
-			const int32 NumTransforms = InCollection.NumElements(FGeometryCollection::TransformGroup);
-
 			const TArray<int32>& InBoneIndices = GetValue(Context, &BoneIndices);
 
 			FDataflowTransformSelection NewTransformSelection;
-			NewTransformSelection.Initialize(NumTransforms, false);
+			NewTransformSelection.InitializeFromCollection(InCollection, false);
 			for (int32 SelectedIdx : InBoneIndices)
 			{
-				if (SelectedIdx >= 0 && SelectedIdx < NumTransforms)
+				if (NewTransformSelection.IsValidIndex(SelectedIdx))
 				{
 					NewTransformSelection.SetSelected(SelectedIdx);
 				}
 				else
 				{
-					UE_LOG(LogChaos, Error, TEXT("[Dataflow ERROR] Invalid selection index %d is outside valid bone index range [0, %d)"), SelectedIdx, NumTransforms);
+					Context.Error(FString::Printf(
+							TEXT("Index %d is not a valid vertex index for transform group size %d."),
+							SelectedIdx, NewTransformSelection.Num()),
+							this);
 				}
 			}
 
@@ -711,7 +710,7 @@ void FCollectionTransformSelectionTargetLevelDataflowNode::Evaluate(UE::Dataflow
 		const TArray<int32> AllAtLevel = TransformSelectionFacade.GetBonesExactlyAtLevel(InTargetLevel, bSkipEmbedded);
 
 		FDataflowTransformSelection NewTransformSelection;
-		NewTransformSelection.Initialize(OutCollection.NumElements(FGeometryCollection::TransformGroup), false);
+		NewTransformSelection.InitializeFromCollection(OutCollection, false);
 		NewTransformSelection.SetFromArray(AllAtLevel);
 
 		SetValue(Context, MoveTemp(NewTransformSelection), &TransformSelection);
@@ -752,7 +751,7 @@ void FCollectionTransformSelectionLeafDataflowNode::Evaluate(UE::Dataflow::FCont
 		const TArray<int32>& SelectionArr = TransformSelectionFacade.SelectLeaf();
 
 		FDataflowTransformSelection NewTransformSelection;
-		NewTransformSelection.Initialize(InCollection.NumElements(FGeometryCollection::TransformGroup), false);
+		NewTransformSelection.InitializeFromCollection(InCollection, false);
 		NewTransformSelection.SetFromArray(SelectionArr);
 
 		SetValue(Context, MoveTemp(NewTransformSelection), &TransformSelection);
@@ -777,7 +776,7 @@ void FCollectionTransformSelectionClusterDataflowNode::Evaluate(UE::Dataflow::FC
 		const TArray<int32>& SelectionArr = TransformSelectionFacade.SelectLeaf(); // used to be buggy SelectCluster() - see comment above 
 
 		FDataflowTransformSelection NewTransformSelection;
-		NewTransformSelection.Initialize(InCollection.NumElements(FGeometryCollection::TransformGroup), false);
+		NewTransformSelection.InitializeFromCollection(InCollection, false);
 		NewTransformSelection.SetFromArray(SelectionArr);
 
 		SetValue(Context, MoveTemp(NewTransformSelection), &TransformSelection);
@@ -798,7 +797,7 @@ void FCollectionTransformSelectionClusterDataflowNode_v2::Evaluate(UE::Dataflow:
 		const TArray<int32>& SelectionArr = TransformSelectionFacade.SelectCluster(); 
 
 		FDataflowTransformSelection NewTransformSelection;
-		NewTransformSelection.Initialize(InCollection.NumElements(FGeometryCollection::TransformGroup), false);
+		NewTransformSelection.InitializeFromCollection(InCollection, false);
 		NewTransformSelection.SetFromArray(SelectionArr);
 
 		SetValue(Context, MoveTemp(NewTransformSelection), &TransformSelection);
@@ -822,7 +821,7 @@ void FCollectionTransformSelectionBySizeDataflowNode::Evaluate(UE::Dataflow::FCo
 		const TArray<int32>& SelectionArr = TransformSelectionFacade.SelectBySize(InSizeMin, InSizeMax, bInclusive, bInsideRange, bUseRelativeSize);
 
 		FDataflowTransformSelection NewTransformSelection;
-		NewTransformSelection.Initialize(InCollection.NumElements(FGeometryCollection::TransformGroup), false);
+		NewTransformSelection.InitializeFromCollection(InCollection, false);
 		NewTransformSelection.SetFromArray(SelectionArr);
 
 		SetValue(Context, MoveTemp(NewTransformSelection), &TransformSelection);
@@ -848,7 +847,7 @@ void FCollectionTransformSelectionByVolumeDataflowNode::Evaluate(UE::Dataflow::F
 		const TArray<int32>& SelectionArr = TransformSelectionFacade.SelectByVolume(InVolumeMin, InVolumeMax, bInclusive, bInsideRange);
 
 		FDataflowTransformSelection NewTransformSelection;
-		NewTransformSelection.Initialize(InCollection.NumElements(FGeometryCollection::TransformGroup), false);
+		NewTransformSelection.InitializeFromCollection(InCollection, false);
 		NewTransformSelection.SetFromArray(SelectionArr);
 
 		SetValue(Context, MoveTemp(NewTransformSelection), &TransformSelection);
@@ -884,7 +883,7 @@ void FCollectionTransformSelectionInBoxDataflowNode::Evaluate(UE::Dataflow::FCon
 		}
 
 		FDataflowTransformSelection NewTransformSelection;
-		NewTransformSelection.Initialize(InCollection.NumElements(FGeometryCollection::TransformGroup), false);
+		NewTransformSelection.InitializeFromCollection(InCollection, false);
 		NewTransformSelection.SetFromArray(SelectionArr);
 
 		SetValue(Context, MoveTemp(NewTransformSelection), &TransformSelection);
@@ -949,7 +948,7 @@ void FCollectionTransformSelectionInSphereDataflowNode::Evaluate(UE::Dataflow::F
 		}
 
 		FDataflowTransformSelection NewTransformSelection;
-		NewTransformSelection.Initialize(InCollection.NumElements(FGeometryCollection::TransformGroup), false);
+		NewTransformSelection.InitializeFromCollection(InCollection, false);
 		NewTransformSelection.SetFromArray(SelectionArr);
 
 		SetValue(Context, MoveTemp(NewTransformSelection), &TransformSelection);
@@ -1002,7 +1001,7 @@ void FCollectionTransformSelectionByFloatAttrDataflowNode::Evaluate(UE::Dataflow
 		const TArray<int32>& SelectionArr = TransformSelectionFacade.SelectByFloatAttribute(GroupName, AttrName, InMin, InMax, bInclusive, bInsideRange);
 
 		FDataflowTransformSelection NewTransformSelection;
-		NewTransformSelection.Initialize(InCollection.NumElements(FGeometryCollection::TransformGroup), false);
+		NewTransformSelection.InitializeFromCollection(InCollection, false);
 		NewTransformSelection.SetFromArray(SelectionArr);
 
 		SetValue(Context, MoveTemp(NewTransformSelection), &TransformSelection);
@@ -1058,7 +1057,7 @@ void FCollectionTransformSelectionByIntAttrDataflowNode::Evaluate(UE::Dataflow::
 		const TArray<int32>& SelectionArr = TransformSelectionFacade.SelectByIntAttribute(GroupName, AttrName, InMin, InMax, bInclusive, bInsideRange);
 
 		FDataflowTransformSelection NewTransformSelection;
-		NewTransformSelection.Initialize(InCollection.NumElements(FGeometryCollection::TransformGroup), false);
+		NewTransformSelection.InitializeFromCollection(InCollection, false);
 		NewTransformSelection.SetFromArray(SelectionArr);
 
 		SetValue(Context, MoveTemp(NewTransformSelection), &TransformSelection);
@@ -1078,10 +1077,8 @@ void FCollectionVertexSelectionCustomDataflowNode::Evaluate(UE::Dataflow::FConte
 
 		if (InCollection.HasGroup(FGeometryCollection::VerticesGroup))
 		{
-			const int32 NumVertices = InCollection.NumElements(FGeometryCollection::VerticesGroup);
-
 			FDataflowVertexSelection NewVertexSelection;
-			NewVertexSelection.Initialize(NumVertices, false);
+			NewVertexSelection.InitializeFromCollection(InCollection, false);
 
 			const FString InVertexIndices = GetValue<FString>(Context, &VertexIndicies);
 
@@ -1093,15 +1090,17 @@ void FCollectionVertexSelectionCustomDataflowNode::Evaluate(UE::Dataflow::FConte
 				if (IndexStr.IsNumeric())
 				{
 					int32 Index = FCString::Atoi(*IndexStr);
-					if (Index >= 0 && Index < NumVertices)
+					if (NewVertexSelection.IsValidIndex(Index))
 					{
 						NewVertexSelection.SetSelected(Index);
 					}
 					else
 					{
 						// ERROR: INVALID INDEX
-						FString ErrorStr = "Invalid specified index found.";
-						UE_LOG(LogTemp, Error, TEXT("[Dataflow ERROR] %s"), *ErrorStr);
+						Context.Error(FString::Printf(
+							TEXT("Index %d is not a valid vertex index for vertex group size %d."),
+							Index, NewVertexSelection.Num()),
+							this);
 					}
 				}
 			}
@@ -1128,10 +1127,8 @@ void FCollectionFaceSelectionCustomDataflowNode::Evaluate(UE::Dataflow::FContext
 
 		if (InCollection.HasGroup(FGeometryCollection::FacesGroup))
 		{
-			const int32 NumFaces = InCollection.NumElements(FGeometryCollection::FacesGroup);
-
 			FDataflowFaceSelection NewFaceSelection;
-			NewFaceSelection.Initialize(NumFaces, false);
+			NewFaceSelection.InitializeFromCollection(InCollection, false);
 
 			const FString InFaceIndices = GetValue<FString>(Context, &FaceIndicies);
 
@@ -1143,14 +1140,16 @@ void FCollectionFaceSelectionCustomDataflowNode::Evaluate(UE::Dataflow::FContext
 				if (IndexStr.IsNumeric())
 				{
 					const int32 Index = FCString::Atoi(*IndexStr);
-					if (Index >= 0 && Index < NumFaces)
+					if (NewFaceSelection.IsValidIndex(Index))
 					{
 						NewFaceSelection.SetSelected(Index);
 					}
 					else
 					{
-						// ERROR: INVALID INDEX
-						UE_LOG(LogTemp, Error, TEXT("[Dataflow ERROR] Invalid specified index found."));
+						Context.Error(FString::Printf(
+							TEXT("Index %d is not a valid vertex index for face group size %d."),
+							Index, NewFaceSelection.Num()),
+							this);
 					}
 				}
 			}
@@ -1182,7 +1181,7 @@ void FCollectionSelectionConvertDataflowNode::Evaluate(UE::Dataflow::FContext& C
 			const TArray<int32>& SelectionArr = TransformSelectionFacade.ConvertVertexSelectionToTransformSelection(InVertexSelection.AsArray(), bAllElementsMustBeSelected);
 
 			FDataflowTransformSelection NewTransformSelection;
-			NewTransformSelection.Initialize(InCollection.NumElements(FGeometryCollection::TransformGroup), false);
+			NewTransformSelection.InitializeFromCollection(InCollection, false);
 			NewTransformSelection.SetFromArray(SelectionArr);
 
 			SetValue(Context, MoveTemp(NewTransformSelection), &TransformSelection);
@@ -1196,7 +1195,7 @@ void FCollectionSelectionConvertDataflowNode::Evaluate(UE::Dataflow::FContext& C
 			const TArray<int32>& SelectionArr = TransformSelectionFacade.ConvertFaceSelectionToTransformSelection(InFaceSelection.AsArray(), bAllElementsMustBeSelected);
 
 			FDataflowTransformSelection NewTransformSelection;
-			NewTransformSelection.Initialize(InCollection.NumElements(FGeometryCollection::TransformGroup), false);
+			NewTransformSelection.InitializeFromCollection(InCollection, false);
 			NewTransformSelection.SetFromArray(SelectionArr);
 
 			SetValue(Context, MoveTemp(NewTransformSelection), &TransformSelection);
@@ -1218,7 +1217,7 @@ void FCollectionSelectionConvertDataflowNode::Evaluate(UE::Dataflow::FContext& C
 			const TArray<int32>& SelectionArr = TransformSelectionFacade.ConvertVertexSelectionToFaceSelection(InVertexSelection.AsArray(), bAllElementsMustBeSelected);
 
 			FDataflowFaceSelection NewFaceSelection;
-			NewFaceSelection.Initialize(InCollection.NumElements(FGeometryCollection::FacesGroup), false);
+			NewFaceSelection.InitializeFromCollection(InCollection, false);
 			NewFaceSelection.SetFromArray(SelectionArr);
 
 			SetValue(Context, MoveTemp(NewFaceSelection), &FaceSelection);
@@ -1232,7 +1231,7 @@ void FCollectionSelectionConvertDataflowNode::Evaluate(UE::Dataflow::FContext& C
 			const TArray<int32>& SelectionArr = TransformSelectionFacade.ConvertTransformSelectionToFaceSelection(InTransformSelection.AsArray());
 
 			FDataflowFaceSelection NewFaceSelection;
-			NewFaceSelection.Initialize(InCollection.NumElements(FGeometryCollection::FacesGroup), false);
+			NewFaceSelection.InitializeFromCollection(InCollection, false);
 			NewFaceSelection.SetFromArray(SelectionArr);
 
 			SetValue(Context, MoveTemp(NewFaceSelection), &FaceSelection);
@@ -1254,7 +1253,7 @@ void FCollectionSelectionConvertDataflowNode::Evaluate(UE::Dataflow::FContext& C
 			const TArray<int32>& SelectionArr = TransformSelectionFacade.ConvertFaceSelectionToVertexSelection(InFaceSelection.AsArray());
 
 			FDataflowVertexSelection NewVertexSelection;
-			NewVertexSelection.Initialize(InCollection.NumElements(FGeometryCollection::VerticesGroup), false);
+			NewVertexSelection.InitializeFromCollection(InCollection, false);
 			NewVertexSelection.SetFromArray(SelectionArr);
 
 			SetValue(Context, MoveTemp(NewVertexSelection), &VertexSelection);
@@ -1268,7 +1267,7 @@ void FCollectionSelectionConvertDataflowNode::Evaluate(UE::Dataflow::FContext& C
 			const TArray<int32>& SelectionArr = TransformSelectionFacade.ConvertTransformSelectionToVertexSelection(InTransformSelection.AsArray());
 
 			FDataflowVertexSelection NewVertexSelection;
-			NewVertexSelection.Initialize(InCollection.NumElements(FGeometryCollection::VerticesGroup), false);
+			NewVertexSelection.InitializeFromCollection(InCollection, false);
 			NewVertexSelection.SetFromArray(SelectionArr);
 
 			SetValue(Context, MoveTemp(NewVertexSelection), &VertexSelection);
@@ -1565,10 +1564,10 @@ void FCollectionSelectionByAttrDataflowNode::Evaluate(UE::Dataflow::FContext& Co
 		{
 			if (InCollection.HasAttribute(AttributeName, GroupName))
 			{
-				const int32 NumFaces = InCollection.NumElements(GroupName);
+				const int32 GroupSize = InCollection.NumElements(GroupName);
 
 				FDataflowSelection NewGenericSelection;
-				NewGenericSelection.Initialize(NumFaces, false);
+				NewGenericSelection.Initialize(GroupSize, false);
 
 				CreateSelectionFromAttr(InCollection,
 					GroupName,
@@ -1635,11 +1634,10 @@ void FGeometrySelectionToVertexSelectionDataflowNode::Evaluate(UE::Dataflow::FCo
 	if (Out->IsA<FDataflowVertexSelection>(&VertexSelection))
 	{
 		const FManagedArrayCollection& InCollection = GetValue<FManagedArrayCollection>(Context, &Collection);
-		const int32 NumVertices = InCollection.NumElements(FGeometryCollection::VerticesGroup);
 		const int32 NumGeometries = InCollection.NumElements(FGeometryCollection::GeometryGroup);
 		
 		FDataflowVertexSelection InVertexSelection;
-		InVertexSelection.Initialize(NumVertices, false);
+		InVertexSelection.InitializeFromCollection(InCollection, false);
 		const TManagedArray<int32>* VertexStart = InCollection.FindAttributeTyped<int32>("VertexStart", FGeometryCollection::GeometryGroup);
 		const TManagedArray<int32>* VertexCount = InCollection.FindAttributeTyped<int32>("VertexCount", FGeometryCollection::GeometryGroup);
 		TArray<int32> InGeometryIndexArray;
@@ -1804,7 +1802,7 @@ void FCollectionSelectInternalFacesDataflowNode::Evaluate(UE::Dataflow::FContext
 		const int32 TotalNumFaces = InCollection.NumElements(FGeometryCollection::FacesGroup);
 
 		FDataflowFaceSelection OutFaceSelection;
-		OutFaceSelection.Initialize(TotalNumFaces, false);
+		OutFaceSelection.InitializeFromCollection(InCollection, false);
 
 		if (TransformIndexFromGeometry && FaceStart && FaceCount && InternalFaces)
 		{
diff --git a/Engine/Plugins/Experimental/GeometryCollectionPlugin/Source/GeometryCollectionNodes/Private/Dataflow/GeometryCollectionUtilityNodes.cpp b/Engine/Plugins/Experimental/GeometryCollectionPlugin/Source/GeometryCollectionNodes/Private/Dataflow/GeometryCollectionUtilityNodes.cpp
--- a/Engine/Plugins/Experimental/GeometryCollectionPlugin/Source/GeometryCollectionNodes/Private/Dataflow/GeometryCollectionUtilityNodes.cpp
+++ b/Engine/Plugins/Experimental/GeometryCollectionPlugin/Source/GeometryCollectionNodes/Private/Dataflow/GeometryCollectionUtilityNodes.cpp
@@ -765,7 +765,7 @@ void FCopyConvexHullsFromRootDataflowNode::DebugDraw(UE::Dataflow::FContext& Con
 			const TArray<int32>& SelectionArr = TransformSelectionFacade.SelectRootBones();
 
 			FDataflowTransformSelection RootSelection;
-			RootSelection.Initialize(InCollection.NumElements(FTransformCollection::TransformGroup), false);
+			RootSelection.InitializeFromCollection(InCollection, false);
 			RootSelection.SetFromArray(SelectionArr);
 
 			UE::Dataflow::Convex::DebugDrawProc(DataflowRenderingInterface, InCollection, bRandomizeColor, ColorRandomSeed, RootSelection);
@@ -1117,7 +1117,7 @@ void FFixTinyGeoDataflowNode::Evaluate(UE::Dataflow::FContext& Context, const FD
 			const TArray<int32>& SelectionArr = TransformSelectionFacade.SelectAll();
 
 			FDataflowTransformSelection NewTransformSelection;
-			NewTransformSelection.Initialize(InCollection.NumElements(FGeometryCollection::TransformGroup), false);
+			NewTransformSelection.InitializeFromCollection(InCollection, false);
 			NewTransformSelection.SetFromArray(SelectionArr);
 
 			InTransformSelection = NewTransformSelection;
@@ -1166,7 +1166,7 @@ void FRecomputeNormalsInGeometryCollectionDataflowNode::Evaluate(UE::Dataflow::F
 			const TArray<int32>& SelectionArr = TransformSelectionFacade.SelectAll();
 
 			FDataflowTransformSelection NewTransformSelection;
-			NewTransformSelection.Initialize(InCollection.NumElements(FGeometryCollection::TransformGroup), false);
+			NewTransformSelection.InitializeFromCollection(InCollection, false);
 			NewTransformSelection.SetFromArray(SelectionArr);
 
 			InTransformSelection = NewTransformSelection;
@@ -1208,7 +1208,7 @@ void FResampleGeometryCollectionDataflowNode::Evaluate(UE::Dataflow::FContext& C
 			const TArray<int32>& SelectionArr = TransformSelectionFacade.SelectAll();
 
 			FDataflowTransformSelection NewTransformSelection;
-			NewTransformSelection.Initialize(InCollection.NumElements(FGeometryCollection::TransformGroup), false);
+			NewTransformSelection.InitializeFromCollection(InCollection, false);
 			NewTransformSelection.SetFromArray(SelectionArr);
 
 			InTransformSelection = NewTransformSelection;
