Dataflow:
- Added coloring centroids by level
- Added centroid point size control by curve

diff --git a/Engine/Plugins/Experimental/GeometryCollectionPlugin/Source/GeometryCollectionNodes/Private/Dataflow/GeometryCollectionNodes.cpp b/Engine/Plugins/Experimental/GeometryCollectionPlugin/Source/GeometryCollectionNodes/Private/Dataflow/GeometryCollectionNodes.cpp
--- a/Engine/Plugins/Experimental/GeometryCollectionPlugin/Source/GeometryCollectionNodes/Private/Dataflow/GeometryCollectionNodes.cpp
+++ b/Engine/Plugins/Experimental/GeometryCollectionPlugin/Source/GeometryCollectionNodes/Private/Dataflow/GeometryCollectionNodes.cpp
@@ -507,6 +507,12 @@ FGetCentroidsFromCollectionDataflowNode::FGetCentroidsFromCollectionDataflowNode
 	RegisterInputConnection(&Collection);
 	RegisterInputConnection(&TransformSelection);
 	RegisterOutputConnection(&Centroids);
+	
+	// This is a temporary solution to store temp data
+	RegisterOutputConnection(&Levels) .SetCanHidePin(true) .SetPinIsHidden(true); 
+
+	PointSize.GetRichCurve()->AddKey(0.f, 12.f);
+	PointSize.GetRichCurve()->AddKey(1.f, 6.f);
 }
 
 void FGetCentroidsFromCollectionDataflowNode::Evaluate(UE::Dataflow::FContext& Context, const FDataflowOutput* Out) const
@@ -514,8 +520,16 @@ void FGetCentroidsFromCollectionDataflowNode::Evaluate(UE::Dataflow::FContext& C
 	if (Out->IsA(&Centroids))
 	{
 		TArray<FVector> OutCentroids;
+		TArray<int32> OutLevels;
 
-		const FManagedArrayCollection& InCollection = GetValue(Context, &Collection);
+		FManagedArrayCollection InCollection = GetValue(Context, &Collection);
+
+		Chaos::Facades::FCollectionHierarchyFacade HierarchyFacade(InCollection);
+		if (!HierarchyFacade.HasLevelAttribute())
+		{
+			HierarchyFacade.GenerateLevelAttribute();
+		}
+		
 		const FDataflowTransformSelection& InTransformSelection = GetValue(Context, &TransformSelection);
 
 		const bool bIsSelectionConnected = IsConnected(&TransformSelection);
@@ -526,11 +540,12 @@ void FGetCentroidsFromCollectionDataflowNode::Evaluate(UE::Dataflow::FContext& C
 		{
 			const TArray<FVector>& InCentroids = BoundsFacade.GetCentroids();
 			const TManagedArray<int32>& TransformToGeometryIndex = BoundsFacade.GetTransformToGeometryIndex();
+			const TManagedArray<int32>& LevelsArr = InCollection.GetAttribute<int32>(FTransformCollection::LevelAttribute, FTransformCollection::TransformGroup);
 
 			const int32 NumTransforms = InCollection.NumElements(FTransformCollection::TransformGroup);
 			for (int32 TransformIdx = 0; TransformIdx < NumTransforms; ++TransformIdx)
 			{
-				const int32 GeometryIndex = TransformToGeometryIndex.IsValidIndex(TransformIdx)? TransformToGeometryIndex[TransformIdx]: INDEX_NONE;
+				const int32 GeometryIndex = TransformToGeometryIndex.IsValidIndex(TransformIdx) ? TransformToGeometryIndex[TransformIdx] : INDEX_NONE;
 				if (InCentroids.IsValidIndex(GeometryIndex))
 				{
 					const FVector PositionInBoneSpace(InCentroids[GeometryIndex]);
@@ -539,11 +554,13 @@ void FGetCentroidsFromCollectionDataflowNode::Evaluate(UE::Dataflow::FContext& C
 					if (!bIsSelectionConnected || TransformIdx >= InTransformSelection.Num() || InTransformSelection.IsSelected(TransformIdx))
 					{
 						OutCentroids.Add(PositionInCollectionSpace);
+						OutLevels.Add(LevelsArr[TransformIdx]);
 					}
 				}
 			}
 		}
 
+		SetValue(Context, MoveTemp(OutLevels), &Levels);
 		SetValue(Context, MoveTemp(OutCentroids), &Centroids);
 	}
 }
@@ -565,23 +582,66 @@ void FGetCentroidsFromCollectionDataflowNode::DebugDraw(UE::Dataflow::FContext&
 	{
 		if (IsConnected(&Collection))
 		{
-			const FDataflowOutput* CentroidsOutput = FindOutput(&Centroids);
-			if (CentroidsOutput)
-			{
-				const TArray<FVector> EmptyArray;
-				const TArray<FVector>& OutCentroids = CentroidsOutput->GetValue<TArray<FVector>>(Context, EmptyArray);
+			FManagedArrayCollection InCollection = GetValue(Context, &Collection);
+			Chaos::Facades::FCollectionHierarchyFacade HierarchyFacade(InCollection);
 
-				DataflowRenderingInterface.SetWireframe(true);
-				DataflowRenderingInterface.SetWorldPriority();
-				DataflowRenderingInterface.SetShaded(false);
-				DataflowRenderingInterface.SetColor(Color);
-				DataflowRenderingInterface.SetPointSize(PointSize);
+			if (HierarchyFacade.IsValid())
+			{
+				if (!HierarchyFacade.HasLevelAttribute())
+				{
+					HierarchyFacade.GenerateLevelAttribute();
+				}
+				float LevelsMin, LevelsMax;
+				const TManagedArray<int32>& LevelsArr = InCollection.GetAttribute<int32>(FTransformCollection::LevelAttribute, FTransformCollection::TransformGroup);
 
-				if (OutCentroids.Num() > 0)
+				GeometryCollection::Facades::FCollectionTransformFacade TransformFacade(InCollection);
+				if (TransformFacade.IsValid())
 				{
-					for (int32 Idx = 0; Idx < OutCentroids.Num(); ++Idx)
+					TransformFacade.ComputeLevelsBounds(LevelsMin, LevelsMax);
+
+					const FDataflowOutput* CentroidsOutput = FindOutput(&Centroids);
+					if (CentroidsOutput)
 					{
-						DataflowRenderingInterface.DrawPoint(FVector(OutCentroids[Idx]));
+						const TArray<FVector> EmptyArray;
+						const TArray<FVector>& OutCentroids = CentroidsOutput->GetValue<TArray<FVector>>(Context, EmptyArray);
+
+						DataflowRenderingInterface.SetWireframe(true);
+						DataflowRenderingInterface.SetWorldPriority();
+						DataflowRenderingInterface.SetShaded(false);
+						DataflowRenderingInterface.SetColor(Color);
+
+						if (OutCentroids.Num() > 0)
+						{
+							for (int32 Idx = 0; Idx < OutCentroids.Num(); ++Idx)
+							{
+								if (bColorByLevel)
+								{
+									DataflowRenderingInterface.SetColor(UE::Dataflow::Utils::GetColorByLevel(InCollection, LevelsArr[Idx]));
+								}
+
+								float x;
+								if (LevelsMin == LevelsMax)
+								{
+									x = (float)LevelsMin;
+								}
+								else
+								{
+									x = ((float)LevelsArr[Idx] - (float)LevelsMin) / ((float)LevelsMax - (float)LevelsMin);
+								}
+								const float PointSizeComputed = PointSize.GetRichCurveConst()->Eval(x);
+
+								if (bSizeByLevel)
+								{
+									DataflowRenderingInterface.SetPointSize(PointSizeComputed);
+								}
+								else
+								{
+									DataflowRenderingInterface.SetPointSize(Size);
+								}
+
+								DataflowRenderingInterface.DrawPoint(FVector(OutCentroids[Idx]));
+							}
+						}
 					}
 				}
 			}
diff --git a/Engine/Plugins/Experimental/GeometryCollectionPlugin/Source/GeometryCollectionNodes/Private/Dataflow/GeometryCollectionUtils.cpp b/Engine/Plugins/Experimental/GeometryCollectionPlugin/Source/GeometryCollectionNodes/Private/Dataflow/GeometryCollectionUtils.cpp
--- a/Engine/Plugins/Experimental/GeometryCollectionPlugin/Source/GeometryCollectionNodes/Private/Dataflow/GeometryCollectionUtils.cpp
+++ b/Engine/Plugins/Experimental/GeometryCollectionPlugin/Source/GeometryCollectionNodes/Private/Dataflow/GeometryCollectionUtils.cpp
@@ -1,6 +1,7 @@
 // Copyright Epic Games, Inc. All Rights Reserved.
 
 #include "Dataflow/GeometryCollectionUtils.h"
+#include "Dataflow/DataflowSettings.h"
 
 #include "GeometryCollection/GeometryCollectionAlgo.h"
 
@@ -86,6 +87,28 @@ namespace UE::Dataflow
 				}
 			}
 		}
+
+		FLinearColor GetColorByLevel(const FManagedArrayCollection& InCollection, int32 InLevel)
+		{
+			const UDataflowSettings* DataflowSettings = GetDefault<UDataflowSettings>();
+			const int32 NumTransformLevelColors = DataflowSettings->TransformLevelColors.LevelColors.Num();
+
+			if (const TManagedArray<int32>* Levels = InCollection.FindAttribute<int32>(FTransformCollection::LevelAttribute, FGeometryCollection::TransformGroup))
+			{
+				const int32 NumTransforms = Levels->Num();
+
+				if (InLevel >= 0)
+				{
+					return DataflowSettings->TransformLevelColors.LevelColors[InLevel % NumTransformLevelColors];
+				}
+
+				return DataflowSettings->TransformLevelColors.BlankColor;
+			}
+			else
+			{
+				return FLinearColor::White;
+			}
+		}
 #endif
 	}
 }
diff --git a/Engine/Plugins/Experimental/GeometryCollectionPlugin/Source/GeometryCollectionNodes/Public/Dataflow/GeometryCollectionNodes.h b/Engine/Plugins/Experimental/GeometryCollectionPlugin/Source/GeometryCollectionNodes/Public/Dataflow/GeometryCollectionNodes.h
--- a/Engine/Plugins/Experimental/GeometryCollectionPlugin/Source/GeometryCollectionNodes/Public/Dataflow/GeometryCollectionNodes.h
+++ b/Engine/Plugins/Experimental/GeometryCollectionPlugin/Source/GeometryCollectionNodes/Public/Dataflow/GeometryCollectionNodes.h
@@ -12,6 +12,7 @@
 #include "Math/MathFwd.h"
 #include "DynamicMesh/DynamicMesh3.h"
 #include "Dataflow/DataflowDebugDraw.h"
+#include "Curves/CurveFloat.h"
 
 #include "GeometryCollectionNodes.generated.h"
 
@@ -640,16 +641,34 @@ private:
 	UPROPERTY(meta = (DataflowInput))
 	FDataflowTransformSelection TransformSelection;
 
+	/** Control point color by level or setting */
 	UPROPERTY(EditAnywhere, Category = "Debug Draw")
+	bool bColorByLevel = true;
+
+	/** Point color */
+	UPROPERTY(EditAnywhere, Category = "Debug Draw", meta = (EditCondition = "bColorByLevel == false", EditConditionHides))
 	FLinearColor Color = FLinearColor::Blue;
 
+	/** Control the point size by level or setting */
+	UPROPERTY(EditAnywhere, Category = "Debug Draw")
+	bool bSizeByLevel = true;
+
+	/** Point size */
+	UPROPERTY(EditAnywhere, Category = "Debug Draw", meta = (EditCondition = "bSizeByLevel == false", EditConditionHides))
+	float Size = 8.f;
+
+	/** Float curve to control point size by level */
 	UPROPERTY(EditAnywhere, Category = "Debug Draw")
-	float PointSize = 12.f;
+	FRuntimeFloatCurve PointSize;
 
 	/** Output centroids */
 	UPROPERTY(meta = (DataflowOutput, DisplayName = "Points"))
 	TArray<FVector> Centroids = TArray<FVector>();
 
+	/** Hidden output to store computed level values for centroids */
+	UPROPERTY(meta = (DataflowOutput))
+	TArray<int32> Levels;
+
 	virtual void Evaluate(UE::Dataflow::FContext& Context, const FDataflowOutput* Out) const override;
 
 #if WITH_EDITOR
@@ -659,7 +678,6 @@ private:
 #endif
 public:
 	FGetCentroidsFromCollectionDataflowNode(const UE::Dataflow::FNodeParameters& InParam, FGuid InGuid = FGuid::NewGuid());
-
 };
 
 
diff --git a/Engine/Plugins/Experimental/GeometryCollectionPlugin/Source/GeometryCollectionNodes/Public/Dataflow/GeometryCollectionUtils.h b/Engine/Plugins/Experimental/GeometryCollectionPlugin/Source/GeometryCollectionNodes/Public/Dataflow/GeometryCollectionUtils.h
--- a/Engine/Plugins/Experimental/GeometryCollectionPlugin/Source/GeometryCollectionNodes/Public/Dataflow/GeometryCollectionUtils.h
+++ b/Engine/Plugins/Experimental/GeometryCollectionPlugin/Source/GeometryCollectionNodes/Public/Dataflow/GeometryCollectionUtils.h
@@ -25,6 +25,8 @@ namespace UE::Dataflow
 			const FLinearColor CenterColor,
 			const bool bRandomizeColor,
 			const int32 ColorRandomSeed);
+
+		GEOMETRYCOLLECTIONNODES_API FLinearColor GetColorByLevel(const FManagedArrayCollection& InCollection, int32 InLevel);
 #endif
 	}
 }
diff --git a/Engine/Source/Runtime/Experimental/Chaos/Private/GeometryCollection/Facades/CollectionTransformFacade.cpp b/Engine/Source/Runtime/Experimental/Chaos/Private/GeometryCollection/Facades/CollectionTransformFacade.cpp
--- a/Engine/Source/Runtime/Experimental/Chaos/Private/GeometryCollection/Facades/CollectionTransformFacade.cpp
+++ b/Engine/Source/Runtime/Experimental/Chaos/Private/GeometryCollection/Facades/CollectionTransformFacade.cpp
@@ -428,4 +428,26 @@ namespace GeometryCollection::Facades
 		}
 	}
 
+	void FCollectionTransformFacade::ComputeLevelsBounds(float& OutLevelsMin, float& OutLevelsMax)
+	{
+		const TManagedArray<int32>& Levels = Collection->GetAttribute<int32>(FTransformCollection::LevelAttribute, FTransformCollection::TransformGroup);
+
+		int32 LevelsMin = INT_MAX, LevelsMax = -INT_MAX;
+		for (int32 Idx = 0; Idx < Levels.Num(); ++Idx)
+		{
+			const int32 Level = Levels[Idx];
+			if (Level < LevelsMin)
+			{
+				LevelsMin = Level;
+			}
+			else if (Level > LevelsMax)
+			{
+				LevelsMax = Level;
+			}
+		}
+
+		OutLevelsMin = (float)LevelsMin;
+		OutLevelsMax = (float)LevelsMax;
+	}
+
 }
\ No newline at end of file
diff --git a/Engine/Source/Runtime/Experimental/Chaos/Public/GeometryCollection/Facades/CollectionTransformFacade.h b/Engine/Source/Runtime/Experimental/Chaos/Public/GeometryCollection/Facades/CollectionTransformFacade.h
--- a/Engine/Source/Runtime/Experimental/Chaos/Public/GeometryCollection/Facades/CollectionTransformFacade.h
+++ b/Engine/Source/Runtime/Experimental/Chaos/Public/GeometryCollection/Facades/CollectionTransformFacade.h
@@ -128,6 +128,9 @@ namespace GeometryCollection::Facades
 		/** Adds a Identity transform and nests all roots under the new transform */
 		CHAOS_API void EnforceSingleRoot(const FString & RootName);
 
+		/** Computes the Min/Max values of the Level attribute */
+		CHAOS_API void ComputeLevelsBounds(float& LevelsMin, float& LevelsMax);
+
 	private:
 		const FManagedArrayCollection& ConstCollection;
 		FManagedArrayCollection* Collection = nullptr;
