Dataflow : Geometry script nodes : fix the issues with Flat blueprint variables ( and arrays as well ) causing the values to be wrong
- properly handle Float type BP variable ( internally stored as double ) to read a flat value from the cache and write a double value to the variable property

diff --git a/Engine/Plugins/Experimental/GeometryCollectionPlugin/Source/GeometryCollectionNodes/Private/Dataflow/GeometryCollectionMeshNodes.cpp b/Engine/Plugins/Experimental/GeometryCollectionPlugin/Source/GeometryCollectionNodes/Private/Dataflow/GeometryCollectionMeshNodes.cpp
--- a/Engine/Plugins/Experimental/GeometryCollectionPlugin/Source/GeometryCollectionNodes/Private/Dataflow/GeometryCollectionMeshNodes.cpp
+++ b/Engine/Plugins/Experimental/GeometryCollectionPlugin/Source/GeometryCollectionNodes/Private/Dataflow/GeometryCollectionMeshNodes.cpp
@@ -898,7 +898,27 @@ void FMeshProcessorDataflowNodeBase::ApplyParametersToBlueprintInstance(UE::Data
 							if (const FDataflowOutput* Output = Input->GetConnection())
 							{
 								Input->PullValue(Context);
-								if (const void* CachedValueAddress = Context.GetUntypedData(Output->CacheKey(), nullptr))
+								if (Desc.ValueType == EPropertyBagPropertyType::Double)
+								{
+									// Special case : the internal storage type for Float BP variable is Double
+									// so we need to pass a pointer on a double but reading a float from Dataflow (Same applies for array type)
+									if (Desc.ContainerTypes.GetFirstContainerType() == EPropertyBagContainerType::Array)
+									{
+										const TArray<float> CacheValue = Context.GetData<TArray<float>>(Output->CacheKey(), Output->GetProperty(), {});
+										TArray<double> DoubleArrayValue;
+										Algo::Copy(CacheValue, DoubleArrayValue);
+										void* TargetAddress = BPProperty->ContainerPtrToValuePtr<void>(MeshProcessorInstance);
+										BPProperty->CopyCompleteValue(TargetAddress, &DoubleArrayValue);
+									}
+									else
+									{
+										const float CachedValue = Context.GetData<float>(Output->CacheKey(), Output->GetProperty(), 0.0f);
+										const double DoubleValue{ CachedValue };
+										void* TargetAddress = BPProperty->ContainerPtrToValuePtr<void>(MeshProcessorInstance);
+										BPProperty->CopyCompleteValue(TargetAddress, &DoubleValue);
+									}
+								}
+								else if (const void* CachedValueAddress = Context.GetUntypedData(Output->CacheKey(), nullptr))
 								{
 									void* TargetAddress = BPProperty->ContainerPtrToValuePtr<void>(MeshProcessorInstance);
 									BPProperty->CopyCompleteValue(TargetAddress, CachedValueAddress);
