Added a 'Cut Planes' array input to the PlaneCutter  dataflow node, to support directly specifying the cutting plane(s)

diff --git a/Engine/Plugins/Experimental/Fracture/Source/FractureEngine/Private/FractureEngineFracturing.cpp b/Engine/Plugins/Experimental/Fracture/Source/FractureEngine/Private/FractureEngineFracturing.cpp
--- a/Engine/Plugins/Experimental/Fracture/Source/FractureEngine/Private/FractureEngineFracturing.cpp
+++ b/Engine/Plugins/Experimental/Fracture/Source/FractureEngine/Private/FractureEngineFracturing.cpp
@@ -515,13 +515,14 @@ int32 FFractureEngineFracturing::PlaneCutter(FManagedArrayCollection& InOutColle
 	int32 InOctaveNumber,
 	float InPointSpacing,
 	bool InAddSamplesForCollision,
-	float InCollisionSampleSpacing)
+	float InCollisionSampleSpacing,
+	TConstArrayView<FTransform> InCutPlaneTransforms)
 {
 
 	if (TUniquePtr<FGeometryCollection> GeomCollection = TUniquePtr<FGeometryCollection>(InOutCollection.NewCopy<FGeometryCollection>()))
 	{
 		TArray<FPlane> CuttingPlanes;
-		TArray<FTransform> CuttingPlaneTransforms;
+		TArray<FTransform> CuttingPlaneTransforms(InCutPlaneTransforms);
 
 		FFractureEngineFracturing::GenerateSliceTransforms(InBoundingBox, InRandomSeed, InNumPlanes, CuttingPlaneTransforms);
 
diff --git a/Engine/Plugins/Experimental/Fracture/Source/FractureEngine/Public/FractureEngineFracturing.h b/Engine/Plugins/Experimental/Fracture/Source/FractureEngine/Public/FractureEngineFracturing.h
--- a/Engine/Plugins/Experimental/Fracture/Source/FractureEngine/Public/FractureEngineFracturing.h
+++ b/Engine/Plugins/Experimental/Fracture/Source/FractureEngine/Public/FractureEngineFracturing.h
@@ -127,7 +127,8 @@ public:
 		int32 InOctaveNumber,
 		float InPointSpacing,
 		bool InAddSamplesForCollision,
-		float InCollisionSampleSpacing);
+		float InCollisionSampleSpacing,
+		TConstArrayView<FTransform> InCutPlaneTransforms = TConstArrayView<FTransform>());
 
 	static UE_API void GenerateSliceTransforms(TArray<FTransform>& InOutCuttingPlaneTransforms,
 		const FBox& InBoundingBox,
diff --git a/Engine/Plugins/Experimental/GeometryCollectionPlugin/Source/GeometryCollectionNodes/Private/Dataflow/GeometryCollectionFracturingNodes.cpp b/Engine/Plugins/Experimental/GeometryCollectionPlugin/Source/GeometryCollectionNodes/Private/Dataflow/GeometryCollectionFracturingNodes.cpp
--- a/Engine/Plugins/Experimental/GeometryCollectionPlugin/Source/GeometryCollectionNodes/Private/Dataflow/GeometryCollectionFracturingNodes.cpp
+++ b/Engine/Plugins/Experimental/GeometryCollectionPlugin/Source/GeometryCollectionNodes/Private/Dataflow/GeometryCollectionFracturingNodes.cpp
@@ -674,6 +674,7 @@ void FPlaneCutterDataflowNode_v2::Evaluate(UE::Dataflow::FContext& Context, cons
 		}
 
 		const int32 InNumPlanes = GetValue(Context, &NumPlanes);
+		const TArray<FTransform>& InCutPlaneTransforms = GetValue(Context, &CutPlanes);
 
 		if (InTransformSelection.AnySelected())
 		{
@@ -693,7 +694,8 @@ void FPlaneCutterDataflowNode_v2::Evaluate(UE::Dataflow::FContext& Context, cons
 				GetValue(Context, &OctaveNumber),
 				GetValue(Context, &PointSpacing),
 				AddSamplesForCollision,
-				GetValue(Context, &CollisionSampleSpacing));
+				GetValue(Context, &CollisionSampleSpacing),
+				InCutPlaneTransforms);
 
 			FDataflowTransformSelection NewSelection;
 			FDataflowTransformSelection OriginalSelection;
@@ -773,7 +775,7 @@ void FPlaneCutterDataflowNode_v2::DebugDraw(UE::Dataflow::FContext& Context, IDa
 
 		const int32 InRandomSeed = GetValue<int32>(Context, &RandomSeed);
 
-		TArray<FTransform> PlaneTransforms;
+		TArray<FTransform> PlaneTransforms = GetValue(Context, &CutPlanes);
 
 		FFractureEngineFracturing::GenerateSliceTransforms(InBoundingBox, InRandomSeed, InNumPlanes, PlaneTransforms);
 
diff --git a/Engine/Plugins/Experimental/GeometryCollectionPlugin/Source/GeometryCollectionNodes/Public/Dataflow/GeometryCollectionFracturingNodes.h b/Engine/Plugins/Experimental/GeometryCollectionPlugin/Source/GeometryCollectionNodes/Public/Dataflow/GeometryCollectionFracturingNodes.h
--- a/Engine/Plugins/Experimental/GeometryCollectionPlugin/Source/GeometryCollectionNodes/Public/Dataflow/GeometryCollectionFracturingNodes.h
+++ b/Engine/Plugins/Experimental/GeometryCollectionPlugin/Source/GeometryCollectionNodes/Public/Dataflow/GeometryCollectionFracturingNodes.h
@@ -723,13 +723,17 @@ public:
 	UPROPERTY(meta = (DataflowInput, DataflowOutput, DataflowPassthrough = "TransformSelection", DisplayName = "TransformSelection"))
 	FDataflowTransformSelection TransformSelection;
 
-	/** Transform to apply to cut planes */
-	UPROPERTY(meta = (DataflowInput, DisplayName = "Transform"))
+	/** Transform to bring the collection into the space of the cutting planes */
+	UPROPERTY(meta = (DataflowInput, DisplayName = "Transform Collection to Plane Space"))
 	FTransform Transform = FTransform::Identity;
+	
+	/** Cutting planes to use. The cut plane normal is aligned to the Z-up of the transform. */
+	UPROPERTY(EditAnywhere, Category = "Plane Cut", meta = (DataflowInput))
+	TArray<FTransform> CutPlanes;
 
-	/** Number of cutting planes */
-	UPROPERTY(EditAnywhere, Category = "Plane Cut", meta = (DataflowInput, DisplayName = "Number of Cuts", UIMin = 1))
-	int32 NumPlanes = 1;
+	/** Number of random cutting planes to add, in addition to those specified in the Cut Planes array */
+	UPROPERTY(EditAnywhere, Category = "Plane Cut", meta = (DataflowInput, DisplayName = "Add Random Cuts", ClampMin = 0))
+	int32 NumPlanes = 0;
 
 	/** Seed for random */
 	UPROPERTY(EditAnywhere, Category = "Common Fracture", meta = (DataflowInput, UIMin = "0"))
@@ -820,6 +824,7 @@ public:
 		RegisterInputConnection(&TransformSelection);
 		RegisterInputConnection(&Transform) .SetCanHidePin(true).SetPinIsHidden(true);
 		RegisterInputConnection(&NumPlanes) .SetCanHidePin(true).SetPinIsHidden(true);
+		RegisterInputConnection(&CutPlanes) .SetCanHidePin(true).SetPinIsHidden(true);
 		RegisterInputConnection(&RandomSeed) .SetCanHidePin(true).SetPinIsHidden(true);
 		RegisterInputConnection(&ChanceToFracture) .SetCanHidePin(true).SetPinIsHidden(true);
 		RegisterInputConnection(&Grout) .SetCanHidePin(true).SetPinIsHidden(true);
