optionally allow dynamic mesh processor blueprint to specify that it is ok to run async

diff --git a/Engine/Plugins/Experimental/GeometryCollectionPlugin/Source/GeometryCollectionNodes/Public/Dataflow/GeometryCollectionMeshNodes.h b/Engine/Plugins/Experimental/GeometryCollectionPlugin/Source/GeometryCollectionNodes/Public/Dataflow/GeometryCollectionMeshNodes.h
--- a/Engine/Plugins/Experimental/GeometryCollectionPlugin/Source/GeometryCollectionNodes/Public/Dataflow/GeometryCollectionMeshNodes.h
+++ b/Engine/Plugins/Experimental/GeometryCollectionPlugin/Source/GeometryCollectionNodes/Public/Dataflow/GeometryCollectionMeshNodes.h
@@ -114,10 +114,10 @@ private:
 
 	virtual void OnPropertyChanged(UE::Dataflow::FContext& Context, const FPropertyChangedEvent& PropertyChangedEvent) override;
 
-	// Blueprint evaluation is not generally thread-safe
+	// Evaluate on game thread unless the blueprint opts in to async evaluation
 	virtual bool EvaluateOnGameThreadOnly() const override
 	{
-		return true;
+		return MeshProcessorInstance && MeshProcessorInstance->bRequiresGameThread;
 	}
 
 #if WITH_EDITOR
diff --git a/Engine/Source/Runtime/GeometryFramework/Public/DynamicMeshProcessor.h b/Engine/Source/Runtime/GeometryFramework/Public/DynamicMeshProcessor.h
--- a/Engine/Source/Runtime/GeometryFramework/Public/DynamicMeshProcessor.h
+++ b/Engine/Source/Runtime/GeometryFramework/Public/DynamicMeshProcessor.h
@@ -23,6 +23,10 @@ public:
 	/// @param bFailed Set to true to report that the processing has failed
 	UFUNCTION(BlueprintImplementableEvent, BlueprintCallable, Category = "MeshProcessor")
 	UE_API void ProcessDynamicMesh(UDynamicMesh* TargetMesh, bool& bFailed);
+
+	/** Whether the blueprint must be run on the game thread -- i.e. if it may use nodes that aren't thread safe. */
+	UPROPERTY(EditDefaultsOnly, Category = "MeshProcessor")
+	bool bRequiresGameThread = true;
 };
 
 #undef UE_API
