Dataflow : add the ability to add a root transform when making a collection

diff --git a/Engine/Plugins/Experimental/GeometryCollectionPlugin/Source/GeometryCollectionNodes/Private/Dataflow/GeometryCollectionMakeNodes.cpp b/Engine/Plugins/Experimental/GeometryCollectionPlugin/Source/GeometryCollectionNodes/Private/Dataflow/GeometryCollectionMakeNodes.cpp
--- a/Engine/Plugins/Experimental/GeometryCollectionPlugin/Source/GeometryCollectionNodes/Private/Dataflow/GeometryCollectionMakeNodes.cpp
+++ b/Engine/Plugins/Experimental/GeometryCollectionPlugin/Source/GeometryCollectionNodes/Private/Dataflow/GeometryCollectionMakeNodes.cpp
@@ -366,6 +366,17 @@ void FMakeCollectionDataflowNode::Evaluate(UE::Dataflow::FContext& Context, cons
 {
 	if (Out->IsA(&Collection))
 	{
+		if (bAddRootTransform)
+		{
+			FGeometryCollection NewCollection;
+			const int32 RootIndex = NewCollection.AddElements(1, FGeometryCollection::TransformGroup);
+			NewCollection.Parent[RootIndex] = INDEX_NONE;
+			NewCollection.BoneColor[RootIndex] = FLinearColor::White;
+			NewCollection.BoneName[RootIndex] = TEXT("Root");
+			SetValue(Context, static_cast<const FManagedArrayCollection&>(NewCollection), &Collection);
+			return;
+		}
+		// completely empty collection 
 		SetValue(Context, FManagedArrayCollection(), &Collection);
 	}
 }
diff --git a/Engine/Plugins/Experimental/GeometryCollectionPlugin/Source/GeometryCollectionNodes/Public/Dataflow/GeometryCollectionMakeNodes.h b/Engine/Plugins/Experimental/GeometryCollectionPlugin/Source/GeometryCollectionNodes/Public/Dataflow/GeometryCollectionMakeNodes.h
--- a/Engine/Plugins/Experimental/GeometryCollectionPlugin/Source/GeometryCollectionNodes/Public/Dataflow/GeometryCollectionMakeNodes.h
+++ b/Engine/Plugins/Experimental/GeometryCollectionPlugin/Source/GeometryCollectionNodes/Public/Dataflow/GeometryCollectionMakeNodes.h
@@ -575,11 +575,16 @@ struct FMakeCollectionDataflowNode : public FDataflowNode
 {
 	GENERATED_USTRUCT_BODY()
 	DATAFLOW_NODE_DEFINE_INTERNAL(FMakeCollectionDataflowNode, "MakeCollection", "Generators|Collection", "")
+	DATAFLOW_NODE_RENDER_TYPE("SurfaceRender", FGeometryCollection::StaticType(), "Collection")
 
 private:
 	UPROPERTY(meta = (DataflowOutput));
 	FManagedArrayCollection Collection;
 
+	/** if true, create a root transform */
+	UPROPERTY(EditAnyWhere, Category = "Options")
+	bool bAddRootTransform = false;
+
 public:
 	FMakeCollectionDataflowNode(const UE::Dataflow::FNodeParameters& InParam, FGuid InGuid = FGuid::NewGuid())
 		: FDataflowNode(InParam, InGuid)
