fix crash on double click of dataflow geometry script node when there is no blueprint attached

diff --git a/Engine/Plugins/Experimental/GeometryCollectionPlugin/Source/GeometryCollectionNodes/Private/Dataflow/GeometryCollectionMeshNodes.cpp b/Engine/Plugins/Experimental/GeometryCollectionPlugin/Source/GeometryCollectionNodes/Private/Dataflow/GeometryCollectionMeshNodes.cpp
--- a/Engine/Plugins/Experimental/GeometryCollectionPlugin/Source/GeometryCollectionNodes/Private/Dataflow/GeometryCollectionMeshNodes.cpp
+++ b/Engine/Plugins/Experimental/GeometryCollectionPlugin/Source/GeometryCollectionNodes/Private/Dataflow/GeometryCollectionMeshNodes.cpp
@@ -943,6 +943,10 @@ void FMeshProcessorDataflowNodeBase::ApplyParametersToBlueprintInstance(UE::Data
 #if WITH_EDITOR
 void FMeshProcessorDataflowNodeBase::OnDoubleClicked() const
 {
+	if (!MeshProcessor)
+	{
+		return;
+	}
 	if (UAssetEditorSubsystem* AssetEditorSubsystem = GEditor->GetEditorSubsystem<UAssetEditorSubsystem>())
 	{
 		if (UBlueprint* Blueprint = Cast<UBlueprint>(MeshProcessor->ClassGeneratedBy))
