Dataflow :  Convert most geometry collection nodes as possible to use the new rendering system
- allow for faster generation of the dynamic mesh from the collection and caching of those dynamic meshes

diff --git a/Engine/Plugins/Experimental/GeometryCollectionPlugin/Source/GeometryCollectionNodes/Public/Dataflow/GeometryCollectionAssetNodes.h b/Engine/Plugins/Experimental/GeometryCollectionPlugin/Source/GeometryCollectionNodes/Public/Dataflow/GeometryCollectionAssetNodes.h
--- a/Engine/Plugins/Experimental/GeometryCollectionPlugin/Source/GeometryCollectionNodes/Public/Dataflow/GeometryCollectionAssetNodes.h
+++ b/Engine/Plugins/Experimental/GeometryCollectionPlugin/Source/GeometryCollectionNodes/Public/Dataflow/GeometryCollectionAssetNodes.h
@@ -113,7 +113,6 @@ struct FGeometryCollectionTerminalDataflowNode_v2 : public FDataflowTerminalNode
 {
 	GENERATED_USTRUCT_BODY()
 	DATAFLOW_NODE_DEFINE_INTERNAL(FGeometryCollectionTerminalDataflowNode_v2, "GeometryCollectionTerminal", "Terminal", "")
-		DATAFLOW_NODE_RENDER_TYPE("SurfaceRender", FGeometryCollection::StaticType(), "Collection")
 
 private:
 	UPROPERTY(meta = (DataflowInput, DataflowOutput, DataflowPassthrough = "Collection"))
@@ -146,7 +145,6 @@ struct FGeometryCollectionTerminalDataflowNode : public FDataflowTerminalNode
 {
 	GENERATED_USTRUCT_BODY()
 	DATAFLOW_NODE_DEFINE_INTERNAL(FGeometryCollectionTerminalDataflowNode, "GeometryCollectionTerminal", "Terminal", "")
-	DATAFLOW_NODE_RENDER_TYPE("SurfaceRender",FGeometryCollection::StaticType(),  "Collection")
 
 public:
 
@@ -295,7 +293,6 @@ struct FGeometryCollectionToCollectionDataflowNode : public FDataflowNode
 {
 	GENERATED_USTRUCT_BODY()
 	DATAFLOW_NODE_DEFINE_INTERNAL(FGeometryCollectionToCollectionDataflowNode, "GeometryCollectionToCollection", "GeometryCollection|Asset", "")
-	DATAFLOW_NODE_RENDER_TYPE("SurfaceRender",FGeometryCollection::StaticType(),  "Collection")
 
 public:
 	/** Asset input */
@@ -330,7 +327,6 @@ struct FGeometryCollectionToCollectionDataflowNode_v2 : public FDataflowNode
 {
 	GENERATED_USTRUCT_BODY()
 	DATAFLOW_NODE_DEFINE_INTERNAL(FGeometryCollectionToCollectionDataflowNode_v2, "GeometryCollectionToCollection", "GeometryCollection|Asset", "")
-		DATAFLOW_NODE_RENDER_TYPE("SurfaceRender", FGeometryCollection::StaticType(), "Collection")
 
 public:
 	FGeometryCollectionToCollectionDataflowNode_v2(const UE::Dataflow::FNodeParameters& InParam, FGuid InGuid = FGuid::NewGuid());
@@ -368,7 +364,6 @@ private:
  {
  	GENERATED_USTRUCT_BODY()
  	DATAFLOW_NODE_DEFINE_INTERNAL(FBlueprintToCollectionDataflowNode, "BlueprintToCollection", "GeometryCollection|Asset", "")
-	DATAFLOW_NODE_RENDER_TYPE("SurfaceRender",FGeometryCollection::StaticType(),  "Collection")
 
  public:
  	/** Asset input */
@@ -408,7 +403,6 @@ USTRUCT(meta = (DataflowContext = "GeometryCollection", DataflowGeometryCollecti
 {
 	GENERATED_USTRUCT_BODY()
 	DATAFLOW_NODE_DEFINE_INTERNAL(FBlueprintToCollectionDataflowNode_v2, "BlueprintToCollection", "GeometryCollection|Asset", "")
-	DATAFLOW_NODE_RENDER_TYPE("SurfaceRender", FGeometryCollection::StaticType(), "Collection")
 
 public:
 	FBlueprintToCollectionDataflowNode_v2(const UE::Dataflow::FNodeParameters& InParam, FGuid InGuid = FGuid::NewGuid());
diff --git a/Engine/Plugins/Experimental/GeometryCollectionPlugin/Source/GeometryCollectionNodes/Public/Dataflow/GeometryCollectionClusteringNodes.h b/Engine/Plugins/Experimental/GeometryCollectionPlugin/Source/GeometryCollectionNodes/Public/Dataflow/GeometryCollectionClusteringNodes.h
--- a/Engine/Plugins/Experimental/GeometryCollectionPlugin/Source/GeometryCollectionNodes/Public/Dataflow/GeometryCollectionClusteringNodes.h
+++ b/Engine/Plugins/Experimental/GeometryCollectionPlugin/Source/GeometryCollectionNodes/Public/Dataflow/GeometryCollectionClusteringNodes.h
@@ -40,7 +40,6 @@ struct FAutoClusterDataflowNode : public FDataflowNode
 {
 	GENERATED_USTRUCT_BODY()
 	DATAFLOW_NODE_DEFINE_INTERNAL(FAutoClusterDataflowNode, "AutoCluster", "GeometryCollection|Cluster", "")
-	DATAFLOW_NODE_RENDER_TYPE("SurfaceRender", FGeometryCollection::StaticType(), "Collection")
 
 private:
 	/** How to choose the size of the clusters to create */
diff --git a/Engine/Plugins/Experimental/GeometryCollectionPlugin/Source/GeometryCollectionNodes/Public/Dataflow/GeometryCollectionDebugNodes.h b/Engine/Plugins/Experimental/GeometryCollectionPlugin/Source/GeometryCollectionNodes/Public/Dataflow/GeometryCollectionDebugNodes.h
--- a/Engine/Plugins/Experimental/GeometryCollectionPlugin/Source/GeometryCollectionNodes/Public/Dataflow/GeometryCollectionDebugNodes.h
+++ b/Engine/Plugins/Experimental/GeometryCollectionPlugin/Source/GeometryCollectionNodes/Public/Dataflow/GeometryCollectionDebugNodes.h
@@ -59,7 +59,6 @@ struct FSphereCoveringToMeshDataflowNode : public FDataflowNode
 {
 	GENERATED_USTRUCT_BODY()
 	DATAFLOW_NODE_DEFINE_INTERNAL(FSphereCoveringToMeshDataflowNode, "Sphere Covering to Mesh", "GeometryCollection|Mesh", "")
-	DATAFLOW_NODE_RENDER_TYPE("SurfaceRender", FName("FDynamicMesh3"), "Mesh")
 
 public:
 	// The sphere covering to convert to a mesh
diff --git a/Engine/Plugins/Experimental/GeometryCollectionPlugin/Source/GeometryCollectionNodes/Public/Dataflow/GeometryCollectionEditNodes.h b/Engine/Plugins/Experimental/GeometryCollectionPlugin/Source/GeometryCollectionNodes/Public/Dataflow/GeometryCollectionEditNodes.h
--- a/Engine/Plugins/Experimental/GeometryCollectionPlugin/Source/GeometryCollectionNodes/Public/Dataflow/GeometryCollectionEditNodes.h
+++ b/Engine/Plugins/Experimental/GeometryCollectionPlugin/Source/GeometryCollectionNodes/Public/Dataflow/GeometryCollectionEditNodes.h
@@ -27,7 +27,6 @@ struct FPruneInCollectionDataflowNode : public FDataflowNode
 {
 	GENERATED_USTRUCT_BODY()
 	DATAFLOW_NODE_DEFINE_INTERNAL(FPruneInCollectionDataflowNode, "PruneInCollection", "GeometryCollection|Edit", "")
-	DATAFLOW_NODE_RENDER_TYPE("SurfaceRender", FGeometryCollection::StaticType(), "Collection")
 
 public:
 	/** Fractured GeometryCollection to prune */
@@ -71,7 +70,6 @@ struct FSetVisibilityInCollectionDataflowNode : public FDataflowNode
 {
 	GENERATED_USTRUCT_BODY()
 	DATAFLOW_NODE_DEFINE_INTERNAL(FSetVisibilityInCollectionDataflowNode, "SetVisibilityInCollection", "GeometryCollection|Edit", "")
-	DATAFLOW_NODE_RENDER_TYPE("SurfaceRender", FGeometryCollection::StaticType(), "Collection")
 
 public:
 	/** What to set the visibility of the selected bones */
@@ -114,7 +112,6 @@ struct FMergeInCollectionDataflowNode : public FDataflowNode
 {
 	GENERATED_USTRUCT_BODY()
 	DATAFLOW_NODE_DEFINE_INTERNAL(FMergeInCollectionDataflowNode, "MergeInCollection", "GeometryCollection|Edit", "")
-	DATAFLOW_NODE_RENDER_TYPE("SurfaceRender", FGeometryCollection::StaticType(), "Collection")
 
 public:
 	/** Fractured GeometryCollection to merge */
diff --git a/Engine/Plugins/Experimental/GeometryCollectionPlugin/Source/GeometryCollectionNodes/Public/Dataflow/GeometryCollectionFracturingNodes.h b/Engine/Plugins/Experimental/GeometryCollectionPlugin/Source/GeometryCollectionNodes/Public/Dataflow/GeometryCollectionFracturingNodes.h
--- a/Engine/Plugins/Experimental/GeometryCollectionPlugin/Source/GeometryCollectionNodes/Public/Dataflow/GeometryCollectionFracturingNodes.h
+++ b/Engine/Plugins/Experimental/GeometryCollectionPlugin/Source/GeometryCollectionNodes/Public/Dataflow/GeometryCollectionFracturingNodes.h
@@ -511,7 +511,6 @@ struct FVoronoiFractureDataflowNode_v2 : public FDataflowNode
 {
 	GENERATED_USTRUCT_BODY()
 	DATAFLOW_NODE_DEFINE_INTERNAL(FVoronoiFractureDataflowNode_v2, "VoronoiFracture", "GeometryCollection|Fracture", "")
-	DATAFLOW_NODE_RENDER_TYPE("SurfaceRender", FGeometryCollection::StaticType(), "Collection")
 
 public:
 	/** Collection to fracture */
@@ -708,7 +707,6 @@ struct FPlaneCutterDataflowNode_v2 : public FDataflowNode
 {
 	GENERATED_USTRUCT_BODY()
 	DATAFLOW_NODE_DEFINE_INTERNAL(FPlaneCutterDataflowNode_v2, "PlaneCutter", "GeometryCollection|Fracture", "")
-	DATAFLOW_NODE_RENDER_TYPE("SurfaceRender", FGeometryCollection::StaticType(), "Collection")
 
 public:
 	/** Collection to fracture */
@@ -862,7 +860,6 @@ struct FExplodedViewDataflowNode : public FDataflowNode
 {
 	GENERATED_USTRUCT_BODY()
 	DATAFLOW_NODE_DEFINE_INTERNAL(FExplodedViewDataflowNode, "ExplodedView", "GeometryCollection|Fracture|Utilities", "")
-	DATAFLOW_NODE_RENDER_TYPE("SurfaceRender",FGeometryCollection::StaticType(),  "Collection")
 
 public:
 	/** Collection to explode */
@@ -915,7 +912,6 @@ struct FSliceCutterDataflowNode : public FDataflowNode
 {
 	GENERATED_USTRUCT_BODY()
 	DATAFLOW_NODE_DEFINE_INTERNAL(FSliceCutterDataflowNode, "SliceCutter", "GeometryCollection|Fracture", "")
-	DATAFLOW_NODE_RENDER_TYPE("SurfaceRender", FGeometryCollection::StaticType(), "Collection")
 
 public:
 	/** Collection to fracture */
@@ -1053,7 +1049,6 @@ struct FBrickCutterDataflowNode : public FDataflowNode
 {
 	GENERATED_USTRUCT_BODY()
 	DATAFLOW_NODE_DEFINE_INTERNAL(FBrickCutterDataflowNode, "BrickCutter", "GeometryCollection|Fracture", "")
-	DATAFLOW_NODE_RENDER_TYPE("SurfaceRender", FGeometryCollection::StaticType(), "Collection")
 
 public:
 	/** Collection to cut */
@@ -1190,7 +1185,6 @@ struct FMeshCutterDataflowNode : public FDataflowNode
 {
 	GENERATED_USTRUCT_BODY()
 	DATAFLOW_NODE_DEFINE_INTERNAL(FMeshCutterDataflowNode, "MeshCutter", "GeometryCollection|Fracture", "")
-	DATAFLOW_NODE_RENDER_TYPE("SurfaceRender", FGeometryCollection::StaticType(), "Collection")
 
 public:
 	/** Collection to cut */
@@ -1350,7 +1344,6 @@ struct FUniformFractureDataflowNode : public FDataflowNode
 {
 	GENERATED_USTRUCT_BODY()
 	DATAFLOW_NODE_DEFINE_INTERNAL(FUniformFractureDataflowNode, "UniformFracture", "GeometryCollection|Fracture", "")
-	DATAFLOW_NODE_RENDER_TYPE("SurfaceRender", FGeometryCollection::StaticType(), "Collection")
 
 private:
 	/** Collection to fracture */
@@ -1495,7 +1488,6 @@ struct FVisualizeFractureDataflowNode : public FDataflowNode
 {
 	GENERATED_USTRUCT_BODY()
 	DATAFLOW_NODE_DEFINE_INTERNAL(FVisualizeFractureDataflowNode, "VisualizeFracture", "GeometryCollection|Fracture|Utilities", "")
-	DATAFLOW_NODE_RENDER_TYPE("SurfaceRender", FGeometryCollection::StaticType(), "Collection")
 
 private:
 	/** Collection to visualize */
@@ -1578,7 +1570,6 @@ struct FSetFloatAttributeDataflowNode : public FDataflowNode
 {
 	GENERATED_USTRUCT_BODY()
 	DATAFLOW_NODE_DEFINE_INTERNAL(FSetFloatAttributeDataflowNode, "SetFloatAttribute", "GeometryCollection|Utilities", "")
-	DATAFLOW_NODE_RENDER_TYPE("SurfaceRender", FGeometryCollection::StaticType(), "Collection")
 
 private:
 	/** Collection to visualize */
diff --git a/Engine/Plugins/Experimental/GeometryCollectionPlugin/Source/GeometryCollectionNodes/Public/Dataflow/GeometryCollectionMakeNodes.h b/Engine/Plugins/Experimental/GeometryCollectionPlugin/Source/GeometryCollectionNodes/Public/Dataflow/GeometryCollectionMakeNodes.h
--- a/Engine/Plugins/Experimental/GeometryCollectionPlugin/Source/GeometryCollectionNodes/Public/Dataflow/GeometryCollectionMakeNodes.h
+++ b/Engine/Plugins/Experimental/GeometryCollectionPlugin/Source/GeometryCollectionNodes/Public/Dataflow/GeometryCollectionMakeNodes.h
@@ -575,7 +575,6 @@ struct FMakeCollectionDataflowNode : public FDataflowNode
 {
 	GENERATED_USTRUCT_BODY()
 	DATAFLOW_NODE_DEFINE_INTERNAL(FMakeCollectionDataflowNode, "MakeCollection", "Generators|Collection", "")
-	DATAFLOW_NODE_RENDER_TYPE("SurfaceRender", FGeometryCollection::StaticType(), "Collection")
 
 private:
 	UPROPERTY(meta = (DataflowOutput));
@@ -686,7 +685,6 @@ struct FMakeSphereMeshDataflowNode : public FDataflowNode
 {
 	GENERATED_USTRUCT_BODY()
 	DATAFLOW_NODE_DEFINE_INTERNAL(FMakeSphereMeshDataflowNode, "MakeSphereMesh", "Generators|Mesh", "")
-	DATAFLOW_NODE_RENDER_TYPE("SurfaceRender", FName("FDynamicMesh3"), "Mesh")
 
 private:
 	/** Sphere Radius */
@@ -720,7 +718,6 @@ struct FMakeCapsuleMeshDataflowNode : public FDataflowNode
 {
 	GENERATED_USTRUCT_BODY()
 	DATAFLOW_NODE_DEFINE_INTERNAL(FMakeCapsuleMeshDataflowNode, "MakeCapsuleMesh", "Generators|Mesh", "")
-	DATAFLOW_NODE_RENDER_TYPE("SurfaceRender", FName("FDynamicMesh3"), "Mesh")
 
 private:
 	/** Radius of capsule */
@@ -762,7 +759,6 @@ struct FMakeCylinderMeshDataflowNode : public FDataflowNode
 {
 	GENERATED_USTRUCT_BODY()
 	DATAFLOW_NODE_DEFINE_INTERNAL(FMakeCylinderMeshDataflowNode, "MakeCylinderMesh", "Generators|Mesh", "")
-	DATAFLOW_NODE_RENDER_TYPE("SurfaceRender", FName("FDynamicMesh3"), "Mesh")
 
 private:
 	/** Radius1 of cylinder */
@@ -804,7 +800,6 @@ struct FMakeBoxMeshDataflowNode : public FDataflowNode
 {
 	GENERATED_USTRUCT_BODY()
 	DATAFLOW_NODE_DEFINE_INTERNAL(FMakeBoxMeshDataflowNode, "MakeBoxMesh", "Generators|Mesh", "")
-	DATAFLOW_NODE_RENDER_TYPE("SurfaceRender", FName("FDynamicMesh3"), "Mesh")
 
 private:
 	/**  */
@@ -890,7 +885,6 @@ struct FMakeDiscMeshDataflowNode : public FDataflowNode
 {
 	GENERATED_USTRUCT_BODY()
 	DATAFLOW_NODE_DEFINE_INTERNAL(FMakeDiscMeshDataflowNode, "MakeDiscMesh", "Generators|Mesh", "")
-	DATAFLOW_NODE_RENDER_TYPE("SurfaceRender", FName("FDynamicMesh3"), "Mesh")
 
 private:
 	/** Radius */
@@ -944,7 +938,6 @@ struct FMakeStairMeshDataflowNode : public FDataflowNode
 {
 	GENERATED_USTRUCT_BODY()
 	DATAFLOW_NODE_DEFINE_INTERNAL(FMakeStairMeshDataflowNode, "MakeStairMesh", "Generators|Mesh", "")
-	DATAFLOW_NODE_RENDER_TYPE("SurfaceRender", FName("FDynamicMesh3"), "Mesh")
 
 private:
 	/** Type of staircase */
@@ -994,7 +987,6 @@ struct FMakeRectangleMeshDataflowNode : public FDataflowNode
 {
 	GENERATED_USTRUCT_BODY()
 	DATAFLOW_NODE_DEFINE_INTERNAL(FMakeRectangleMeshDataflowNode, "MakeRectangleMesh", "Generators|Mesh", "")
-	DATAFLOW_NODE_RENDER_TYPE("SurfaceRender", FName("FDynamicMesh3"), "Mesh")
 
 private:
 	/** Rectangle will be translated so that center is at this point */
@@ -1040,7 +1032,6 @@ struct FMakeTorusMeshDataflowNode : public FDataflowNode
 {
 	GENERATED_USTRUCT_BODY()
 	DATAFLOW_NODE_DEFINE_INTERNAL(FMakeTorusMeshDataflowNode, "MakeTorusMesh", "Generators|Mesh", "")
-	DATAFLOW_NODE_RENDER_TYPE("SurfaceRender", FName("FDynamicMesh3"), "Mesh")
 
 private:
 	/** Torus will be translated so that center is at this point */
diff --git a/Engine/Plugins/Experimental/GeometryCollectionPlugin/Source/GeometryCollectionNodes/Public/Dataflow/GeometryCollectionMeshNodes.h b/Engine/Plugins/Experimental/GeometryCollectionPlugin/Source/GeometryCollectionNodes/Public/Dataflow/GeometryCollectionMeshNodes.h
--- a/Engine/Plugins/Experimental/GeometryCollectionPlugin/Source/GeometryCollectionNodes/Public/Dataflow/GeometryCollectionMeshNodes.h
+++ b/Engine/Plugins/Experimental/GeometryCollectionPlugin/Source/GeometryCollectionNodes/Public/Dataflow/GeometryCollectionMeshNodes.h
@@ -149,7 +149,6 @@ struct FApplyMeshProcessorToMeshDataflowNode final : public FMeshProcessorDatafl
 {
 	GENERATED_USTRUCT_BODY()
 	DATAFLOW_NODE_DEFINE_INTERNAL(FApplyMeshProcessorToMeshDataflowNode, "ApplyGeometryScriptToMesh", "Mesh|Utilities", "")
-	DATAFLOW_NODE_RENDER_TYPE("SurfaceRender", FName("FDynamicMesh3"), "Mesh")
 
 public:
 
@@ -175,7 +174,6 @@ struct FApplyMeshProcessorToGeometryCollectionDataflowNode final : public FMeshP
 {
 	GENERATED_USTRUCT_BODY()
 	DATAFLOW_NODE_DEFINE_INTERNAL(FApplyMeshProcessorToGeometryCollectionDataflowNode, "ApplyGeometryScriptToCollection", "Mesh|Utilities", "")
-	DATAFLOW_NODE_RENDER_TYPE("SurfaceRender", FGeometryCollection::StaticType(), "Collection")
 
 public:
 
@@ -217,7 +215,6 @@ struct FCollectionSelectionToMeshesDataflowNode : public FDataflowNode
 {
 	GENERATED_USTRUCT_BODY()
 	DATAFLOW_NODE_DEFINE_INTERNAL(FCollectionSelectionToMeshesDataflowNode, "CollectionSelectionToMeshes", "Mesh|Utilities", "")
-	//DATAFLOW_NODE_RENDER_TYPE("SurfaceRender", FName("FDynamicMesh3"), "Meshes")
 
 public:
 	/** Collection to convert*/
@@ -260,7 +257,6 @@ struct FAppendMeshesToCollectionDataflowNode : public FDataflowNode
 {
 	GENERATED_USTRUCT_BODY()
 	DATAFLOW_NODE_DEFINE_INTERNAL(FAppendMeshesToCollectionDataflowNode, "AppendMeshesToCollection", "Mesh|Utilities", "")
-	DATAFLOW_NODE_RENDER_TYPE("SurfaceRender", FGeometryCollection::StaticType(), "Collection")
 
 public:
 	/** Meshes will be appended to this collection */
@@ -296,7 +292,6 @@ struct FBoxToMeshDataflowNode : public FDataflowNode
 {
 	GENERATED_USTRUCT_BODY()
 	DATAFLOW_NODE_DEFINE_INTERNAL(FBoxToMeshDataflowNode, "BoxToMesh", "Mesh|Utilities", "")
-	DATAFLOW_NODE_RENDER_TYPE("SurfaceRender",FName("FDynamicMesh3"), "Mesh")
 
 public:
 	/** BoundingBox input */
@@ -366,7 +361,6 @@ struct FMeshToCollectionDataflowNode : public FDataflowNode
 {
 	GENERATED_USTRUCT_BODY()
 	DATAFLOW_NODE_DEFINE_INTERNAL(FMeshToCollectionDataflowNode, "MeshToCollection", "Mesh|Utilities", "")
-	DATAFLOW_NODE_RENDER_TYPE("SurfaceRender",FGeometryCollection::StaticType(),  "Collection")
 
 public:
 	/** DynamicMesh to convert */
@@ -410,7 +404,6 @@ struct FCollectionToMeshDataflowNode : public FDataflowNode
 {
 	GENERATED_USTRUCT_BODY()
 	DATAFLOW_NODE_DEFINE_INTERNAL(FCollectionToMeshDataflowNode, "CollectionToMesh", "Mesh|Utilities", "")
-	DATAFLOW_NODE_RENDER_TYPE("SurfaceRender",FName("FDynamicMesh3"), "Mesh")
 
 public:
 	/** Collection to convert*/
@@ -541,7 +534,6 @@ struct FMeshAppendDataflowNode : public FDataflowNode
 {
 	GENERATED_USTRUCT_BODY()
 	DATAFLOW_NODE_DEFINE_INTERNAL(FMeshAppendDataflowNode, "MeshAppend", "Mesh|Utilities", "")
-	DATAFLOW_NODE_RENDER_TYPE("SurfaceRender",FName("FDynamicMesh3"), "Mesh")
 
 public:
 	/** Mesh input */
@@ -755,7 +747,6 @@ struct FMeshCopyToPointsDataflowNode : public FDataflowNode
 {
 	GENERATED_USTRUCT_BODY()
 	DATAFLOW_NODE_DEFINE_INTERNAL(FMeshCopyToPointsDataflowNode, "ScatterMeshes", "Mesh|Utilities", "")
-	DATAFLOW_NODE_RENDER_TYPE("SurfaceRender",FName("FDynamicMesh3"), "Mesh")
 
 public:
 	/** Points to copy meshes onto */
diff --git a/Engine/Plugins/Experimental/GeometryCollectionPlugin/Source/GeometryCollectionNodes/Public/Dataflow/GeometryCollectionNodes.h b/Engine/Plugins/Experimental/GeometryCollectionPlugin/Source/GeometryCollectionNodes/Public/Dataflow/GeometryCollectionNodes.h
--- a/Engine/Plugins/Experimental/GeometryCollectionPlugin/Source/GeometryCollectionNodes/Public/Dataflow/GeometryCollectionNodes.h
+++ b/Engine/Plugins/Experimental/GeometryCollectionPlugin/Source/GeometryCollectionNodes/Public/Dataflow/GeometryCollectionNodes.h
@@ -38,7 +38,6 @@ struct FGetCollectionFromAssetDataflowNode : public FDataflowNode
 {
 	GENERATED_USTRUCT_BODY()
 	DATAFLOW_NODE_DEFINE_INTERNAL(FGetCollectionFromAssetDataflowNode, "GetCollectionFromAsset", "GeometryCollection|Asset", "")
-	DATAFLOW_NODE_RENDER_TYPE("SurfaceRender", FGeometryCollection::StaticType(), "Collection")
 
 public:
 	UPROPERTY(EditAnywhere, Category = "Dataflow", meta = (DataflowInput, DisplayName = "Asset"))
@@ -69,7 +68,6 @@ struct FAppendCollectionAssetsDataflowNode : public FDataflowNode
 {
 	GENERATED_USTRUCT_BODY()
 	DATAFLOW_NODE_DEFINE_INTERNAL(FAppendCollectionAssetsDataflowNode, "AppendCollections", "GeometryCollection", "")
-	DATAFLOW_NODE_RENDER_TYPE("SurfaceRender", FGeometryCollection::StaticType(), "Collection1")
 
 public:
 	typedef FManagedArrayCollection DataType;
@@ -705,7 +703,6 @@ struct FTransformCollectionDataflowNode : public FDataflowNode
 {
 	GENERATED_USTRUCT_BODY()
 	DATAFLOW_NODE_DEFINE_INTERNAL(FTransformCollectionDataflowNode, "TransformCollection", "Math|Transform", "")
-	DATAFLOW_NODE_RENDER_TYPE("SurfaceRender",FGeometryCollection::StaticType(),  "Collection")
 
 public:
 	/** Output mesh */
@@ -778,7 +775,6 @@ struct FBakeTransformsInCollectionDataflowNode : public FDataflowNode
 {
 	GENERATED_USTRUCT_BODY()
 	DATAFLOW_NODE_DEFINE_INTERNAL(FBakeTransformsInCollectionDataflowNode, "BakeTransformsInCollection", "Math|Transform", "")
-	DATAFLOW_NODE_RENDER_TYPE("SurfaceRender",FGeometryCollection::StaticType(),  "Collection")
 
 public:
 	/** Collection to bake transforms in */
@@ -807,7 +803,6 @@ struct FTransformMeshDataflowNode : public FDataflowNode
 {
 	GENERATED_USTRUCT_BODY()
 	DATAFLOW_NODE_DEFINE_INTERNAL(FTransformMeshDataflowNode, "TransformMesh", "Math|Transform", "")
-	DATAFLOW_NODE_RENDER_TYPE("SurfaceRender",FName("FDynamicMesh3"), "Mesh")
 
 public:
 	/** Output mesh */
@@ -1068,7 +1063,6 @@ struct FBranchCollectionDataflowNode : public FDataflowNode
 {
 	GENERATED_USTRUCT_BODY()
 	DATAFLOW_NODE_DEFINE_INTERNAL(FBranchCollectionDataflowNode, "BranchCollection", "Utilities|FlowControl", "")
-	DATAFLOW_NODE_RENDER_TYPE("SurfaceRender", FGeometryCollection::StaticType(), "ChosenCollection")
 
 
 public:
@@ -1140,7 +1134,6 @@ struct FRemoveOnBreakDataflowNode : public FDataflowNode
 {
 	GENERATED_USTRUCT_BODY()
 	DATAFLOW_NODE_DEFINE_INTERNAL(FRemoveOnBreakDataflowNode, "RemoveOnBreak", "GeometryCollection|Utilities", "")
-	DATAFLOW_NODE_RENDER_TYPE("SurfaceRender", FGeometryCollection::StaticType(), "Collection")
 
 public:
 	/** Collection to set theremoval data on */
@@ -1204,7 +1197,6 @@ struct FSetAnchorStateDataflowNode : public FDataflowNode
 {
 	GENERATED_USTRUCT_BODY()
 	DATAFLOW_NODE_DEFINE_INTERNAL(FSetAnchorStateDataflowNode, "SetAnchorState", "GeometryCollection|Utilities", "")
-	DATAFLOW_NODE_RENDER_TYPE("SurfaceRender", FGeometryCollection::StaticType(), "Collection")
 
 private:
 	/** What anchor state to set on selected bones */
@@ -1253,7 +1245,6 @@ struct FSetDynamicStateDataflowNode : public FDataflowNode
 {
 	GENERATED_USTRUCT_BODY()
 	DATAFLOW_NODE_DEFINE_INTERNAL(FSetDynamicStateDataflowNode, "SetDynamicState", "GeometryCollection|Utilities", "")
-	DATAFLOW_NODE_RENDER_TYPE("SurfaceRender", FGeometryCollection::StaticType(), "Collection")
 
 private:
 	/** Dynamic state to set on selected bones */
@@ -1325,7 +1316,6 @@ struct FProximityDataflowNode : public FDataflowNode
 {
 	GENERATED_USTRUCT_BODY()
 	DATAFLOW_NODE_DEFINE_INTERNAL(FProximityDataflowNode, "Proximity", "GeometryCollection|Utilities", "")
-	DATAFLOW_NODE_RENDER_TYPE("SurfaceRender", FGeometryCollection::StaticType(), "Collection")
 
 private:
 	/** Which method to use to decide whether a given piece of geometry is in proximity with another */
@@ -1413,7 +1403,6 @@ struct FCollectionSetPivotDataflowNode : public FDataflowNode
 {
 	GENERATED_USTRUCT_BODY()
 	DATAFLOW_NODE_DEFINE_INTERNAL(FCollectionSetPivotDataflowNode, "SetPivot", "GeometryCollection|Utilities", "")
-	DATAFLOW_NODE_RENDER_TYPE("SurfaceRender", FGeometryCollection::StaticType(), "Collection")
 
 public:
 	/** Collection for the pivot change */
@@ -1495,7 +1484,6 @@ struct FAddCustomCollectionAttributeDataflowNode : public FDataflowNode
 {
 	GENERATED_USTRUCT_BODY()
 	DATAFLOW_NODE_DEFINE_INTERNAL(FAddCustomCollectionAttributeDataflowNode, "AddCustomCollectionAttribute", "GeometryCollection|Utilities", "")
-	DATAFLOW_NODE_RENDER_TYPE("SurfaceRender", FGeometryCollection::StaticType(), "Collection")
 
 public:
 	/** Collection for the custom attribute */
@@ -1667,7 +1655,6 @@ struct FGetCollectionAttributeDataTypedDataflowNode_v2 : public FDataflowNode
 {
 	GENERATED_USTRUCT_BODY()
 	DATAFLOW_NODE_DEFINE_INTERNAL(FGetCollectionAttributeDataTypedDataflowNode_v2, "GetCollectionAttributeDataTyped", "GeometryCollection|Utilities", "")
-	DATAFLOW_NODE_RENDER_TYPE("SurfaceRender", FGeometryCollection::StaticType(), "Collection")
 
 private:
 	/** Collection for the custom attribute */
@@ -1732,7 +1719,6 @@ struct FSetCollectionAttributeDataTypedDataflowNode : public FDataflowNode
 {
 	GENERATED_USTRUCT_BODY()
 	DATAFLOW_NODE_DEFINE_INTERNAL(FSetCollectionAttributeDataTypedDataflowNode, "SetCollectionAttributeDataTyped", "GeometryCollection|Utilities", "")
-	DATAFLOW_NODE_RENDER_TYPE("SurfaceRender", FGeometryCollection::StaticType(), "Collection")
 
 public:
 	/** Collection for the custom attribute */
diff --git a/Engine/Plugins/Experimental/GeometryCollectionPlugin/Source/GeometryCollectionNodes/Public/Dataflow/GeometryCollectionSelectionNodes.h b/Engine/Plugins/Experimental/GeometryCollectionPlugin/Source/GeometryCollectionNodes/Public/Dataflow/GeometryCollectionSelectionNodes.h
--- a/Engine/Plugins/Experimental/GeometryCollectionPlugin/Source/GeometryCollectionNodes/Public/Dataflow/GeometryCollectionSelectionNodes.h
+++ b/Engine/Plugins/Experimental/GeometryCollectionPlugin/Source/GeometryCollectionNodes/Public/Dataflow/GeometryCollectionSelectionNodes.h
@@ -32,7 +32,6 @@ struct FCollectionTransformSelectionAllDataflowNode : public FDataflowNode
 {
 	GENERATED_USTRUCT_BODY()
 	DATAFLOW_NODE_DEFINE_INTERNAL(FCollectionTransformSelectionAllDataflowNode, "CollectionTransformSelectAll", "GeometryCollection|Selection|Transform", "")
-	DATAFLOW_NODE_RENDER_TYPE("SurfaceRender", FGeometryCollection::StaticType(), "Collection")
 
 public:
 	/** GeometryCollection for the selection */
@@ -160,7 +159,6 @@ struct FCollectionTransformSelectionNoneDataflowNode : public FDataflowNode
 {
 	GENERATED_USTRUCT_BODY()
 	DATAFLOW_NODE_DEFINE_INTERNAL(FCollectionTransformSelectionNoneDataflowNode, "CollectionTransformSelectNone", "GeometryCollection|Selection|Transform", "")
-	DATAFLOW_NODE_RENDER_TYPE("SurfaceRender", FGeometryCollection::StaticType(), "Collection")
 
 public:
 	/** GeometryCollection for the selection */
@@ -222,7 +220,6 @@ struct FCollectionTransformSelectionRandomDataflowNode : public FDataflowNode
 {
 	GENERATED_USTRUCT_BODY()
 	DATAFLOW_NODE_DEFINE_INTERNAL(FCollectionTransformSelectionRandomDataflowNode, "CollectionTransformSelectRandom", "GeometryCollection|Selection|Transform", "")
-	DATAFLOW_NODE_RENDER_TYPE("SurfaceRender", FGeometryCollection::StaticType(), "Collection")
 
 public:
 	/** If true, it always generates the same result for the same RandomSeed */
@@ -271,7 +268,6 @@ struct FCollectionTransformSelectionRootDataflowNode : public FDataflowNode
 {
 	GENERATED_USTRUCT_BODY()
 	DATAFLOW_NODE_DEFINE_INTERNAL(FCollectionTransformSelectionRootDataflowNode, "CollectionTransformSelectRoot", "GeometryCollection|Selection|Transform", "")
-	DATAFLOW_NODE_RENDER_TYPE("SurfaceRender", FGeometryCollection::StaticType(), "Collection")
 
 public:
 	/** GeometryCollection for the selection */
@@ -306,7 +302,6 @@ struct FCollectionTransformSelectionCustomDataflowNode : public FDataflowNode
 {
 	GENERATED_USTRUCT_BODY()
 	DATAFLOW_NODE_DEFINE_INTERNAL(FCollectionTransformSelectionCustomDataflowNode, "CollectionTransformSelectCustom", "GeometryCollection|Selection|Transform", "")
-	DATAFLOW_NODE_RENDER_TYPE("SurfaceRender", FGeometryCollection::StaticType(), "Collection")
 
 public:
 	/** GeometryCollection for the selection */
@@ -345,7 +340,6 @@ struct FCollectionTransformSelectionCustomDataflowNode_v2 : public FDataflowNode
 {
 	GENERATED_USTRUCT_BODY()
 	DATAFLOW_NODE_DEFINE_INTERNAL(FCollectionTransformSelectionCustomDataflowNode_v2, "CollectionTransformSelectCustom", "GeometryCollection|Selection|Transform", "")
-	DATAFLOW_NODE_RENDER_TYPE("SurfaceRender", FGeometryCollection::StaticType(), "Collection")
 
 public:
 	/** GeometryCollection for the selection */
@@ -375,7 +369,6 @@ struct FCollectionTransformSelectionFromIndexArrayDataflowNode : public FDataflo
 {
 	GENERATED_USTRUCT_BODY()
 	DATAFLOW_NODE_DEFINE_INTERNAL(FCollectionTransformSelectionFromIndexArrayDataflowNode, "CollectionTransformSelectionFromIndexArray", "GeometryCollection|Selection|Array", "")
-	DATAFLOW_NODE_RENDER_TYPE("SurfaceRender", FGeometryCollection::StaticType(), "Collection")
 
 public:
 
@@ -415,7 +408,6 @@ struct FCollectionTransformSelectionParentDataflowNode : public FDataflowNode
 {
 	GENERATED_USTRUCT_BODY()
 	DATAFLOW_NODE_DEFINE_INTERNAL(FCollectionTransformSelectionParentDataflowNode, "CollectionTransformSelectParent", "GeometryCollection|Selection|Transform", "")
-	DATAFLOW_NODE_RENDER_TYPE("SurfaceRender", FGeometryCollection::StaticType(), "Collection")
 
 public:
 	/** Array of the selected bone indices */
@@ -493,7 +485,6 @@ struct FCollectionTransformSelectionChildrenDataflowNode : public FDataflowNode
 {
 	GENERATED_USTRUCT_BODY()
 	DATAFLOW_NODE_DEFINE_INTERNAL(FCollectionTransformSelectionChildrenDataflowNode, "CollectionTransformSelectChildren", "GeometryCollection|Selection|Transform", "")
-	DATAFLOW_NODE_RENDER_TYPE("SurfaceRender", FGeometryCollection::StaticType(), "Collection")
 
 public:
 	/** Array of the selected bone indices */
@@ -528,7 +519,6 @@ struct FCollectionTransformSelectionSiblingsDataflowNode : public FDataflowNode
 {
 	GENERATED_USTRUCT_BODY()
 	DATAFLOW_NODE_DEFINE_INTERNAL(FCollectionTransformSelectionSiblingsDataflowNode, "CollectionTransformSelectSiblings", "GeometryCollection|Selection|Transform", "")
-	DATAFLOW_NODE_RENDER_TYPE("SurfaceRender", FGeometryCollection::StaticType(), "Collection")
 
 public:
 	/** Array of the selected bone indices */
@@ -563,7 +553,6 @@ struct FCollectionTransformSelectionLevelDataflowNode : public FDataflowNode
 {
 	GENERATED_USTRUCT_BODY()
 	DATAFLOW_NODE_DEFINE_INTERNAL(FCollectionTransformSelectionLevelDataflowNode, "CollectionTransformSelectSameLevel", "GeometryCollection|Selection|Transform", "")
-	DATAFLOW_NODE_RENDER_TYPE("SurfaceRender", FGeometryCollection::StaticType(), "Collection")
 
 public:
 	/** Array of the selected bone indices */
@@ -598,7 +587,6 @@ struct FCollectionTransformSelectionTargetLevelDataflowNode : public FDataflowNo
 {
 	GENERATED_USTRUCT_BODY()
 	DATAFLOW_NODE_DEFINE_INTERNAL(FCollectionTransformSelectionTargetLevelDataflowNode, "CollectionTransformSelectTargetLevel", "GeometryCollection|Selection|Transform", "")
-	DATAFLOW_NODE_RENDER_TYPE("SurfaceRender", FGeometryCollection::StaticType(), "Collection")
 
 public:
 	/** GeometryCollection for the selection */
@@ -642,7 +630,6 @@ struct FCollectionTransformSelectionContactDataflowNode : public FDataflowNode
 {
 	GENERATED_USTRUCT_BODY()
 	DATAFLOW_NODE_DEFINE_INTERNAL(FCollectionTransformSelectionContactDataflowNode, "CollectionTransformSelectContact", "GeometryCollection|Selection|Transform", "")
-	DATAFLOW_NODE_RENDER_TYPE("SurfaceRender", FGeometryCollection::StaticType(), "Collection")
 
 public:
 	/** Array of the selected bone indices */
@@ -681,7 +668,6 @@ struct FCollectionTransformSelectionLeafDataflowNode : public FDataflowNode
 {
 	GENERATED_USTRUCT_BODY()
 	DATAFLOW_NODE_DEFINE_INTERNAL(FCollectionTransformSelectionLeafDataflowNode, "CollectionTransformSelectLeaf", "GeometryCollection|Selection|Transform", "")
-	DATAFLOW_NODE_RENDER_TYPE("SurfaceRender", FGeometryCollection::StaticType(), "Collection")
 
 public:
 	/** GeometryCollection for the selection */
@@ -717,7 +703,6 @@ struct FCollectionTransformSelectionClusterDataflowNode : public FDataflowNode
 {
 	GENERATED_USTRUCT_BODY()
 	DATAFLOW_NODE_DEFINE_INTERNAL(FCollectionTransformSelectionClusterDataflowNode, "CollectionTransformSelectCluster", "GeometryCollection|Selection|Transform", "")
-	DATAFLOW_NODE_RENDER_TYPE("SurfaceRender", FGeometryCollection::StaticType(), "Collection")
 
 public:
 	/** GeometryCollection for the selection */
@@ -751,7 +736,6 @@ struct FCollectionTransformSelectionClusterDataflowNode_v2 : public FDataflowNod
 {
 	GENERATED_USTRUCT_BODY()
 	DATAFLOW_NODE_DEFINE_INTERNAL(FCollectionTransformSelectionClusterDataflowNode_v2, "CollectionTransformSelectCluster", "GeometryCollection|Selection|Transform", "")
-	DATAFLOW_NODE_RENDER_TYPE("SurfaceRender", FGeometryCollection::StaticType(), "Collection")
 
 public:
 	/** GeometryCollection for the selection */
@@ -848,7 +832,6 @@ struct FCollectionTransformSelectionBySizeDataflowNode : public FDataflowNode
 {
 	GENERATED_USTRUCT_BODY()
 	DATAFLOW_NODE_DEFINE_INTERNAL(FCollectionTransformSelectionBySizeDataflowNode, "CollectionTransformSelectBySize", "GeometryCollection|Selection|Transform", "")
-	DATAFLOW_NODE_RENDER_TYPE("SurfaceRender", FGeometryCollection::StaticType(), "Collection")
 
 public:
 	/** GeometryCollection for the selection */
@@ -904,7 +887,6 @@ struct FCollectionTransformSelectionByVolumeDataflowNode : public FDataflowNode
 {
 	GENERATED_USTRUCT_BODY()
 	DATAFLOW_NODE_DEFINE_INTERNAL(FCollectionTransformSelectionByVolumeDataflowNode, "CollectionTransformSelectByVolume", "GeometryCollection|Selection|Transform", "")
-	DATAFLOW_NODE_RENDER_TYPE("SurfaceRender", FGeometryCollection::StaticType(), "Collection")
 
 public:
 	/** GeometryCollection for the selection */
@@ -970,7 +952,6 @@ struct FCollectionTransformSelectionInBoxDataflowNode : public FDataflowNode
 {
 	GENERATED_USTRUCT_BODY()
 	DATAFLOW_NODE_DEFINE_INTERNAL(FCollectionTransformSelectionInBoxDataflowNode, "CollectionTransformSelectInBox", "GeometryCollection|Selection|Transform", "")
-	DATAFLOW_NODE_RENDER_TYPE("SurfaceRender", FGeometryCollection::StaticType(), "Collection")
 
 public:
 	/** GeometryCollection for the selection */
@@ -1026,7 +1007,6 @@ struct FCollectionTransformSelectionInSphereDataflowNode : public FDataflowNode
 {
 	GENERATED_USTRUCT_BODY()
 	DATAFLOW_NODE_DEFINE_INTERNAL(FCollectionTransformSelectionInSphereDataflowNode, "CollectionTransformSelectInSphere", "GeometryCollection|Selection|Transform", "")
-	DATAFLOW_NODE_RENDER_TYPE("SurfaceRender", FGeometryCollection::StaticType(), "Collection")
 
 public:
 	/** GeometryCollection for the selection */
@@ -1082,7 +1062,6 @@ struct FCollectionTransformSelectionByFloatAttrDataflowNode : public FDataflowNo
 {
 	GENERATED_USTRUCT_BODY()
 	DATAFLOW_NODE_DEFINE_INTERNAL(FCollectionTransformSelectionByFloatAttrDataflowNode, "CollectionTransformSelectByFloatAttribute", "GeometryCollection|Selection|Transform", "")
-	DATAFLOW_NODE_RENDER_TYPE("SurfaceRender", FGeometryCollection::StaticType(), "Collection")
 
 public:
 	/** GeometryCollection for the selection */
@@ -1142,7 +1121,6 @@ struct FCollectionTransformSelectionByIntAttrDataflowNode : public FDataflowNode
 {
 	GENERATED_USTRUCT_BODY()
 	DATAFLOW_NODE_DEFINE_INTERNAL(FCollectionTransformSelectionByIntAttrDataflowNode, "CollectionTransformSelectByIntAttribute", "GeometryCollection|Selection|Transform", "")
-	DATAFLOW_NODE_RENDER_TYPE("SurfaceRender", FGeometryCollection::StaticType(), "Collection")
 
 public:
 	/** GeometryCollection for the selection */
@@ -1203,7 +1181,6 @@ struct FCollectionVertexSelectionCustomDataflowNode : public FDataflowNode
 {
 	GENERATED_USTRUCT_BODY()
 	DATAFLOW_NODE_DEFINE_INTERNAL(FCollectionVertexSelectionCustomDataflowNode, "CollectionVertexSelectCustom", "GeometryCollection|Selection|Vertex", "")
-	DATAFLOW_NODE_RENDER_TYPE("SurfaceRender", FGeometryCollection::StaticType(), "Collection")
 
 public:
 	/** GeometryCollection for the selection */
@@ -1243,7 +1220,6 @@ struct FCollectionFaceSelectionCustomDataflowNode : public FDataflowNode
 {
 	GENERATED_USTRUCT_BODY()
 	DATAFLOW_NODE_DEFINE_INTERNAL(FCollectionFaceSelectionCustomDataflowNode, "CollectionFaceSelectCustom", "GeometryCollection|Selection|Face", "")
-	DATAFLOW_NODE_RENDER_TYPE("SurfaceRender", FGeometryCollection::StaticType(), "Collection")
 
 public:
 	/** GeometryCollection for the selection */
@@ -1282,7 +1258,6 @@ struct FCollectionSelectionConvertDataflowNode : public FDataflowNode
 {
 	GENERATED_USTRUCT_BODY()
 	DATAFLOW_NODE_DEFINE_INTERNAL(FCollectionSelectionConvertDataflowNode, "CollectionSelectionConvert", "GeometryCollection|Selection", "")
-	DATAFLOW_NODE_RENDER_TYPE("SurfaceRender", FGeometryCollection::StaticType(), "Collection")
 
 public:
 	/** GeometryCollection for the selection */
@@ -1486,7 +1461,6 @@ struct FCollectionSelectionByAttrDataflowNode : public FDataflowNode
 {
 	GENERATED_USTRUCT_BODY()
 	DATAFLOW_NODE_DEFINE_INTERNAL(FCollectionSelectionByAttrDataflowNode, "CollectionSelectByAttr", "GeometryCollection|Selection|All", "")
-	DATAFLOW_NODE_RENDER_TYPE("SurfaceRender", FGeometryCollection::StaticType(), "Collection")
 
 public:
 	/** GeometryCollection for the selection */
@@ -1659,7 +1633,6 @@ struct FCollectionSelectInternalFacesDataflowNode : public FDataflowNode
 {
 	GENERATED_USTRUCT_BODY()
 	DATAFLOW_NODE_DEFINE_INTERNAL(FCollectionSelectInternalFacesDataflowNode, "CollectionSelectInternalFaces", "GeometryCollection|Selection", "")
-	DATAFLOW_NODE_RENDER_TYPE("SurfaceRender", FGeometryCollection::StaticType(), "Collection")
 
 private:
 	/** Collection to select the internal faces from */
@@ -1704,7 +1677,6 @@ struct FCollectionSelectTransformStringDataflowNode : public FDataflowNode
 {
 	GENERATED_USTRUCT_BODY()
 	DATAFLOW_NODE_DEFINE_INTERNAL(FCollectionSelectTransformStringDataflowNode, "CollectionSelectTransformString", "GeometryCollection|Selection|Transform", "Name, Bone, Attribute")
-	DATAFLOW_NODE_RENDER_TYPE("SurfaceRender", FGeometryCollection::StaticType(), "Collection")
 
 private:
 	/** Collection for the selection */
@@ -1744,7 +1716,6 @@ struct FCollectionSetTransformStringValueDataflowNode : public FDataflowNode
 {
 	GENERATED_USTRUCT_BODY()
 	DATAFLOW_NODE_DEFINE_INTERNAL(FCollectionSetTransformStringValueDataflowNode, "CollectionSetTransformString", "GeometryCollection|Selection|Transform", "Bone Name Attribute")
-	DATAFLOW_NODE_RENDER_TYPE("SurfaceRender", FGeometryCollection::StaticType(), "Collection")
 
 private:
 	/** Collection for the selection */
diff --git a/Engine/Plugins/Experimental/GeometryCollectionPlugin/Source/GeometryCollectionNodes/Public/Dataflow/GeometryCollectionSkeletalMeshToCollectionNode.h b/Engine/Plugins/Experimental/GeometryCollectionPlugin/Source/GeometryCollectionNodes/Public/Dataflow/GeometryCollectionSkeletalMeshToCollectionNode.h
--- a/Engine/Plugins/Experimental/GeometryCollectionPlugin/Source/GeometryCollectionNodes/Public/Dataflow/GeometryCollectionSkeletalMeshToCollectionNode.h
+++ b/Engine/Plugins/Experimental/GeometryCollectionPlugin/Source/GeometryCollectionNodes/Public/Dataflow/GeometryCollectionSkeletalMeshToCollectionNode.h
@@ -23,7 +23,6 @@ struct FSkeletalMeshToCollectionDataflowNode : public FDataflowNode
 {
 	GENERATED_USTRUCT_BODY()
 	DATAFLOW_NODE_DEFINE_INTERNAL(FSkeletalMeshToCollectionDataflowNode, "SkeletalMeshToCollection", "GeometryCollection", "")
-	DATAFLOW_NODE_RENDER_TYPE("SurfaceRender",FGeometryCollection::StaticType(),  "Collection")
 
 public:
 	typedef FManagedArrayCollection DataType;
@@ -53,7 +52,6 @@ struct FCollectionToSkeletalMeshDataflowNode : public FDataflowNode
 {
 	GENERATED_USTRUCT_BODY()
 	DATAFLOW_NODE_DEFINE_INTERNAL(FCollectionToSkeletalMeshDataflowNode, "CollectionToSkeletalMesh", "GeometryCollection", "")
-	DATAFLOW_NODE_RENDER_TYPE("SurfaceRender", FGeometryCollection::StaticType(), "Collection")
 
 private:
 	UPROPERTY(meta = (DataflowInput))
diff --git a/Engine/Plugins/Experimental/GeometryCollectionPlugin/Source/GeometryCollectionNodes/Public/Dataflow/GeometryCollectionStaticMeshToCollectionNode.h b/Engine/Plugins/Experimental/GeometryCollectionPlugin/Source/GeometryCollectionNodes/Public/Dataflow/GeometryCollectionStaticMeshToCollectionNode.h
--- a/Engine/Plugins/Experimental/GeometryCollectionPlugin/Source/GeometryCollectionNodes/Public/Dataflow/GeometryCollectionStaticMeshToCollectionNode.h
+++ b/Engine/Plugins/Experimental/GeometryCollectionPlugin/Source/GeometryCollectionNodes/Public/Dataflow/GeometryCollectionStaticMeshToCollectionNode.h
@@ -27,7 +27,6 @@ struct FStaticMeshToCollectionDataflowNode_v2 : public FDataflowNode
 {
 	GENERATED_USTRUCT_BODY()
 	DATAFLOW_NODE_DEFINE_INTERNAL(FStaticMeshToCollectionDataflowNode_v2, "StaticMeshToCollection", "GeometryCollection|Asset", "")
-	DATAFLOW_NODE_RENDER_TYPE("SurfaceRender",FGeometryCollection::StaticType(),  "Collection")
 
 private:
 	/** Asset input */
@@ -77,7 +76,6 @@ struct FStaticMeshToCollectionDataflowNode : public FDataflowNode
 {
 	GENERATED_USTRUCT_BODY()
 	DATAFLOW_NODE_DEFINE_INTERNAL(FStaticMeshToCollectionDataflowNode, "StaticMeshToCollection", "GeometryCollection|Asset", "")
-		DATAFLOW_NODE_RENDER_TYPE("SurfaceRender", FGeometryCollection::StaticType(), "Collection")
 
 public:
 	/** Asset input */
diff --git a/Engine/Plugins/Experimental/GeometryCollectionPlugin/Source/GeometryCollectionNodes/Public/Dataflow/GeometryCollectionTransferVertexAttributeNode.h b/Engine/Plugins/Experimental/GeometryCollectionPlugin/Source/GeometryCollectionNodes/Public/Dataflow/GeometryCollectionTransferVertexAttributeNode.h
--- a/Engine/Plugins/Experimental/GeometryCollectionPlugin/Source/GeometryCollectionNodes/Public/Dataflow/GeometryCollectionTransferVertexAttributeNode.h
+++ b/Engine/Plugins/Experimental/GeometryCollectionPlugin/Source/GeometryCollectionNodes/Public/Dataflow/GeometryCollectionTransferVertexAttributeNode.h
@@ -179,7 +179,6 @@ struct FGeometryCollectionTransferVertexSkinWeightsNode : public FDataflowNode
 {
 	GENERATED_USTRUCT_BODY()
 	DATAFLOW_NODE_DEFINE_INTERNAL(FGeometryCollectionTransferVertexSkinWeightsNode, "TransferVertexSkinWeights", "GeometryCollection", "Transfer vertex skin weights from the Source Collection to the Target Collection")
-	DATAFLOW_NODE_RENDER_TYPE("SurfaceRender", FGeometryCollection::StaticType(), "Collection")
 
 public:
 
@@ -258,7 +257,6 @@ struct FGeometryCollectionSetKinematicVertexSelectionNode : public FDataflowNode
 {
 	GENERATED_USTRUCT_BODY()
 	DATAFLOW_NODE_DEFINE_INTERNAL(FGeometryCollectionSetKinematicVertexSelectionNode, "SetKinematicVertexSelection", "GeometryCollection", "Set Vertex Collection to be kinematic")
-	DATAFLOW_NODE_RENDER_TYPE("SurfaceRender", FGeometryCollection::StaticType(), "Collection")
 
 public:
 
diff --git a/Engine/Plugins/Experimental/GeometryCollectionPlugin/Source/GeometryCollectionNodes/Public/Dataflow/GeometryCollectionTriangleBoundaryIndicesNode.h b/Engine/Plugins/Experimental/GeometryCollectionPlugin/Source/GeometryCollectionNodes/Public/Dataflow/GeometryCollectionTriangleBoundaryIndicesNode.h
--- a/Engine/Plugins/Experimental/GeometryCollectionPlugin/Source/GeometryCollectionNodes/Public/Dataflow/GeometryCollectionTriangleBoundaryIndicesNode.h
+++ b/Engine/Plugins/Experimental/GeometryCollectionPlugin/Source/GeometryCollectionNodes/Public/Dataflow/GeometryCollectionTriangleBoundaryIndicesNode.h
@@ -20,7 +20,6 @@ struct FTriangleBoundaryIndicesNode : public FDataflowNode
 {
 	GENERATED_USTRUCT_BODY()
 	DATAFLOW_NODE_DEFINE_INTERNAL(FTriangleBoundaryIndicesNode, "TriangleBoundaryIndices", "Flesh|Experimental", "")
-	DATAFLOW_NODE_RENDER_TYPE("SurfaceRender",FGeometryCollection::StaticType(),  "Collection")
 
 private:
 	typedef FManagedArrayCollection DataType;
diff --git a/Engine/Plugins/Experimental/GeometryCollectionPlugin/Source/GeometryCollectionNodes/Public/Dataflow/GeometryCollectionUtilityNodes.h b/Engine/Plugins/Experimental/GeometryCollectionPlugin/Source/GeometryCollectionNodes/Public/Dataflow/GeometryCollectionUtilityNodes.h
--- a/Engine/Plugins/Experimental/GeometryCollectionPlugin/Source/GeometryCollectionNodes/Public/Dataflow/GeometryCollectionUtilityNodes.h
+++ b/Engine/Plugins/Experimental/GeometryCollectionPlugin/Source/GeometryCollectionNodes/Public/Dataflow/GeometryCollectionUtilityNodes.h
@@ -159,7 +159,6 @@ struct FCreateLeafConvexHullsDataflowNode : public FDataflowNode
 {
 	GENERATED_USTRUCT_BODY()
 	DATAFLOW_NODE_DEFINE_INTERNAL(FCreateLeafConvexHullsDataflowNode, "CreateLeafConvexHulls", "GeometryCollection|Utilities", "")
-	DATAFLOW_NODE_RENDER_TYPE("SurfaceRender", FGeometryCollection::StaticType(), "Collection")
 
 private:
 	UPROPERTY(meta = (DataflowInput, DataflowOutput))
@@ -363,7 +362,6 @@ struct FGenerateClusterConvexHullsFromLeafHullsDataflowNode : public FDataflowNo
 {
 	GENERATED_USTRUCT_BODY()
 	DATAFLOW_NODE_DEFINE_INTERNAL(FGenerateClusterConvexHullsFromLeafHullsDataflowNode, "GenerateClusterConvexHullsFromLeafHulls", "GeometryCollection|Utilities", "")
-	DATAFLOW_NODE_RENDER_TYPE("SurfaceRender", FGeometryCollection::StaticType(), "Collection")
 
 private:
 	UPROPERTY(meta = (DataflowInput, DataflowOutput))
@@ -481,7 +479,6 @@ struct FGenerateClusterConvexHullsFromChildrenHullsDataflowNode : public FDatafl
 {
 	GENERATED_USTRUCT_BODY()
 	DATAFLOW_NODE_DEFINE_INTERNAL(FGenerateClusterConvexHullsFromChildrenHullsDataflowNode, "GenerateClusterConvexHullsFromChildrenHulls", "GeometryCollection|Utilities", "")
-	DATAFLOW_NODE_RENDER_TYPE("SurfaceRender", FGeometryCollection::StaticType(), "Collection")
 
 private:
 	UPROPERTY(meta = (DataflowInput, DataflowOutput))
@@ -620,7 +617,6 @@ struct FCopyConvexHullsFromRootDataflowNode : public FDataflowNode
 {
 	GENERATED_USTRUCT_BODY()
 	DATAFLOW_NODE_DEFINE_INTERNAL(FCopyConvexHullsFromRootDataflowNode, "CopyConvexHullsFromRoot", "GeometryCollection|Utilities", "")
-	DATAFLOW_NODE_RENDER_TYPE("SurfaceRender", FGeometryCollection::StaticType(), "Collection")
 
 public:
 	FCopyConvexHullsFromRootDataflowNode(const UE::Dataflow::FNodeParameters& InParam, FGuid InGuid = FGuid::NewGuid());
diff --git a/Engine/Plugins/Experimental/GeometryCollectionPlugin/Source/GeometryCollectionNodes/Public/Dataflow/GeometryCollectionVerticesNodes.h b/Engine/Plugins/Experimental/GeometryCollectionPlugin/Source/GeometryCollectionNodes/Public/Dataflow/GeometryCollectionVerticesNodes.h
--- a/Engine/Plugins/Experimental/GeometryCollectionPlugin/Source/GeometryCollectionNodes/Public/Dataflow/GeometryCollectionVerticesNodes.h
+++ b/Engine/Plugins/Experimental/GeometryCollectionPlugin/Source/GeometryCollectionNodes/Public/Dataflow/GeometryCollectionVerticesNodes.h
@@ -21,7 +21,6 @@ struct FTransformCollectionAttributeDataflowNode : public FDataflowNode
 {
 	GENERATED_USTRUCT_BODY()
 	DATAFLOW_NODE_DEFINE_INTERNAL(FTransformCollectionAttributeDataflowNode, "TransformCollectionAttribute", "GeometryCollection", "")
-	DATAFLOW_NODE_RENDER_TYPE("SurfaceRender",FGeometryCollection::StaticType(),  "Collection")
 
 public:
 	typedef FManagedArrayCollection DataType;
