Dataflow : Make Create geometry collection from source node to use the context asset source by default is the input is not connected

diff --git a/Engine/Plugins/Experimental/GeometryCollectionPlugin/Source/GeometryCollectionNodes/Private/Dataflow/GeometryCollectionAssetNodes.cpp b/Engine/Plugins/Experimental/GeometryCollectionPlugin/Source/GeometryCollectionNodes/Private/Dataflow/GeometryCollectionAssetNodes.cpp
--- a/Engine/Plugins/Experimental/GeometryCollectionPlugin/Source/GeometryCollectionNodes/Private/Dataflow/GeometryCollectionAssetNodes.cpp
+++ b/Engine/Plugins/Experimental/GeometryCollectionPlugin/Source/GeometryCollectionNodes/Private/Dataflow/GeometryCollectionAssetNodes.cpp
@@ -401,7 +401,30 @@ void FCreateGeometryCollectionFromSourcesDataflowNode_v2::Evaluate(UE::Dataflow:
 {
 	ensure(Out->IsA(&Collection) || Out->IsA(&Materials) || Out->IsA(&InstancedMeshes) || Out->IsA(&RootProxyMeshes));
 
-	const TArray<FGeometryCollectionSource>& InSources = GetValue(Context, &Sources);
+	TArray<FGeometryCollectionSource> InSources;
+	if (IsConnected(&Sources))
+	{
+		InSources = GetValue(Context, &Sources);
+	}
+	else
+	{
+		// not connected let's try to get them from the current asset
+		if (const UE::Dataflow::FEngineContext* EngineContext = Context.AsType<UE::Dataflow::FEngineContext>())
+		{
+			if (UGeometryCollection* CollectionAsset = Cast<UGeometryCollection>(EngineContext->Owner))
+			{
+#if WITH_EDITORONLY_DATA
+				InSources = CollectionAsset->GeometrySource;
+#else
+				ensureMsgf(false, TEXT("FCreateGeometryCollectionFromSourcesDataflowNode - GeometrySource is only available in editor, returning an empty array"));
+#endif
+			}
+		}
+	}
+	if (InSources.IsEmpty())
+	{
+		Context.Warning(TEXT("Geometry source array is empty, this will result in an empty collection"), this, Out);
+	}
 
 	FGeometryCollection OutCollection;
 	TArray<TObjectPtr<UMaterialInterface>> OutMaterials;
@@ -458,13 +481,6 @@ void FCreateGeometryCollectionFromSourcesDataflowNode_v2::Evaluate(UE::Dataflow:
 		OutCollection.ReindexMaterials();
 	}
 
-	// add the instanced mesh indices
-
-	const int32 NumTransforms = InstancedMeshFacade.GetNumIndices();
-	for (int32 TransformIndex = 0; TransformIndex < NumTransforms; TransformIndex++)
-	{
-	}
-
 	// make sure we have only one root
 	if (FGeometryCollectionClusteringUtility::ContainsMultipleRootBones(&OutCollection))
 	{
diff --git a/Engine/Plugins/Experimental/GeometryCollectionPlugin/Source/GeometryCollectionNodes/Public/Dataflow/GeometryCollectionAssetNodes.h b/Engine/Plugins/Experimental/GeometryCollectionPlugin/Source/GeometryCollectionNodes/Public/Dataflow/GeometryCollectionAssetNodes.h
--- a/Engine/Plugins/Experimental/GeometryCollectionPlugin/Source/GeometryCollectionNodes/Public/Dataflow/GeometryCollectionAssetNodes.h
+++ b/Engine/Plugins/Experimental/GeometryCollectionPlugin/Source/GeometryCollectionNodes/Public/Dataflow/GeometryCollectionAssetNodes.h
@@ -250,6 +250,7 @@ public:
 
 /**
  * create a geometry collection from a set of geometry sources
+ * if the source array is not connected, the source array from the current asset is going to be used
   */
 USTRUCT(meta = (DataflowGeometryCollection))
 struct FCreateGeometryCollectionFromSourcesDataflowNode_v2 : public FDataflowNode
