dataflow : expose all properties of TransformCollection node as inputs

diff --git a/Engine/Plugins/Experimental/GeometryCollectionPlugin/Source/GeometryCollectionNodes/Private/Dataflow/GeometryCollectionNodes.cpp b/Engine/Plugins/Experimental/GeometryCollectionPlugin/Source/GeometryCollectionNodes/Private/Dataflow/GeometryCollectionNodes.cpp
--- a/Engine/Plugins/Experimental/GeometryCollectionPlugin/Source/GeometryCollectionNodes/Private/Dataflow/GeometryCollectionNodes.cpp
+++ b/Engine/Plugins/Experimental/GeometryCollectionPlugin/Source/GeometryCollectionNodes/Private/Dataflow/GeometryCollectionNodes.cpp
@@ -652,6 +652,28 @@ void FGetCentroidsFromCollectionDataflowNode::DebugDraw(UE::Dataflow::FContext&
 
 //-----------------------------------------------------------------------------------------------
 
+FTransformCollectionDataflowNode::FTransformCollectionDataflowNode(const UE::Dataflow::FNodeParameters& InParam, FGuid InGuid)
+	: FDataflowNode(InParam, InGuid)
+{
+	RegisterInputConnection(&Collection);
+	RegisterInputConnection(&TransformSelection);
+	RegisterInputConnection(&Translate);
+	RegisterInputConnection(&Rotate);
+	RegisterInputConnection(&Scale);
+	RegisterInputConnection(&RotationOrder)
+		.SetCanHidePin(true).SetPinIsHidden(true);
+	RegisterInputConnection(&UniformScale)
+		.SetCanHidePin(true).SetPinIsHidden(true);
+	RegisterInputConnection(&RotatePivot)
+		.SetCanHidePin(true).SetPinIsHidden(true);
+	RegisterInputConnection(&ScalePivot)
+		.SetCanHidePin(true).SetPinIsHidden(true);
+	RegisterInputConnection(&bInvertTransformation)
+		.SetCanHidePin(true).SetPinIsHidden(true);
+	
+	RegisterOutputConnection(&Collection, &Collection);
+}
+
 void FTransformCollectionDataflowNode::Evaluate(UE::Dataflow::FContext& Context, const FDataflowOutput* Out) const
 {
 	if (Out->IsA<FManagedArrayCollection>(&Collection))
@@ -661,16 +683,21 @@ void FTransformCollectionDataflowNode::Evaluate(UE::Dataflow::FContext& Context,
 		const FVector InTranslate = GetValue(Context, &Translate);
 		const FVector InRotate = GetValue(Context, &Rotate);
 		const FVector InScale = GetValue(Context, &Scale);
+		const ERotationOrderEnum InRotationOrder = GetValue(Context, &RotationOrder);
+		const float InUniformScale = GetValue(Context, &UniformScale);
+		const FVector InRotatePivot = GetValue(Context, &RotatePivot);
+		const FVector InScalePivot = GetValue(Context, &ScalePivot);
+		const bool bInInvertTransformation = GetValue(Context, &bInvertTransformation);
 
-		FTransform NewTransform = GeometryCollection::Facades::FCollectionTransformFacade::BuildTransform(
+		const FTransform NewTransform = GeometryCollection::Facades::FCollectionTransformFacade::BuildTransform(
 			InTranslate,
-			(uint8)RotationOrder,
+			(uint8)InRotationOrder,
 			InRotate,
 			InScale,
-			UniformScale,
-			RotatePivot,
-			ScalePivot,
-			bInvertTransformation);
+			InUniformScale,
+			InRotatePivot,
+			InScalePivot,
+			bInInvertTransformation);
 
 		GeometryCollection::Facades::FCollectionTransformFacade TransformFacade(InCollection);
 
diff --git a/Engine/Plugins/Experimental/GeometryCollectionPlugin/Source/GeometryCollectionNodes/Public/Dataflow/GeometryCollectionNodes.h b/Engine/Plugins/Experimental/GeometryCollectionPlugin/Source/GeometryCollectionNodes/Public/Dataflow/GeometryCollectionNodes.h
--- a/Engine/Plugins/Experimental/GeometryCollectionPlugin/Source/GeometryCollectionNodes/Public/Dataflow/GeometryCollectionNodes.h
+++ b/Engine/Plugins/Experimental/GeometryCollectionPlugin/Source/GeometryCollectionNodes/Public/Dataflow/GeometryCollectionNodes.h
@@ -729,10 +729,6 @@ public:
 	UPROPERTY(EditAnywhere, Category = "Transform", meta = (UIMin = 0.001f), meta = (DataflowInput));
 	FVector Scale = FVector(1.0);
 
-	/** Shear */
-//	UPROPERTY(EditAnywhere, Category = "Transform");
-//	FVector Shear = FVector(0.0);
-
 	/** Uniform scale */
 	UPROPERTY(EditAnywhere, Category = "Transform", meta = (UIMin = 0.001f));
 	float UniformScale = 1.f;
@@ -749,16 +745,7 @@ public:
 	UPROPERTY(EditAnywhere, Category = "General");
 	bool bInvertTransformation = false;
 
-	FTransformCollectionDataflowNode(const UE::Dataflow::FNodeParameters& InParam, FGuid InGuid = FGuid::NewGuid())
-		: FDataflowNode(InParam, InGuid)
-	{
-		RegisterInputConnection(&Collection);
-		RegisterInputConnection(&TransformSelection);
-		RegisterInputConnection(&Translate);
-		RegisterInputConnection(&Rotate);
-		RegisterInputConnection(&Scale);
-		RegisterOutputConnection(&Collection, &Collection);
-	}
+	FTransformCollectionDataflowNode(const UE::Dataflow::FNodeParameters& InParam, FGuid InGuid = FGuid::NewGuid());
 
 	virtual void Evaluate(UE::Dataflow::FContext& Context, const FDataflowOutput* Out) const override;
 
