Geometry collection : add cvar to disable the dialog that ask to create a Dataflow graph when opening a geometry collection asset editor
the variable is p.Chaos.GC.OpenEditorAskForDataflow and is set to true by default

diff --git a/Engine/Plugins/Experimental/GeometryCollectionPlugin/Source/GeometryCollectionEditor/Private/GeometryCollection/AssetDefinition_GeometryCollection.cpp b/Engine/Plugins/Experimental/GeometryCollectionPlugin/Source/GeometryCollectionEditor/Private/GeometryCollection/AssetDefinition_GeometryCollection.cpp
--- a/Engine/Plugins/Experimental/GeometryCollectionPlugin/Source/GeometryCollectionEditor/Private/GeometryCollection/AssetDefinition_GeometryCollection.cpp
+++ b/Engine/Plugins/Experimental/GeometryCollectionPlugin/Source/GeometryCollectionEditor/Private/GeometryCollection/AssetDefinition_GeometryCollection.cpp
@@ -27,6 +27,9 @@
 bool bCanEditGeometryCollection = true;
 FAutoConsoleVariableRef CVarGeometryCollectionIsEditable(TEXT("p.Chaos.GC.IsEditable"), bCanEditGeometryCollection, TEXT("Whether to allow edits of the geometry collection"));
 
+bool bGeometryCollectionOpenEditorAskForDataflow = true;
+FAutoConsoleVariableRef CVarGeometryCollectionOpenEditorAskForDataflow(TEXT("p.Chaos.GC.OpenEditorAskForDataflow"), bGeometryCollectionOpenEditorAskForDataflow, TEXT("Whether to ask to create add a Dataflow graph when opening a Geometry Collection asset editor"));
+
 namespace UE::GeometryCollection
 {
 	struct FColorScheme
@@ -84,11 +87,14 @@ EAssetCommandResult UAssetDefinition_GeometryCollection::OpenAssets(const FAsset
 		// Validate the asset
 		if (UGeometryCollection* const GeometryCollection = Cast<UGeometryCollection>(GeometryCollectionObjects[0]))
 		{
-			if (!FDataflowEditorToolkit::HasDataflowAsset(GeometryCollection))
+			if (bGeometryCollectionOpenEditorAskForDataflow)
 			{
-				if (UDataflow* const NewDataflowAsset = Cast<UDataflow>(UE::DataflowAssetDefinitionHelpers::NewOrOpenDataflowAsset(GeometryCollection)))
+				if (!FDataflowEditorToolkit::HasDataflowAsset(GeometryCollection))
 				{
-					GeometryCollection->DataflowInstance.SetDataflowAsset(NewDataflowAsset);
+					if (UDataflow* const NewDataflowAsset = Cast<UDataflow>(UE::DataflowAssetDefinitionHelpers::NewOrOpenDataflowAsset(GeometryCollection)))
+					{
+						GeometryCollection->DataflowInstance.SetDataflowAsset(NewDataflowAsset);
+					}
 				}
 			}
 
