Add support for override materials for geometry collection root proxies
- updated dataflow node reading or writing the root proxies

diff --git a/Engine/Plugins/Experimental/GeometryCollectionPlugin/Source/GeometryCollectionNodes/Private/Dataflow/GeometryCollectionAssetNodes.cpp b/Engine/Plugins/Experimental/GeometryCollectionPlugin/Source/GeometryCollectionNodes/Private/Dataflow/GeometryCollectionAssetNodes.cpp
--- a/Engine/Plugins/Experimental/GeometryCollectionPlugin/Source/GeometryCollectionNodes/Private/Dataflow/GeometryCollectionAssetNodes.cpp
+++ b/Engine/Plugins/Experimental/GeometryCollectionPlugin/Source/GeometryCollectionNodes/Private/Dataflow/GeometryCollectionAssetNodes.cpp
@@ -46,6 +46,7 @@ FMakeRootProxyMeshDataflowNode::FMakeRootProxyMeshDataflowNode(const UE::Dataflo
 {
 	RegisterInputConnection(&Mesh);
 	RegisterInputConnection(&Transform);
+	RegisterInputConnection(&OverrideMaterials);
 
 	RegisterOutputConnection(&RootProxyMesh);
 }
@@ -57,6 +58,7 @@ void FMakeRootProxyMeshDataflowNode::Evaluate(UE::Dataflow::FContext& Context, c
 		FDataflowRootProxyMesh OutRootProxyMesh;
 		OutRootProxyMesh.Mesh = GetValue(Context, &Mesh);
 		OutRootProxyMesh.Transform = GetValue(Context, &Transform);
+		OutRootProxyMesh.OverrideMaterials = GetValue(Context, &OverrideMaterials);
 
 		SetValue(Context, MoveTemp(OutRootProxyMesh), &RootProxyMesh);
 	}
@@ -127,8 +129,7 @@ void FGeometryCollectionTerminalDataflowNode_v2::SetAssetValue(TObjectPtr<UObjec
 		// make sure to always reset the root proxies even if the collection is empty 
 		if (IsConnected(&RootProxyMeshes))
 		{
-			CollectionAsset->RootProxyData.ProxyMeshes.Reset();
-			CollectionAsset->RootProxyData.MeshTransforms.Reset();
+			CollectionAsset->RootProxyData.Reset();
 		}
 
 		if (FGeometryCollectionPtr GeometryCollection = CollectionAsset->GetGeometryCollection())
@@ -154,6 +155,7 @@ void FGeometryCollectionTerminalDataflowNode_v2::SetAssetValue(TObjectPtr<UObjec
 					Data.ProxyMeshes.Add(ProxyMesh.Mesh);
 					// root proxy transforms are relative to the root transform but FDataflowRootProxyMesh is in component space 
 					Data.MeshTransforms.Add(FTransform3f(ProxyMesh.Transform.GetRelativeTransform(RootTransform)));
+					Data.MeshOverrideMaterials.AddDefaulted_GetRef().Materials = ProxyMesh.OverrideMaterials;
 				}
 			}
 
@@ -473,6 +475,7 @@ void FCreateGeometryCollectionFromSourcesDataflowNode_v2::Evaluate(UE::Dataflow:
 		FDataflowRootProxyMesh RootProxyMesh;
 		RootProxyMesh.Mesh = StaticMeshObject;
 		RootProxyMesh.Transform = Source.LocalTransform;
+		RootProxyMesh.OverrideMaterials = Source.SourceMaterial;
 		OutRootProxyMeshes.Add(RootProxyMesh);
 
 	}
@@ -564,11 +567,16 @@ void FGeometryCollectionToCollectionDataflowNode_v2::Evaluate(UE::Dataflow::FCon
 			RootTransform = TransformFacade.GetRootTransform();
 		}
 
+		const FGeometryCollectionProxyMeshData& RootProxyData = GeometryCollection->RootProxyData;
 		for (int32 MeshIndex = 0; MeshIndex < GeometryCollection->RootProxyData.ProxyMeshes.Num(); MeshIndex++)
 		{
 			FDataflowRootProxyMesh RootProxyMesh;
-			RootProxyMesh.Mesh = GeometryCollection->RootProxyData.ProxyMeshes[MeshIndex];
-			RootProxyMesh.Transform = FTransform(GeometryCollection->RootProxyData.GetMeshTransform(MeshIndex)) * RootTransform;
+			RootProxyMesh.Mesh = RootProxyData.ProxyMeshes[MeshIndex];
+			RootProxyMesh.Transform = FTransform(RootProxyData.GetMeshTransform(MeshIndex)) * RootTransform;
+			if (RootProxyData.MeshOverrideMaterials.IsValidIndex(MeshIndex))
+			{
+				RootProxyMesh.OverrideMaterials = RootProxyData.MeshOverrideMaterials[MeshIndex].Materials;
+			}
 			OutRootProxyMeshes.Emplace(RootProxyMesh);
 		}
 	}
@@ -682,6 +690,7 @@ void FBlueprintToCollectionDataflowNode_v2::Evaluate(UE::Dataflow::FContext& Con
 								FDataflowRootProxyMesh RootProxyMesh;
 								RootProxyMesh.Mesh = StaticMeshObject;
 								RootProxyMesh.Transform = MeshTransform;
+								RootProxyMesh.OverrideMaterials = StaticMeshComponent->GetMaterials();
 								OutRootProxyMeshes.Add(RootProxyMesh);
 							}
 						}
diff --git a/Engine/Plugins/Experimental/GeometryCollectionPlugin/Source/GeometryCollectionNodes/Public/Dataflow/GeometryCollectionAssetNodes.h b/Engine/Plugins/Experimental/GeometryCollectionPlugin/Source/GeometryCollectionNodes/Public/Dataflow/GeometryCollectionAssetNodes.h
--- a/Engine/Plugins/Experimental/GeometryCollectionPlugin/Source/GeometryCollectionNodes/Public/Dataflow/GeometryCollectionAssetNodes.h
+++ b/Engine/Plugins/Experimental/GeometryCollectionPlugin/Source/GeometryCollectionNodes/Public/Dataflow/GeometryCollectionAssetNodes.h
@@ -29,6 +29,9 @@ struct FDataflowRootProxyMesh
 
 	UPROPERTY(EditAnywhere, Category = RootProxyMesh);
 	FTransform Transform;
+
+	UPROPERTY(EditAnywhere, Category = RootProxyMesh);
+	TArray<TObjectPtr<UMaterialInterface>> OverrideMaterials;
 };
 
 /**
@@ -50,6 +53,9 @@ private:
 	UPROPERTY(EditAnywhere, Category = RootProxyMesh, meta = (DataflowInput));
 	FTransform Transform;
 
+	UPROPERTY(EditAnywhere, Category = RootProxyMesh, meta = (DataflowInput));
+	TArray<TObjectPtr<UMaterialInterface>> OverrideMaterials;
+
 	/** mesh to use as a proxy */
 	UPROPERTY(meta = (DataflowOutput));
 	FDataflowRootProxyMesh RootProxyMesh;
diff --git a/Engine/Source/Runtime/Experimental/GeometryCollectionEngine/Private/GeometryCollection/GeometryCollectionISMPoolRenderer.cpp b/Engine/Source/Runtime/Experimental/GeometryCollectionEngine/Private/GeometryCollection/GeometryCollectionISMPoolRenderer.cpp
--- a/Engine/Source/Runtime/Experimental/GeometryCollectionEngine/Private/GeometryCollection/GeometryCollectionISMPoolRenderer.cpp
+++ b/Engine/Source/Runtime/Experimental/GeometryCollectionEngine/Private/GeometryCollection/GeometryCollectionISMPoolRenderer.cpp
@@ -142,7 +142,8 @@ UISMPoolComponent* UGeometryCollectionISMPoolRenderer::GetOrCreateISMPoolCompone
 
 void UGeometryCollectionISMPoolRenderer::InitMergedMeshFromGeometryCollection(UGeometryCollection const& InGeometryCollection)
 {
-	if (InGeometryCollection.RootProxyData.ProxyMeshes.Num() == 0)
+	const FGeometryCollectionProxyMeshData& RootProxyData = InGeometryCollection.RootProxyData;
+	if (RootProxyData.ProxyMeshes.Num() == 0)
 	{
 		return;
 	}
@@ -155,8 +156,10 @@ void UGeometryCollectionISMPoolRenderer::InitMergedMeshFromGeometryCollection(UG
 		return;
 	}
 
-	for (UStaticMesh* StaticMesh : InGeometryCollection.RootProxyData.ProxyMeshes)
+	const int32 NumProxyMesh = RootProxyData.ProxyMeshes.Num();
+	for (int32 MeshIndex = 0; MeshIndex < NumProxyMesh; ++MeshIndex)
 	{
+		UStaticMesh* StaticMesh = RootProxyData.ProxyMeshes[MeshIndex];
 		if (StaticMesh == nullptr)
 		{
 			continue;
@@ -165,6 +168,15 @@ void UGeometryCollectionISMPoolRenderer::InitMergedMeshFromGeometryCollection(UG
 		FISMPoolStaticMeshInstance StaticMeshInstance;
 		StaticMeshInstance.StaticMesh = StaticMesh;
 
+		if (RootProxyData.MeshOverrideMaterials.IsValidIndex(MeshIndex))
+		{
+			const FGeometryCollectionProxyMeshMaterials& ProxyMaterials = RootProxyData.MeshOverrideMaterials[MeshIndex];
+			for (int32 MatIndex = 0; MatIndex < ProxyMaterials.Materials.Num(); ++MatIndex)
+			{
+				StaticMeshInstance.MaterialsOverrides.Add(ProxyMaterials.Materials[MatIndex]);
+			}
+		}
+
 		TArray<float> DummyCustomData;
 		MergedMeshGroup.MeshIds.Add(ISMPoolComponent->AddMeshToGroup(MergedMeshGroup.GroupIndex, StaticMeshInstance, 1, DummyCustomData));
 	}
diff --git a/Engine/Source/Runtime/Experimental/GeometryCollectionEngine/Private/GeometryCollection/GeometryCollectionObject.cpp b/Engine/Source/Runtime/Experimental/GeometryCollectionEngine/Private/GeometryCollection/GeometryCollectionObject.cpp
--- a/Engine/Source/Runtime/Experimental/GeometryCollectionEngine/Private/GeometryCollection/GeometryCollectionObject.cpp
+++ b/Engine/Source/Runtime/Experimental/GeometryCollectionEngine/Private/GeometryCollection/GeometryCollectionObject.cpp
@@ -254,6 +254,12 @@ const FTransform3f& FGeometryCollectionProxyMeshData::GetMeshTransform(int32 Mes
 	return FTransform3f::Identity;
 }
 
+void FGeometryCollectionProxyMeshData::Reset()
+{
+	ProxyMeshes.Reset();
+	MeshTransforms.Reset();
+	MeshOverrideMaterials.Reset();
+}
 
 void FillSharedSimulationSizeSpecificData(FSharedSimulationSizeSpecificData& ToData, const FGeometryCollectionSizeSpecificData& FromData)
 {
@@ -1543,8 +1549,7 @@ void  UGeometryCollection::PropagateTransformUpdateToComponents() const
 
 void UGeometryCollection::SetRootProxiesFromGeometrySources()
 {
-	RootProxyData.ProxyMeshes.Reset();
-	RootProxyData.MeshTransforms.Reset();
+	RootProxyData.Reset();
 
 	// make sure the root index is properly cached
 	if (RootIndex == INDEX_NONE)
@@ -1567,6 +1572,9 @@ void UGeometryCollection::SetRootProxiesFromGeometrySources()
 
 			const FTransform3f MeshTransform = FTransform3f(Source.LocalTransform).GetRelativeTransform(RootTransform);
 			RootProxyData.MeshTransforms.Add(MeshTransform);
+
+			FGeometryCollectionProxyMeshMaterials& MeshMaterials = RootProxyData.MeshOverrideMaterials.AddDefaulted_GetRef();
+			MeshMaterials.Materials = Source.SourceMaterial;
 		}
 	}
 }
diff --git a/Engine/Source/Runtime/Experimental/GeometryCollectionEngine/Private/GeometryCollection/GeometryCollectionRootProxyRenderer.cpp b/Engine/Source/Runtime/Experimental/GeometryCollectionEngine/Private/GeometryCollection/GeometryCollectionRootProxyRenderer.cpp
--- a/Engine/Source/Runtime/Experimental/GeometryCollectionEngine/Private/GeometryCollection/GeometryCollectionRootProxyRenderer.cpp
+++ b/Engine/Source/Runtime/Experimental/GeometryCollectionEngine/Private/GeometryCollection/GeometryCollectionRootProxyRenderer.cpp
@@ -63,13 +63,16 @@ void UGeometryCollectionRootProxyRenderer::CreateRootProxyComponents(UGeometryCo
 		return;
 	}
 	ClearRootProxyComponents();
-	StaticMeshComponents.Reserve(GeometryCollection->RootProxyData.ProxyMeshes.Num());
 
-	for (UStaticMesh* ProxyMesh : GeometryCollection->RootProxyData.ProxyMeshes)
+	const FGeometryCollectionProxyMeshData& ProxyData = GeometryCollection->RootProxyData;
+	const int32 NumProxyMesh = ProxyData.ProxyMeshes.Num();
+	StaticMeshComponents.Reserve(NumProxyMesh);
+
+	for (int32 MeshIndex = 0; MeshIndex < NumProxyMesh; ++MeshIndex)
 	{
 		UStaticMeshComponent* MeshComponent = nullptr;
 		
-		if (ProxyMesh != nullptr)
+		if (UStaticMesh* ProxyMesh = ProxyData.ProxyMeshes[MeshIndex])
 		{
 			MeshComponent = NewObject<UStaticMeshComponent>(this, NAME_None, RF_DuplicateTransient | RF_Transient);
 
@@ -79,6 +82,15 @@ void UGeometryCollectionRootProxyRenderer::CreateRootProxyComponents(UGeometryCo
 			MeshComponent->SetMobility(InComponent.Mobility);
 			MeshComponent->SetupAttachment(&InComponent);
 			MeshComponent->RegisterComponent();
+
+			if (ProxyData.MeshOverrideMaterials.IsValidIndex(MeshIndex))
+			{
+				const FGeometryCollectionProxyMeshMaterials& ProxyMaterials = ProxyData.MeshOverrideMaterials[MeshIndex];
+				for (int32 MatIndex = 0; MatIndex < ProxyMaterials.Materials.Num(); ++MatIndex)
+				{
+					MeshComponent->SetMaterial(MatIndex, ProxyMaterials.Materials[MatIndex]);
+				}
+			}
 		}
 
 		StaticMeshComponents.Add(MeshComponent);
diff --git a/Engine/Source/Runtime/Experimental/GeometryCollectionEngine/Public/GeometryCollection/GeometryCollectionObject.h b/Engine/Source/Runtime/Experimental/GeometryCollectionEngine/Public/GeometryCollection/GeometryCollectionObject.h
--- a/Engine/Source/Runtime/Experimental/GeometryCollectionEngine/Public/GeometryCollection/GeometryCollectionObject.h
+++ b/Engine/Source/Runtime/Experimental/GeometryCollectionEngine/Public/GeometryCollection/GeometryCollectionObject.h
@@ -326,6 +326,15 @@ struct TStructOpsTypeTraits<FGeometryCollectionSizeSpecificData> : public TStruc
 	static constexpr EPropertyObjectReferenceType WithSerializerObjectReferences = EPropertyObjectReferenceType::None;
 };
 
+USTRUCT(BlueprintType)
+struct FGeometryCollectionProxyMeshMaterials
+{
+	GENERATED_BODY()
+
+	UPROPERTY(EditAnywhere, BlueprintReadWrite, Category = "Materials")
+	TArray<TObjectPtr<UMaterialInterface>> Materials;
+};
+
 USTRUCT(BlueprintType)
 struct FGeometryCollectionProxyMeshData
 {
@@ -343,7 +352,16 @@ struct FGeometryCollectionProxyMeshData
 	UPROPERTY(EditAnywhere, BlueprintReadWrite, Category = "Rendering")
 	TArray<FTransform3f> MeshTransforms;
 
+	/**
+	* Optional arrays of override material matching each proxy Mesh
+	* if an entry is missing or if the arary is empty , then no material will be assigned
+	*/
+	UPROPERTY(EditAnywhere, BlueprintReadWrite, Category = "Rendering")
+	TArray<FGeometryCollectionProxyMeshMaterials> MeshOverrideMaterials;
+
 	GEOMETRYCOLLECTIONENGINE_API const FTransform3f& GetMeshTransform(int32 MeshIndex) const;
+
+	GEOMETRYCOLLECTIONENGINE_API void Reset();
 };
 
 USTRUCT()
