add a RemoveMeshOverlaps dataflow node; useful for creating a nice input for the mesh cut fracture node

diff --git a/Engine/Plugins/Experimental/GeometryCollectionPlugin/Source/GeometryCollectionNodes/Private/Dataflow/GeometryCollectionNodesPlugin.cpp b/Engine/Plugins/Experimental/GeometryCollectionPlugin/Source/GeometryCollectionNodes/Private/Dataflow/GeometryCollectionNodesPlugin.cpp
--- a/Engine/Plugins/Experimental/GeometryCollectionPlugin/Source/GeometryCollectionNodes/Private/Dataflow/GeometryCollectionNodesPlugin.cpp
+++ b/Engine/Plugins/Experimental/GeometryCollectionPlugin/Source/GeometryCollectionNodes/Private/Dataflow/GeometryCollectionNodesPlugin.cpp
@@ -25,6 +25,7 @@
 #include "Dataflow/GeometryCollectionProcessingNodes.h"
 #include "Dataflow/GeometryCollectionSelectionNodes.h"
 #include "Dataflow/GeometryCollectionSamplingNodes.h"
+#include "Dataflow/GeometryCollectionRemoveMeshOverlapsNode.h"
 #include "Dataflow/GeometryCollectionSkeletalMeshToCollectionNode.h"
 #include "Dataflow/GeometryCollectionSkeletonToCollectionNode.h"
 #include "Dataflow/GeometryCollectionStaticMeshToCollectionNode.h"
@@ -66,6 +67,7 @@ void IGeometryCollectionNodesPlugin::StartupModule()
 	UE::Dataflow::GeometryCollectionArrayNodes();
 	UE::Dataflow::GeometryCollectionDebugNodes();
 	UE::Dataflow::GeometryCollectionSamplingNodes();
+	UE::Dataflow::GeometryCollectionRemoveMeshOverlapsNode();
 	UE::Dataflow::RegisterGeometryCollectionUVNodes();
 	UE::Dataflow::RegisterGeometryCollectionTextureNodes();
 	UE::Dataflow::RegisterGeometryCollectionMaterialInterfaceNodes();
diff --git a/Engine/Plugins/Experimental/GeometryCollectionPlugin/Source/GeometryCollectionNodes/Private/Dataflow/GeometryCollectionRemoveMeshOverlapsNode.cpp b/Engine/Plugins/Experimental/GeometryCollectionPlugin/Source/GeometryCollectionNodes/Private/Dataflow/GeometryCollectionRemoveMeshOverlapsNode.cpp
new file mode 100644
--- /dev/null
+++ b/Engine/Plugins/Experimental/GeometryCollectionPlugin/Source/GeometryCollectionNodes/Private/Dataflow/GeometryCollectionRemoveMeshOverlapsNode.cpp
@@ -0,0 +1,137 @@
+// Copyright Epic Games, Inc. All Rights Reserved.
+
+
+#include "Dataflow/GeometryCollectionRemoveMeshOverlapsNode.h"
+#include "Dataflow/DataflowCore.h"
+
+#include "Algo/Reverse.h"
+#include "Operations/MeshBoolean.h"
+#include "DynamicMesh/MeshTransforms.h"
+
+#include "DynamicMesh/DynamicMesh3.h"
+#include "UDynamicMesh.h"
+
+#include UE_INLINE_GENERATED_CPP_BY_NAME(GeometryCollectionRemoveMeshOverlapsNode)
+
+namespace UE::Dataflow
+{
+
+	void GeometryCollectionRemoveMeshOverlapsNode()
+	{
+		DATAFLOW_NODE_REGISTER_CREATION_FACTORY(FRemoveMeshOverlapsDataflowNode);
+	}
+}
+
+
+FRemoveMeshOverlapsDataflowNode::FRemoveMeshOverlapsDataflowNode(const UE::Dataflow::FNodeParameters& InParam, FGuid InGuid)
+	: FDataflowNode(InParam, InGuid)
+{
+	RegisterInputConnection(&DynamicMeshes);
+	RegisterOutputConnection(&DynamicMeshes, &DynamicMeshes);
+	RegisterInputConnection(&PerMeshTransforms).SetCanHidePin(true).SetPinIsHidden(true);
+}
+
+void FRemoveMeshOverlapsDataflowNode::Evaluate(UE::Dataflow::FContext& Context, const FDataflowOutput* Out) const
+{
+	using namespace UE::Geometry;
+	
+	if (Out->IsA(&DynamicMeshes) ||
+		Out->IsA(&PerMeshTransforms))
+	{
+		const TArray<TObjectPtr<UDynamicMesh>>& InMeshes = GetValue(Context, &DynamicMeshes);
+		const TArray<FTransform>& InPerMeshTransforms = GetValue(Context, &PerMeshTransforms);
+
+		TArray<TObjectPtr<UDynamicMesh>> ResultMeshes;
+		TArray<UE::Geometry::FAxisAlignedBox3d> MeshBounds;
+
+		// Define ordering to apply boolean subtraction operations
+		TArray<int32> Ordering = [&InMeshes, &ResultMeshes, this]() -> TArray<int32>
+		{
+			TArray<int32> Indices;
+			Indices.SetNumUninitialized(InMeshes.Num());
+			for (int32 Idx = 0; Idx < InMeshes.Num(); ++Idx)
+			{
+				Indices[Idx] = Idx;
+			}
+			bool bVolumeOrder = false, bReverseOrder = false;
+			switch (SubtractOrder)
+			{
+			case EDataflowRemoveOverlapsMeshSortOrder::IncreasingVolume:
+				bVolumeOrder = true;
+				break;
+			case EDataflowRemoveOverlapsMeshSortOrder::DecreasingVolume:
+				bReverseOrder = bVolumeOrder = true;
+				break;
+			case EDataflowRemoveOverlapsMeshSortOrder::ReverseInputOrder:
+				bReverseOrder = true;
+				break;
+			case EDataflowRemoveOverlapsMeshSortOrder::InputOrder:
+				// defaults apply
+				break;
+			}
+			if (bVolumeOrder)
+			{
+				TArray<float> Values;
+				Values.SetNumUninitialized(InMeshes.Num());
+				for (int32 Idx = 0; Idx < InMeshes.Num(); ++Idx)
+				{
+					float Value = 0;
+					InMeshes[Idx]->ProcessMesh([&Value](const FDynamicMesh3& Mesh) -> void { Value = TMeshQueries<FDynamicMesh3>::GetVolumeNonWatertight(Mesh); });
+					Values[Idx] = Value;
+				}
+				Indices.Sort([&Values](int32 A, int32 B) { return Values[A] < Values[B]; });
+			}
+			if (bReverseOrder)
+			{
+				Algo::Reverse(Indices);
+			}
+			return Indices;
+		}();
+
+		// Copy initial meshes to results array, bake transforms, and compute initial bounds
+		ResultMeshes.SetNum(InMeshes.Num());
+		MeshBounds.Init(UE::Geometry::FAxisAlignedBox3d::Empty(), InMeshes.Num());
+		for (int32 MeshIdx = 0; MeshIdx < InMeshes.Num(); ++MeshIdx)
+		{
+			ResultMeshes[MeshIdx] = NewObject<UDynamicMesh>();
+			ResultMeshes[MeshIdx]->EditMesh([&InMeshes, &InPerMeshTransforms, &MeshBounds, MeshIdx](FDynamicMesh3& Mesh) -> void
+				{
+					InMeshes[MeshIdx]->ProcessMesh([&Mesh](const FDynamicMesh3& InMesh) -> void
+					{
+						Mesh = InMesh;
+					});
+					if (InPerMeshTransforms.IsValidIndex(MeshIdx))
+					{
+						MeshTransforms::ApplyTransform(Mesh, InPerMeshTransforms[MeshIdx]);
+					}
+					MeshBounds[MeshIdx] = Mesh.GetBounds();
+				});
+		}
+		// Compute mesh boolean subtractions to remove overlaps
+		if (InMeshes.Num() > 1)
+		{
+			ParallelFor(InMeshes.Num() - 1, 
+			[&Ordering, &ResultMeshes, &MeshBounds, &InMeshes, &InPerMeshTransforms, this](int32 FromOrderingIdx)
+			{
+				const int32 FromMeshIdx = Ordering[FromOrderingIdx];
+				UE::Geometry::FDynamicMesh3& ToTrim = ResultMeshes[FromMeshIdx]->GetMeshRef();
+				for (int32 VsOrderingIdx = FromOrderingIdx + 1; VsOrderingIdx < InMeshes.Num(); ++VsOrderingIdx)
+				{
+					const int32 VsMeshIdx = Ordering[VsOrderingIdx];
+					if (MeshBounds[FromMeshIdx].Intersects(MeshBounds[VsMeshIdx]))
+					{
+						FTransform VsTransform = InPerMeshTransforms.IsValidIndex(VsMeshIdx) ? InPerMeshTransforms[VsMeshIdx] : FTransform::Identity;
+						UE::Geometry::FMeshBoolean Boolean(
+							&ToTrim, FTransform::Identity, 
+							InMeshes[VsMeshIdx]->GetMeshPtr(), VsTransform,
+							&ToTrim, UE::Geometry::FMeshBoolean::EBooleanOp::Difference);
+						Boolean.bSimplifyAlongNewEdges = bSimplifyNewMeshEdges;
+						Boolean.Compute();
+					}
+				}
+			}, EParallelForFlags::Unbalanced);
+		}
+
+		SetValue(Context, ResultMeshes, &DynamicMeshes);
+	}
+}
diff --git a/Engine/Plugins/Experimental/GeometryCollectionPlugin/Source/GeometryCollectionNodes/Private/Dataflow/GeometryCollectionRemoveMeshOverlapsNode.h b/Engine/Plugins/Experimental/GeometryCollectionPlugin/Source/GeometryCollectionNodes/Private/Dataflow/GeometryCollectionRemoveMeshOverlapsNode.h
new file mode 100644
--- /dev/null
+++ b/Engine/Plugins/Experimental/GeometryCollectionPlugin/Source/GeometryCollectionNodes/Private/Dataflow/GeometryCollectionRemoveMeshOverlapsNode.h
@@ -0,0 +1,64 @@
+// Copyright Epic Games, Inc. All Rights Reserved.
+
+#pragma once 
+
+#include "Dataflow/DataflowEngine.h"
+
+#include "GeometryCollectionRemoveMeshOverlapsNode.generated.h"
+
+class FGeometryCollection;
+class UDynamicMesh;
+
+// The mesh ordering for used when removing overlaps
+UENUM()
+enum class EDataflowRemoveOverlapsMeshSortOrder
+{
+	// Process in order of increasing volume; largest is unchanged
+	IncreasingVolume,
+	// Process in order of decreasing volume; smallest is unchanged
+	DecreasingVolume,
+	// Process in input ordering; last mesh is unchanged
+	InputOrder,
+	// Process in the reverse of the input ordering; first mesh is unchanged
+	ReverseInputOrder
+};
+
+/**
+ * Cut away any overlaps between input meshes, creating a set of non-overlapping meshes as output
+ * For example, given meshes w/ order [A, B, C], non-overlapping meshes will be [(A-B)-C, B-C, C]
+ */
+USTRUCT(meta = (DataflowGeometryCollection))
+struct FRemoveMeshOverlapsDataflowNode : public FDataflowNode
+{
+	GENERATED_USTRUCT_BODY()
+	DATAFLOW_NODE_DEFINE_INTERNAL(FRemoveMeshOverlapsDataflowNode, "RemoveMeshOverlaps", "GeometryCollection|Fracture", "")
+
+public:
+	FRemoveMeshOverlapsDataflowNode(const UE::Dataflow::FNodeParameters& InParam, FGuid InGuid = FGuid::NewGuid());
+
+private:
+
+	virtual void Evaluate(UE::Dataflow::FContext& Context, const FDataflowOutput* Out) const override;
+
+	/** Dynamic Meshes to remove overlaps from */
+	UPROPERTY(EditAnywhere, Category = "Meshes", meta = (DataflowInput, DataflowOutput, DataflowPassthrough = "DynamicMeshes", DataflowIntrinsic))
+	TArray<TObjectPtr<UDynamicMesh>> DynamicMeshes;
+
+	/** Optional transforms for input meshes. Will be baked into output. If not provided, meshes will not be transformed. */
+	UPROPERTY(EditAnywhere, Category = "Meshes", meta = (DataflowInput))
+	TArray<FTransform> PerMeshTransforms;
+
+	/** Ordering to use when applying subtraction operations. Later meshes will be subtracted from earlier meshes in the ordering. */
+	UPROPERTY(EditAnywhere, Category = "Ordering")
+	EDataflowRemoveOverlapsMeshSortOrder SubtractOrder = EDataflowRemoveOverlapsMeshSortOrder::InputOrder;
+
+	/** Whether to simplify along new edges created by the mesh subtraction operations */
+	UPROPERTY(EditAnywhere, Category = "Boolean")
+	bool bSimplifyNewMeshEdges = true;
+};
+
+namespace UE::Dataflow
+{
+	void GeometryCollectionRemoveMeshOverlapsNode();
+}
+
