Dataflow:
- Added GetNumElementsInSelection node

diff --git a/Engine/Plugins/Experimental/GeometryCollectionPlugin/Source/GeometryCollectionNodes/Private/Dataflow/GeometryCollectionSelectionNodes.cpp b/Engine/Plugins/Experimental/GeometryCollectionPlugin/Source/GeometryCollectionNodes/Private/Dataflow/GeometryCollectionSelectionNodes.cpp
--- a/Engine/Plugins/Experimental/GeometryCollectionPlugin/Source/GeometryCollectionNodes/Private/Dataflow/GeometryCollectionSelectionNodes.cpp
+++ b/Engine/Plugins/Experimental/GeometryCollectionPlugin/Source/GeometryCollectionNodes/Private/Dataflow/GeometryCollectionSelectionNodes.cpp
@@ -79,7 +79,8 @@ namespace UE::Dataflow
 		DATAFLOW_NODE_REGISTER_CREATION_FACTORY(FCollectionSelectInternalFacesDataflowNode);
 		DATAFLOW_NODE_REGISTER_CREATION_FACTORY(FCollectionSelectTransformStringDataflowNode);
 		DATAFLOW_NODE_REGISTER_CREATION_FACTORY(FCollectionSetTransformStringValueDataflowNode);
-			;
+		DATAFLOW_NODE_REGISTER_CREATION_FACTORY(FGetNumElementsInSelectionDataflowNode);
+		
 
 		// generic input nodes
 		DATAFLOW_NODE_REGISTER_CREATION_FACTORY(FCollectionSelectionSetOperationDataflowNode);
@@ -1971,4 +1972,67 @@ void FCollectionSetTransformStringValueDataflowNode::Evaluate(UE::Dataflow::FCon
 		}
 		SetValue(Context, MoveTemp(CollectionToUpdate), &Collection);
 	}
-}
\ No newline at end of file
+}
+
+/* ----------------------------------------------------------------------------------------------------------------------- */
+
+FGetNumElementsInSelectionDataflowNode::FGetNumElementsInSelectionDataflowNode(const UE::Dataflow::FNodeParameters& InParam, FGuid InGuid)
+	: FDataflowNode(InParam, InGuid)
+{
+	static const FName MainTypeGroup("Main");
+
+	// Inputs
+	RegisterInputConnection(&Selection)
+		.SetTypeDependencyGroup(MainTypeGroup);
+
+	// Outputs
+	RegisterOutputConnection(&Selection, &Selection)
+		.SetTypeDependencyGroup(MainTypeGroup);
+
+	RegisterOutputConnection(&NumElements);
+	RegisterOutputConnection(&NumSelectedElements);
+}
+
+void FGetNumElementsInSelectionDataflowNode::Evaluate(UE::Dataflow::FContext& Context, const FDataflowOutput* Out) const
+{
+	if (Out->IsA(&Selection))
+	{
+		SafeForwardInput(Context, &Selection, &Selection);
+	}
+	else if (Out->IsA(&NumElements))
+	{
+		if (IsConnected(&Selection))
+		{
+			const FDataflowInput* SelectionInput = FindInput(&Selection);
+			if (SelectionInput)
+			{
+				const FDataflowSelection& InSelection = GetValue(Context, &Selection);
+
+				SetValue(Context, InSelection.Num(), &NumElements);
+
+				return;
+			}
+		}
+
+		SetValue(Context, 0, &NumElements);
+	}
+	else if (Out->IsA(&NumSelectedElements))
+	{
+		if (IsConnected(&Selection))
+		{
+			const FDataflowInput* SelectionInput = FindInput(&Selection);
+			if (SelectionInput)
+			{
+				const FDataflowSelection& InSelection = GetValue(Context, &Selection);
+
+				SetValue(Context, InSelection.NumSelected(), &NumSelectedElements);
+
+				return;
+			}
+		}
+
+		SetValue(Context, 0, &NumSelectedElements);
+	}
+}
+
+/* ----------------------------------------------------------------------------------------------------------------------- */
diff --git a/Engine/Plugins/Experimental/GeometryCollectionPlugin/Source/GeometryCollectionNodes/Public/Dataflow/GeometryCollectionSelectionNodes.h b/Engine/Plugins/Experimental/GeometryCollectionPlugin/Source/GeometryCollectionNodes/Public/Dataflow/GeometryCollectionSelectionNodes.h
--- a/Engine/Plugins/Experimental/GeometryCollectionPlugin/Source/GeometryCollectionNodes/Public/Dataflow/GeometryCollectionSelectionNodes.h
+++ b/Engine/Plugins/Experimental/GeometryCollectionPlugin/Source/GeometryCollectionNodes/Public/Dataflow/GeometryCollectionSelectionNodes.h
@@ -1740,6 +1740,36 @@ public:
 	FCollectionSetTransformStringValueDataflowNode(const UE::Dataflow::FNodeParameters& InParam, FGuid InGuid = FGuid::NewGuid());
 };
 
+
+/**
+ *
+ * Returns number of elements in a Selection
+ *
+ */
+USTRUCT(meta = (DataflowGeometryCollection))
+struct FGetNumElementsInSelectionDataflowNode : public FDataflowNode
+{
+	GENERATED_USTRUCT_BODY()
+	DATAFLOW_NODE_DEFINE_INTERNAL(FGetNumElementsInSelectionDataflowNode, "GetNumElementsInSelection", "GeometryCollection|Utilities", "")
+
+private:
+	UPROPERTY(meta = (DataflowInput, DataflowOutput, DataflowPassthrough = "Selection", DataflowIntrinsic))
+	FDataflowSelectionTypes Selection;
+
+	/** Number of all (selected and non-selected) elements in the Selection */
+	UPROPERTY(meta = (DataflowOutput));
+	int32 NumElements = 0;
+
+	/** Number of selected elements in the Selection */
+	UPROPERTY(meta = (DataflowOutput));
+	int32 NumSelectedElements = 0;
+
+	virtual void Evaluate(UE::Dataflow::FContext& Context, const FDataflowOutput* Out) const override;
+
+public:
+	FGetNumElementsInSelectionDataflowNode(const UE::Dataflow::FNodeParameters& InParam, FGuid InGuid = FGuid::NewGuid());
+};
+
 namespace UE::Dataflow
 {
 	void GeometryCollectionSelectionNodes();
