Dataflow:
- Added CollectionInfo node to compute MaxLevel

diff --git a/Engine/Plugins/Experimental/GeometryCollectionPlugin/Source/GeometryCollectionNodes/Private/Dataflow/GeometryCollectionNodes.cpp b/Engine/Plugins/Experimental/GeometryCollectionPlugin/Source/GeometryCollectionNodes/Private/Dataflow/GeometryCollectionNodes.cpp
--- a/Engine/Plugins/Experimental/GeometryCollectionPlugin/Source/GeometryCollectionNodes/Private/Dataflow/GeometryCollectionNodes.cpp
+++ b/Engine/Plugins/Experimental/GeometryCollectionPlugin/Source/GeometryCollectionNodes/Private/Dataflow/GeometryCollectionNodes.cpp
@@ -97,6 +97,7 @@ namespace UE::Dataflow
 		DATAFLOW_NODE_REGISTER_CREATION_FACTORY(FPointsToCollectionDataflowNode);
 		DATAFLOW_NODE_REGISTER_CREATION_FACTORY(FCollectionToPointsDataflowNode);
 		DATAFLOW_NODE_REGISTER_CREATION_FACTORY(FSpheresToPointsDataflowNode);
+		DATAFLOW_NODE_REGISTER_CREATION_FACTORY(FCollectionInfoDataflowNode);
 
 		// Deprecated
 		DATAFLOW_NODE_REGISTER_CREATION_FACTORY(FStringAppendDataflowNode);
@@ -2033,4 +2034,35 @@ void FSpheresToPointsDataflowNode::Evaluate(UE::Dataflow::FContext& Context, con
 		SetValue(Context, TArray<FVector>(), &Points);
 		SetValue(Context, TArray<float>(), &Radii);
 	}
+}
+
+void FCollectionInfoDataflowNode::Evaluate(UE::Dataflow::FContext& Context, const FDataflowOutput* Out) const
+{
+	if (Out->IsA(&Collection))
+	{
+		SafeForwardInput(Context, &Collection, &Collection);
+	}
+	else if (Out->IsA(&MaxLevel))
+	{
+		if (IsConnected(&Collection))
+		{
+			FManagedArrayCollection InCollection = GetValue(Context, &Collection);
+
+			Chaos::Facades::FCollectionHierarchyFacade HierarchyFacade(InCollection);
+			if (HierarchyFacade.IsValid())
+			{
+				if (!HierarchyFacade.HasLevelAttribute())
+				{
+					HierarchyFacade.GenerateLevelAttribute();
+				}
+
+				int32 MaxLevelAttr = HierarchyFacade.ComputeLevelAttributeMax();
+
+				SetValue(Context, MaxLevelAttr, &MaxLevel);
+				return;
+			}
+		}
+
+		SetValue(Context, -1, &MaxLevel);
+	}
 }
\ No newline at end of file
diff --git a/Engine/Plugins/Experimental/GeometryCollectionPlugin/Source/GeometryCollectionNodes/Public/Dataflow/GeometryCollectionNodes.h b/Engine/Plugins/Experimental/GeometryCollectionPlugin/Source/GeometryCollectionNodes/Public/Dataflow/GeometryCollectionNodes.h
--- a/Engine/Plugins/Experimental/GeometryCollectionPlugin/Source/GeometryCollectionNodes/Public/Dataflow/GeometryCollectionNodes.h
+++ b/Engine/Plugins/Experimental/GeometryCollectionPlugin/Source/GeometryCollectionNodes/Public/Dataflow/GeometryCollectionNodes.h
@@ -2111,6 +2111,38 @@ public:
 	virtual void Evaluate(UE::Dataflow::FContext& Context, const FDataflowOutput* Out) const override;
 };
 
+/**
+ *
+ * Outputs Collection info
+ *
+ */
+USTRUCT(meta = (DataflowGeometryCollection))
+struct FCollectionInfoDataflowNode : public FDataflowNode
+{
+	GENERATED_USTRUCT_BODY()
+	DATAFLOW_NODE_DEFINE_INTERNAL(FCollectionInfoDataflowNode, "CollectionInfo", "GeometryCollection|Utilities", "")
+
+private:
+	/** Collection storing the points */
+	UPROPERTY(meta = (DataflowInput, DataflowOutput, DataflowPassthrough = "Collection", DataflowIntrinsic))
+	FManagedArrayCollection Collection;
+
+	/** MaxLevel */
+	UPROPERTY(meta = (DataflowOutput))
+	int32 MaxLevel;
+
+public:
+	FCollectionInfoDataflowNode(const UE::Dataflow::FNodeParameters& InParam, FGuid InGuid = FGuid::NewGuid())
+		: FDataflowNode(InParam, InGuid)
+	{
+		RegisterInputConnection(&Collection);
+		RegisterOutputConnection(&Collection, &Collection);
+		RegisterOutputConnection(&MaxLevel);
+	}
+
+	virtual void Evaluate(UE::Dataflow::FContext& Context, const FDataflowOutput* Out) const override;
+};
+
 namespace UE::Dataflow
 {
 	void GeometryCollectionEngineNodes();
diff --git a/Engine/Source/Runtime/Experimental/Chaos/Private/GeometryCollection/Facades/CollectionHierarchyFacade.cpp b/Engine/Source/Runtime/Experimental/Chaos/Private/GeometryCollection/Facades/CollectionHierarchyFacade.cpp
--- a/Engine/Source/Runtime/Experimental/Chaos/Private/GeometryCollection/Facades/CollectionHierarchyFacade.cpp
+++ b/Engine/Source/Runtime/Experimental/Chaos/Private/GeometryCollection/Facades/CollectionHierarchyFacade.cpp
@@ -271,4 +271,25 @@ namespace Chaos::Facades
 		}
 		return BreadthFirstIndices;
 	}
+
+	int32 FCollectionHierarchyFacade::ComputeLevelAttributeMax() const
+	{
+		if (LevelAttribute.IsValid())
+		{
+			const TManagedArray<int32>& Level = LevelAttribute.Get();
+
+			int32 LevelMax = -1;
+			for (int Index = 0; Index < Level.Num(); Index++)
+			{
+				if (Level[Index] > LevelMax)
+				{
+					LevelMax = Level[Index];
+				}
+			}
+
+			return LevelMax;
+		}
+
+		return -1;
+	}
 }
\ No newline at end of file
diff --git a/Engine/Source/Runtime/Experimental/Chaos/Public/GeometryCollection/Facades/CollectionHierarchyFacade.h b/Engine/Source/Runtime/Experimental/Chaos/Public/GeometryCollection/Facades/CollectionHierarchyFacade.h
--- a/Engine/Source/Runtime/Experimental/Chaos/Public/GeometryCollection/Facades/CollectionHierarchyFacade.h
+++ b/Engine/Source/Runtime/Experimental/Chaos/Public/GeometryCollection/Facades/CollectionHierarchyFacade.h
@@ -96,6 +96,8 @@ namespace Chaos::Facades
 		*/
 		CHAOS_API const TArray<int32>& GetBreadthFirstOrderIndicesFromAttribute() const;
 
+		CHAOS_API int32 ComputeLevelAttributeMax() const;
+
 	public:
 		/** Get the root indicies */
 		static CHAOS_API TArray<int32> GetRootIndices(const TManagedArrayAccessor<int32>& ParentAttribute);
