Dataflow : add geometry collection asset to GeometryCollectionToCollectionNode
- also deprecated the collection from asset node as this is a duplicate from GeometryCollectionToCollectionNode

diff --git a/Engine/Plugins/Experimental/GeometryCollectionPlugin/Source/GeometryCollectionNodes/Private/Dataflow/GeometryCollectionAssetNodes.cpp b/Engine/Plugins/Experimental/GeometryCollectionPlugin/Source/GeometryCollectionNodes/Private/Dataflow/GeometryCollectionAssetNodes.cpp
--- a/Engine/Plugins/Experimental/GeometryCollectionPlugin/Source/GeometryCollectionNodes/Private/Dataflow/GeometryCollectionAssetNodes.cpp
+++ b/Engine/Plugins/Experimental/GeometryCollectionPlugin/Source/GeometryCollectionNodes/Private/Dataflow/GeometryCollectionAssetNodes.cpp
@@ -541,6 +541,8 @@ void FGeometryCollectionToCollectionDataflowNode::Evaluate(UE::Dataflow::FContex
 FGeometryCollectionToCollectionDataflowNode_v2::FGeometryCollectionToCollectionDataflowNode_v2(const UE::Dataflow::FNodeParameters& InParam, FGuid InGuid)
 	: FDataflowNode(InParam, InGuid)
 {
+	RegisterInputConnection(&GeometryCollection);
+
 	RegisterOutputConnection(&Collection);
 	RegisterOutputConnection(&Materials);
 	RegisterOutputConnection(&InstancedMeshes);
@@ -556,19 +558,20 @@ void FGeometryCollectionToCollectionDataflowNode_v2::Evaluate(UE::Dataflow::FCon
 	TArray<FGeometryCollectionAutoInstanceMesh> OutInstancedMeshes;
 	TArray<FDataflowRootProxyMesh> OutRootProxyMeshes;
 
-	if (GeometryCollection)
+	const TObjectPtr<UGeometryCollection> InGeometryCollection = GetValue(Context, &GeometryCollection);
+	if (InGeometryCollection)
 	{
-		FGeometryCollectionEngineConversion::ConvertGeometryCollectionToGeometryCollection(GeometryCollection, OutCollection, OutMaterials, OutInstancedMeshes);
+		FGeometryCollectionEngineConversion::ConvertGeometryCollectionToGeometryCollection(InGeometryCollection, OutCollection, OutMaterials, OutInstancedMeshes);
 
 		FTransform RootTransform = FTransform::Identity;
-		if (GeometryCollection->GetGeometryCollection())
+		if (InGeometryCollection->GetGeometryCollection())
 		{
-			const GeometryCollection::Facades::FCollectionTransformFacade TransformFacade(*GeometryCollection->GetGeometryCollection());
+			const GeometryCollection::Facades::FCollectionTransformFacade TransformFacade(*InGeometryCollection->GetGeometryCollection());
 			RootTransform = TransformFacade.GetRootTransform();
 		}
 
-		const FGeometryCollectionProxyMeshData& RootProxyData = GeometryCollection->RootProxyData;
-		for (int32 MeshIndex = 0; MeshIndex < GeometryCollection->RootProxyData.ProxyMeshes.Num(); MeshIndex++)
+		const FGeometryCollectionProxyMeshData& RootProxyData = InGeometryCollection->RootProxyData;
+		for (int32 MeshIndex = 0; MeshIndex < InGeometryCollection->RootProxyData.ProxyMeshes.Num(); MeshIndex++)
 		{
 			FDataflowRootProxyMesh RootProxyMesh;
 			RootProxyMesh.Mesh = RootProxyData.ProxyMeshes[MeshIndex];
diff --git a/Engine/Plugins/Experimental/GeometryCollectionPlugin/Source/GeometryCollectionNodes/Private/Dataflow/GeometryCollectionNodes.cpp b/Engine/Plugins/Experimental/GeometryCollectionPlugin/Source/GeometryCollectionNodes/Private/Dataflow/GeometryCollectionNodes.cpp
--- a/Engine/Plugins/Experimental/GeometryCollectionPlugin/Source/GeometryCollectionNodes/Private/Dataflow/GeometryCollectionNodes.cpp
+++ b/Engine/Plugins/Experimental/GeometryCollectionPlugin/Source/GeometryCollectionNodes/Private/Dataflow/GeometryCollectionNodes.cpp
@@ -52,7 +52,6 @@ namespace UE::Dataflow
 {
 	void GeometryCollectionEngineNodes()
 	{
-		DATAFLOW_NODE_REGISTER_CREATION_FACTORY(FGetCollectionFromAssetDataflowNode);
 		DATAFLOW_NODE_REGISTER_CREATION_FACTORY(FAppendCollectionAssetsDataflowNode);
 		DATAFLOW_NODE_REGISTER_CREATION_FACTORY(FPrintStringDataflowNode);
 		DATAFLOW_NODE_REGISTER_CREATION_FACTORY(FLogStringDataflowNode);
@@ -99,11 +98,13 @@ namespace UE::Dataflow
 		DATAFLOW_NODE_REGISTER_CREATION_FACTORY(FSpheresToPointsDataflowNode);
 		DATAFLOW_NODE_REGISTER_CREATION_FACTORY(FCollectionInfoDataflowNode);
 
-		// Deprecated
+		// Deprecated nodes
+		DATAFLOW_NODE_REGISTER_CREATION_FACTORY(FGetCollectionFromAssetDataflowNode);
 		DATAFLOW_NODE_REGISTER_CREATION_FACTORY(FStringAppendDataflowNode);
 	}
 }
 
+// DEPRECATED 5.7
 void FGetCollectionFromAssetDataflowNode::Evaluate(UE::Dataflow::FContext& Context, const FDataflowOutput* Out) const
 {
 	if (Out->IsA<FManagedArrayCollection>(&Collection))
diff --git a/Engine/Plugins/Experimental/GeometryCollectionPlugin/Source/GeometryCollectionNodes/Public/Dataflow/GeometryCollectionAssetNodes.h b/Engine/Plugins/Experimental/GeometryCollectionPlugin/Source/GeometryCollectionNodes/Public/Dataflow/GeometryCollectionAssetNodes.h
--- a/Engine/Plugins/Experimental/GeometryCollectionPlugin/Source/GeometryCollectionNodes/Public/Dataflow/GeometryCollectionAssetNodes.h
+++ b/Engine/Plugins/Experimental/GeometryCollectionPlugin/Source/GeometryCollectionNodes/Public/Dataflow/GeometryCollectionAssetNodes.h
@@ -333,14 +333,14 @@ USTRUCT(meta = (DataflowContext = "GeometryCollection", DataflowGeometryCollecti
 struct FGeometryCollectionToCollectionDataflowNode_v2 : public FDataflowNode
 {
 	GENERATED_USTRUCT_BODY()
-	DATAFLOW_NODE_DEFINE_INTERNAL(FGeometryCollectionToCollectionDataflowNode_v2, "GeometryCollectionToCollection", "GeometryCollection|Asset", "")
+	DATAFLOW_NODE_DEFINE_INTERNAL(FGeometryCollectionToCollectionDataflowNode_v2, "GeometryCollectionToCollection", "GeometryCollection|Asset", "GC Convert")
 
 public:
 	FGeometryCollectionToCollectionDataflowNode_v2(const UE::Dataflow::FNodeParameters& InParam, FGuid InGuid = FGuid::NewGuid());
 
 private:
 	/** Asset input */
-	UPROPERTY(EditAnywhere, Category = "Asset");
+	UPROPERTY(EditAnywhere, Category = "Asset", meta = (DataflowInput));
 	TObjectPtr<UGeometryCollection> GeometryCollection;
 
 	/** Geometry collection newly created */
diff --git a/Engine/Plugins/Experimental/GeometryCollectionPlugin/Source/GeometryCollectionNodes/Public/Dataflow/GeometryCollectionNodes.h b/Engine/Plugins/Experimental/GeometryCollectionPlugin/Source/GeometryCollectionNodes/Public/Dataflow/GeometryCollectionNodes.h
--- a/Engine/Plugins/Experimental/GeometryCollectionPlugin/Source/GeometryCollectionNodes/Public/Dataflow/GeometryCollectionNodes.h
+++ b/Engine/Plugins/Experimental/GeometryCollectionPlugin/Source/GeometryCollectionNodes/Public/Dataflow/GeometryCollectionNodes.h
@@ -29,11 +29,10 @@ class UMaterial;
 
 
 /**
- *
- * Description for this node
- *
+ * Get a collection from a geometry collecion asset
+ * DEPRECATED 5.7 : use GeometryCollectionToCollection node instead
  */
-USTRUCT(meta = (DataflowGeometryCollection))
+USTRUCT(meta = (DataflowGeometryCollection, Deprecated = "5.7"))
 struct FGetCollectionFromAssetDataflowNode : public FDataflowNode
 {
 	GENERATED_USTRUCT_BODY()
@@ -54,7 +53,6 @@ public:
 	}
 
 	virtual void Evaluate(UE::Dataflow::FContext& Context, const FDataflowOutput* Out) const override;
-
 };
 
 
@@ -102,7 +100,7 @@ public:
 /**
  *
  * Description for this node
- * DEPRECATED - use Print node ( core nodes ) 
+ * DEPRECATED 5.5 - use Print node ( core nodes ) 
  */
 USTRUCT(meta = (Deprecated = "5.5"))
 struct FPrintStringDataflowNode : public FDataflowNode
