Dataflow : fix code analysis warnings

diff --git a/Engine/Plugins/Experimental/GeometryCollectionPlugin/Source/GeometryCollectionNodes/Public/Dataflow/GeometryCollectionNodes.h b/Engine/Plugins/Experimental/GeometryCollectionPlugin/Source/GeometryCollectionNodes/Public/Dataflow/GeometryCollectionNodes.h
--- a/Engine/Plugins/Experimental/GeometryCollectionPlugin/Source/GeometryCollectionNodes/Public/Dataflow/GeometryCollectionNodes.h
+++ b/Engine/Plugins/Experimental/GeometryCollectionPlugin/Source/GeometryCollectionNodes/Public/Dataflow/GeometryCollectionNodes.h
@@ -1309,7 +1309,7 @@ private:
 
 	/** If hull-based proximity detection is enabled, amount to expand hulls when searching for overlapping neighbors */
 	UPROPERTY(EditAnywhere, Category = "Proximity", meta = (DataflowInput, ClampMin = "0", Units = cm,
-		EditCondition = "ProximityMethod == EProximityMethodEnum::Dataflow_ProximityMethod_ConvexHull || FilterContactMethod == EProximityContactFilteringMethodEnum::Dataflow_ProximityContactFilteringMethod_ConvexHullSharp || FilterContactMethod == EProximityContactFilteringMethodEnum::Dataflow_ProximityContactFilteringMethod_ConvexHullArea || ContactAreaMethod = EConnectionContactAreaMethodEnum::Dataflow_ProximityContactFilteringMethod_ConvexHullArea"))
+		EditCondition = "ProximityMethod == EProximityMethodEnum::Dataflow_ProximityMethod_ConvexHull || FilterContactMethod == EProximityContactFilteringMethodEnum::Dataflow_ProximityContactFilteringMethod_ConvexHullSharp || FilterContactMethod == EProximityContactFilteringMethodEnum::Dataflow_ProximityContactFilteringMethod_ConvexHullArea || ContactAreaMethod == EConnectionContactAreaMethodEnum::Dataflow_ProximityContactFilteringMethod_ConvexHullArea"))
 	float DistanceThreshold = 1;
 
 	// If greater than zero, proximity will be additionally filtered by a 'contact' threshold, in cm, to exclude grazing / corner proximity
@@ -1330,7 +1330,7 @@ private:
 
 	/** Whether to compute new convex hulls for proximity, or use the pre-existing hulls on the Collection, when using convex hulls to determine proximity */
 	UPROPERTY(EditAnywhere, Category = "Proximity", meta = (
-		EditCondition = "ProximityMethod == EProximityMethodEnum::Dataflow_ProximityMethod_ConvexHull || FilterContactMethod == EProximityContactFilteringMethodEnum::Dataflow_ProximityContactFilteringMethod_ConvexHullSharp || FilterContactMethod == EProximityContactFilteringMethodEnum::Dataflow_ProximityContactFilteringMethod_ConvexHullArea || ContactAreaMethod = EConnectionContactAreaMethodEnum::Dataflow_ProximityContactFilteringMethod_ConvexHullArea"))
+		EditCondition = "ProximityMethod == EProximityMethodEnum::Dataflow_ProximityMethod_ConvexHull || FilterContactMethod == EProximityContactFilteringMethodEnum::Dataflow_ProximityContactFilteringMethod_ConvexHullSharp || FilterContactMethod == EProximityContactFilteringMethodEnum::Dataflow_ProximityContactFilteringMethod_ConvexHullArea || ContactAreaMethod == EConnectionContactAreaMethodEnum::Dataflow_ProximityContactFilteringMethod_ConvexHullArea"))
 	bool bRecomputeConvexHulls = true;
 
 	/** GeometryCollection to update the proximity graph on */
@@ -2020,7 +2020,7 @@ private:
 	FManagedArrayCollection Collection;
 
 	/** Points to add to the collection */
-	UPROPERTY(EditAnyWhere, Category = "Points", meta = (DataflowInput))
+	UPROPERTY(EditAnywhere, Category = "Points", meta = (DataflowInput))
 	TArray<FVector> Points;
 
 public:
