Dataflow:
- Added MakeIntArray node

diff --git a/Engine/Plugins/Experimental/GeometryCollectionPlugin/Source/GeometryCollectionNodes/Private/Dataflow/GeometryCollectionMakeNodes.cpp b/Engine/Plugins/Experimental/GeometryCollectionPlugin/Source/GeometryCollectionNodes/Private/Dataflow/GeometryCollectionMakeNodes.cpp
--- a/Engine/Plugins/Experimental/GeometryCollectionPlugin/Source/GeometryCollectionNodes/Private/Dataflow/GeometryCollectionMakeNodes.cpp
+++ b/Engine/Plugins/Experimental/GeometryCollectionPlugin/Source/GeometryCollectionNodes/Private/Dataflow/GeometryCollectionMakeNodes.cpp
@@ -86,6 +86,7 @@ namespace UE::Dataflow
 		DATAFLOW_NODE_REGISTER_CREATION_FACTORY(FMakeRectangleMeshDataflowNode);
 		DATAFLOW_NODE_REGISTER_CREATION_FACTORY(FMakeTorusMeshDataflowNode);
 		DATAFLOW_NODE_REGISTER_CREATION_FACTORY(FMakeTransformDataflowNode);
+		DATAFLOW_NODE_REGISTER_CREATION_FACTORY(FMakeIntArrayDataflowNode);
 
 		// Deprecated
 		DATAFLOW_NODE_REGISTER_CREATION_FACTORY(FMakeTransformDataflowNode_v2);
@@ -848,3 +849,32 @@ void FMakeTorusMeshDataflowNode::Evaluate(UE::Dataflow::FContext& Context, const
 }
 
 /* -------------------------------------------------------------------------------- */
+
+FMakeIntArrayDataflowNode::FMakeIntArrayDataflowNode(const UE::Dataflow::FNodeParameters& InParam, FGuid InGuid)
+	: FDataflowNode(InParam, InGuid)
+{
+	RegisterInputConnection(&Start);
+	RegisterInputConnection(&Count);
+	RegisterOutputConnection(&IntArray);
+}
+
+void FMakeIntArrayDataflowNode::Evaluate(UE::Dataflow::FContext& Context, const FDataflowOutput* Out) const
+{
+	if (Out->IsA(&IntArray))
+	{
+		const int32 InStart = GetValue(Context, &Start);
+		const int32 InCount = FMath::Max(0, GetValue(Context, &Count));
+		
+		TArray<int32> OutIntArray;
+		OutIntArray.SetNumUninitialized(InCount);
+
+		for (int32 Idx = 0; Idx < InCount; ++Idx)
+		{
+			OutIntArray[Idx] = InStart + Idx;
+		}
+
+		SetValue(Context, OutIntArray, &IntArray);
+	}
+}
+
+/* -------------------------------------------------------------------------------- */
diff --git a/Engine/Plugins/Experimental/GeometryCollectionPlugin/Source/GeometryCollectionNodes/Public/Dataflow/GeometryCollectionMakeNodes.h b/Engine/Plugins/Experimental/GeometryCollectionPlugin/Source/GeometryCollectionNodes/Public/Dataflow/GeometryCollectionMakeNodes.h
--- a/Engine/Plugins/Experimental/GeometryCollectionPlugin/Source/GeometryCollectionNodes/Public/Dataflow/GeometryCollectionMakeNodes.h
+++ b/Engine/Plugins/Experimental/GeometryCollectionPlugin/Source/GeometryCollectionNodes/Public/Dataflow/GeometryCollectionMakeNodes.h
@@ -1065,6 +1065,37 @@ private:
 	virtual void Evaluate(UE::Dataflow::FContext& Context, const FDataflowOutput* Out) const override;
 };
 
+/**
+ *
+ * Make an int32 array starting with Start with Count elements
+ *
+ */
+USTRUCT()
+struct FMakeIntArrayDataflowNode : public FDataflowNode
+{
+	GENERATED_USTRUCT_BODY()
+	DATAFLOW_NODE_DEFINE_INTERNAL(FMakeIntArrayDataflowNode, "MakeIntArray", "Generators|Array", "integer indices start end")
+
+private:
+	/** First element of IntArray */
+	UPROPERTY(EditAnywhere, Category = "Int", meta = (DataflowInput))
+	int32 Start = 0;
+
+	/** Number of elements in IntArray */
+	UPROPERTY(EditAnywhere, Category = "Int", meta = (DataflowInput, UIMin = "0", ClampMin = "0"))
+	int32 Count = 3;
+
+	/** IntArray */
+	UPROPERTY(meta = (DataflowOutput))
+	TArray<int32> IntArray;
+
+public:
+	FMakeIntArrayDataflowNode(const UE::Dataflow::FNodeParameters& InParam, FGuid InGuid = FGuid::NewGuid());
+
+private:
+	virtual void Evaluate(UE::Dataflow::FContext& Context, const FDataflowOutput* Out) const override;
+};
+
 namespace UE::Dataflow
 {
 	void GeometryCollectionMakeNodes();
