Dataflow:
- Added ReverseOrder flag to MakeIntArray node

diff --git a/Engine/Plugins/Experimental/GeometryCollectionPlugin/Source/GeometryCollectionNodes/Private/Dataflow/GeometryCollectionMakeNodes.cpp b/Engine/Plugins/Experimental/GeometryCollectionPlugin/Source/GeometryCollectionNodes/Private/Dataflow/GeometryCollectionMakeNodes.cpp
--- a/Engine/Plugins/Experimental/GeometryCollectionPlugin/Source/GeometryCollectionNodes/Private/Dataflow/GeometryCollectionMakeNodes.cpp
+++ b/Engine/Plugins/Experimental/GeometryCollectionPlugin/Source/GeometryCollectionNodes/Private/Dataflow/GeometryCollectionMakeNodes.cpp
@@ -51,6 +51,7 @@
 #include "Generators/DiscMeshGenerator.h"
 #include "Generators/StairGenerator.h"
 #include "Generators/RectangleMeshGenerator.h"
+#include "Algo/Reverse.h"
 
 #include UE_INLINE_GENERATED_CPP_BY_NAME(GeometryCollectionMakeNodes)
 
@@ -873,6 +874,11 @@ void FMakeIntArrayDataflowNode::Evaluate(UE::Dataflow::FContext& Context, const
 			OutIntArray[Idx] = InStart + Idx;
 		}
 
+		if (bReverseOrder)
+		{
+			Algo::Reverse(OutIntArray);
+		}
+
 		SetValue(Context, OutIntArray, &IntArray);
 	}
 }
diff --git a/Engine/Plugins/Experimental/GeometryCollectionPlugin/Source/GeometryCollectionNodes/Public/Dataflow/GeometryCollectionMakeNodes.h b/Engine/Plugins/Experimental/GeometryCollectionPlugin/Source/GeometryCollectionNodes/Public/Dataflow/GeometryCollectionMakeNodes.h
--- a/Engine/Plugins/Experimental/GeometryCollectionPlugin/Source/GeometryCollectionNodes/Public/Dataflow/GeometryCollectionMakeNodes.h
+++ b/Engine/Plugins/Experimental/GeometryCollectionPlugin/Source/GeometryCollectionNodes/Public/Dataflow/GeometryCollectionMakeNodes.h
@@ -1085,6 +1085,9 @@ private:
 	UPROPERTY(EditAnywhere, Category = "Int", meta = (DataflowInput, UIMin = "0", ClampMin = "0"))
 	int32 Count = 3;
 
+	UPROPERTY(EditAnywhere, Category = "Int")
+	bool bReverseOrder = false;
+
 	/** IntArray */
 	UPROPERTY(meta = (DataflowOutput))
 	TArray<int32> IntArray;
