Dataflow:
- Added MakeIntArrayFromList node to create a TArray<int32> from a comma separated list
- Updated Print node to outout incoming Value

diff --git a/Engine/Plugins/Experimental/GeometryCollectionPlugin/Source/GeometryCollectionNodes/Private/Dataflow/GeometryCollectionMakeNodes.cpp b/Engine/Plugins/Experimental/GeometryCollectionPlugin/Source/GeometryCollectionNodes/Private/Dataflow/GeometryCollectionMakeNodes.cpp
--- a/Engine/Plugins/Experimental/GeometryCollectionPlugin/Source/GeometryCollectionNodes/Private/Dataflow/GeometryCollectionMakeNodes.cpp
+++ b/Engine/Plugins/Experimental/GeometryCollectionPlugin/Source/GeometryCollectionNodes/Private/Dataflow/GeometryCollectionMakeNodes.cpp
@@ -52,6 +52,7 @@
 #include "Generators/StairGenerator.h"
 #include "Generators/RectangleMeshGenerator.h"
 #include "Algo/Reverse.h"
+#include "Dataflow/GeometryCollectionUtils.h"
 
 #include UE_INLINE_GENERATED_CPP_BY_NAME(GeometryCollectionMakeNodes)
 
@@ -88,6 +89,7 @@ namespace UE::Dataflow
 		DATAFLOW_NODE_REGISTER_CREATION_FACTORY(FMakeTorusMeshDataflowNode);
 		DATAFLOW_NODE_REGISTER_CREATION_FACTORY(FMakeTransformDataflowNode);
 		DATAFLOW_NODE_REGISTER_CREATION_FACTORY(FMakeIntArrayDataflowNode);
+		DATAFLOW_NODE_REGISTER_CREATION_FACTORY(FMakeIntArrayFromListDataflowNode);
 
 		// Deprecated
 		DATAFLOW_NODE_REGISTER_CREATION_FACTORY(FMakeTransformDataflowNode_v2);
@@ -884,3 +886,42 @@ void FMakeIntArrayDataflowNode::Evaluate(UE::Dataflow::FContext& Context, const
 }
 
 /* -------------------------------------------------------------------------------- */
+
+FMakeIntArrayFromListDataflowNode::FMakeIntArrayFromListDataflowNode(const UE::Dataflow::FNodeParameters& InParam, FGuid InGuid)
+	: FDataflowNode(InParam, InGuid)
+{
+	RegisterOutputConnection(&IntArray);
+}
+
+void FMakeIntArrayFromListDataflowNode::Evaluate(UE::Dataflow::FContext& Context, const FDataflowOutput* Out) const
+{
+	if (Out->IsA(&IntArray))
+	{
+		TArray<int32> OutIntArray;
+		uint32 ErrorCode = 0;
+
+		if (!UE::Dataflow::Utils::ParseIndicesStr(List, OutIntArray, ErrorCode))
+		{
+			if (ErrorCode == UE::Dataflow::Utils::Error_InvalidChars)
+			{
+				SetError(Context, &IntArray, TEXT("Invalid character(s) specified in list"));
+			}
+			else if (ErrorCode == UE::Dataflow::Utils::Error_InvalidFormatInSegment)
+			{
+				SetError(Context, &IntArray, TEXT("Invalid format in segment"));
+			}
+
+			SetValue(Context, TArray<int32>(), &IntArray);
+			return;
+		}
+
+		if (bReverseOrder)
+		{
+			Algo::Reverse(OutIntArray);
+		}
+
+		SetValue(Context, OutIntArray, &IntArray);
+	}
+}
+
+/* -------------------------------------------------------------------------------- */
diff --git a/Engine/Plugins/Experimental/GeometryCollectionPlugin/Source/GeometryCollectionNodes/Private/Dataflow/GeometryCollectionSelectionNodes.cpp b/Engine/Plugins/Experimental/GeometryCollectionPlugin/Source/GeometryCollectionNodes/Private/Dataflow/GeometryCollectionSelectionNodes.cpp
--- a/Engine/Plugins/Experimental/GeometryCollectionPlugin/Source/GeometryCollectionNodes/Private/Dataflow/GeometryCollectionSelectionNodes.cpp
+++ b/Engine/Plugins/Experimental/GeometryCollectionPlugin/Source/GeometryCollectionNodes/Private/Dataflow/GeometryCollectionSelectionNodes.cpp
@@ -35,7 +35,7 @@
 #include "FractureEngineSelection.h"
 #include "GeometryCollection/Facades/CollectionTransformSelectionFacade.h"
 #include "GeometryCollection/Facades/CollectionHierarchyFacade.h"
-#include "Internationalization/Regex.h"
+#include "Dataflow/GeometryCollectionUtils.h"
 
 #include UE_INLINE_GENERATED_CPP_BY_NAME(GeometryCollectionSelectionNodes)
 
@@ -424,67 +424,6 @@ void FCollectionTransformSelectionCustomDataflowNode::Evaluate(UE::Dataflow::FCo
 	}
 }
 
-namespace UE::Dataflow::Private
-{
-	const static uint32 Error_InvalidChars = 1;
-	const static uint32 Error_InvalidFormatInSegment = 2;
-
-	/* e.g. "0, 2, 5-10, 12-15". If left empty, all will be used */
-	static bool ParseIndicesStr(const FString& InFramesString, TArray<int32>& OutIndices, uint32& OutErrorCode)
-	{
-		static const FRegexPattern AllowedCharsPattern(TEXT("^[-,0-9\\s]+$"));
-
-		if (!FRegexMatcher(AllowedCharsPattern, InFramesString).FindNext())
-		{
-			// Input contains invalid characters
-			OutErrorCode = Error_InvalidChars;
-			return false;
-		}
-
-		static const FRegexPattern SingleNumberPattern(TEXT("^\\s*(\\d+)\\s*$"));
-		static const FRegexPattern RangePattern(TEXT("^\\s*(\\d+)\\s*-\\s*(\\d+)\\s*$"));
-
-		TArray<FString> Segments;
-		InFramesString.ParseIntoArray(Segments, TEXT(","), true);
-		for (const FString& Segment : Segments)
-		{
-			bool bSegmentValid = false;
-
-			FRegexMatcher SingleNumberMatcher(SingleNumberPattern, Segment);
-			if (SingleNumberMatcher.FindNext())
-			{
-				const int32 SingleNumber = FCString::Atoi(*SingleNumberMatcher.GetCaptureGroup(1));
-				OutIndices.Add(SingleNumber);
-				bSegmentValid = true;
-			}
-			else
-			{
-				FRegexMatcher RangeMatcher(RangePattern, Segment);
-				if (RangeMatcher.FindNext())
-				{
-					const int32 RangeStart = FCString::Atoi(*RangeMatcher.GetCaptureGroup(1));
-					const int32 RangeEnd = FCString::Atoi(*RangeMatcher.GetCaptureGroup(2));
-
-					for (int32 i = RangeStart; i <= RangeEnd; ++i)
-					{
-						OutIndices.Add(i);
-					}
-					bSegmentValid = true;
-				}
-			}
-
-			if (!bSegmentValid)
-			{
-				// Invalid format in segment
-				OutErrorCode = Error_InvalidFormatInSegment;
-				return false;
-			}
-		}
-
-		return true;
-	}
-}
-
 FCollectionTransformSelectionCustomDataflowNode_v2::FCollectionTransformSelectionCustomDataflowNode_v2(const UE::Dataflow::FNodeParameters& InParam, FGuid InGuid)
 	: FDataflowNode(InParam, InGuid)
 {
@@ -510,19 +449,19 @@ void FCollectionTransformSelectionCustomDataflowNode_v2::Evaluate(UE::Dataflow::
 
 			TArray<int32> Indices;
 			uint32 ErrorCode = 0;
-			if (UE::Dataflow::Private::ParseIndicesStr(InBoneIndices, Indices, ErrorCode))
+			if (UE::Dataflow::Utils::ParseIndicesStr(InBoneIndices, Indices, ErrorCode))
 			{
 				NewTransformSelection.SetSelected(Indices);
 			}
 			else
 			{
-				if (ErrorCode == UE::Dataflow::Private::Error_InvalidChars)
+				if (ErrorCode == UE::Dataflow::Utils::Error_InvalidChars)
 				{
-					// Handle Error				
+					SetError(Context, &TransformSelection, TEXT("Invalid character(s) specified in list"));
 				}
-				else if (ErrorCode == UE::Dataflow::Private::Error_InvalidFormatInSegment)
+				else if (ErrorCode == UE::Dataflow::Utils::Error_InvalidFormatInSegment)
 				{
-					// Handle Error				
+					SetError(Context, &TransformSelection, TEXT("Invalid format in segment"));
 				}
 			}
 
diff --git a/Engine/Plugins/Experimental/GeometryCollectionPlugin/Source/GeometryCollectionNodes/Private/Dataflow/GeometryCollectionUtils.cpp b/Engine/Plugins/Experimental/GeometryCollectionPlugin/Source/GeometryCollectionNodes/Private/Dataflow/GeometryCollectionUtils.cpp
--- a/Engine/Plugins/Experimental/GeometryCollectionPlugin/Source/GeometryCollectionNodes/Private/Dataflow/GeometryCollectionUtils.cpp
+++ b/Engine/Plugins/Experimental/GeometryCollectionPlugin/Source/GeometryCollectionNodes/Private/Dataflow/GeometryCollectionUtils.cpp
@@ -4,6 +4,7 @@
 #include "Dataflow/DataflowSettings.h"
 
 #include "GeometryCollection/GeometryCollectionAlgo.h"
+#include "Internationalization/Regex.h"
 
 namespace UE::Dataflow
 {
@@ -110,5 +111,60 @@ namespace UE::Dataflow
 			}
 		}
 #endif
+
+		/* e.g. "0, 2, 5-10, 12-15". If left empty, all will be used */
+		bool ParseIndicesStr(const FString& InFramesString, TArray<int32>& OutIndices, uint32& OutErrorCode)
+		{
+			static const FRegexPattern AllowedCharsPattern(TEXT("^[-,0-9\\s]+$"));
+
+			if (!FRegexMatcher(AllowedCharsPattern, InFramesString).FindNext())
+			{
+				// Input contains invalid characters
+				OutErrorCode = Error_InvalidChars;
+				return false;
+			}
+
+			static const FRegexPattern SingleNumberPattern(TEXT("^\\s*(\\d+)\\s*$"));
+			static const FRegexPattern RangePattern(TEXT("^\\s*(\\d+)\\s*-\\s*(\\d+)\\s*$"));
+
+			TArray<FString> Segments;
+			InFramesString.ParseIntoArray(Segments, TEXT(","), true);
+			for (const FString& Segment : Segments)
+			{
+				bool bSegmentValid = false;
+
+				FRegexMatcher SingleNumberMatcher(SingleNumberPattern, Segment);
+				if (SingleNumberMatcher.FindNext())
+				{
+					const int32 SingleNumber = FCString::Atoi(*SingleNumberMatcher.GetCaptureGroup(1));
+					OutIndices.Add(SingleNumber);
+					bSegmentValid = true;
+				}
+				else
+				{
+					FRegexMatcher RangeMatcher(RangePattern, Segment);
+					if (RangeMatcher.FindNext())
+					{
+						const int32 RangeStart = FCString::Atoi(*RangeMatcher.GetCaptureGroup(1));
+						const int32 RangeEnd = FCString::Atoi(*RangeMatcher.GetCaptureGroup(2));
+
+						for (int32 i = RangeStart; i <= RangeEnd; ++i)
+						{
+							OutIndices.Add(i);
+						}
+						bSegmentValid = true;
+					}
+				}
+
+				if (!bSegmentValid)
+				{
+					// Invalid format in segment
+					OutErrorCode = Error_InvalidFormatInSegment;
+					return false;
+				}
+			}
+
+			return true;
+		}
 	}
 }
diff --git a/Engine/Plugins/Experimental/GeometryCollectionPlugin/Source/GeometryCollectionNodes/Public/Dataflow/GeometryCollectionMakeNodes.h b/Engine/Plugins/Experimental/GeometryCollectionPlugin/Source/GeometryCollectionNodes/Public/Dataflow/GeometryCollectionMakeNodes.h
--- a/Engine/Plugins/Experimental/GeometryCollectionPlugin/Source/GeometryCollectionNodes/Public/Dataflow/GeometryCollectionMakeNodes.h
+++ b/Engine/Plugins/Experimental/GeometryCollectionPlugin/Source/GeometryCollectionNodes/Public/Dataflow/GeometryCollectionMakeNodes.h
@@ -1099,6 +1099,36 @@ private:
 	virtual void Evaluate(UE::Dataflow::FContext& Context, const FDataflowOutput* Out) const override;
 };
 
+/**
+ *
+ * Make an int32 array from a comma separated list, e.g. "0, 2, 5-10, 12-15"
+ *
+ */
+USTRUCT()
+struct FMakeIntArrayFromListDataflowNode : public FDataflowNode
+{
+	GENERATED_USTRUCT_BODY()
+	DATAFLOW_NODE_DEFINE_INTERNAL(FMakeIntArrayFromListDataflowNode, "MakeIntArrayFromList", "Generators|Array", "integer indices start end list")
+
+private:
+	/** First element of IntArray */
+	UPROPERTY(EditAnywhere, Category = "Int")
+	FString List;
+
+	UPROPERTY(EditAnywhere, Category = "Int")
+	bool bReverseOrder = false;
+
+	/** IntArray */
+	UPROPERTY(meta = (DataflowOutput))
+	TArray<int32> IntArray;
+
+public:
+	FMakeIntArrayFromListDataflowNode(const UE::Dataflow::FNodeParameters& InParam, FGuid InGuid = FGuid::NewGuid());
+
+private:
+	virtual void Evaluate(UE::Dataflow::FContext& Context, const FDataflowOutput* Out) const override;
+};
+
 namespace UE::Dataflow
 {
 	void GeometryCollectionMakeNodes();
diff --git a/Engine/Plugins/Experimental/GeometryCollectionPlugin/Source/GeometryCollectionNodes/Public/Dataflow/GeometryCollectionUtils.h b/Engine/Plugins/Experimental/GeometryCollectionPlugin/Source/GeometryCollectionNodes/Public/Dataflow/GeometryCollectionUtils.h
--- a/Engine/Plugins/Experimental/GeometryCollectionPlugin/Source/GeometryCollectionNodes/Public/Dataflow/GeometryCollectionUtils.h
+++ b/Engine/Plugins/Experimental/GeometryCollectionPlugin/Source/GeometryCollectionNodes/Public/Dataflow/GeometryCollectionUtils.h
@@ -28,6 +28,12 @@ namespace UE::Dataflow
 
 		GEOMETRYCOLLECTIONNODES_API FLinearColor GetColorByLevel(const FManagedArrayCollection& InCollection, int32 InLevel);
 #endif
+
+		const uint32 Error_InvalidChars = 1;
+		const uint32 Error_InvalidFormatInSegment = 2;
+
+		/* e.g. "0, 2, 5-10, 12-15" */
+		GEOMETRYCOLLECTIONNODES_API bool ParseIndicesStr(const FString& InFramesString, TArray<int32>& OutIndices, uint32& OutErrorCode);
 	}
 }
 
diff --git a/Engine/Source/Runtime/Experimental/Dataflow/Core/Private/Dataflow/DataflowCoreNodes.cpp b/Engine/Source/Runtime/Experimental/Dataflow/Core/Private/Dataflow/DataflowCoreNodes.cpp
--- a/Engine/Source/Runtime/Experimental/Dataflow/Core/Private/Dataflow/DataflowCoreNodes.cpp
+++ b/Engine/Source/Runtime/Experimental/Dataflow/Core/Private/Dataflow/DataflowCoreNodes.cpp
@@ -198,11 +198,22 @@ UE::Dataflow::TConnectionReference<FDataflowAnyType> FDataflowSelectNode::GetCon
 FDataflowPrintNode::FDataflowPrintNode(const UE::Dataflow::FNodeParameters& InParam, FGuid InGuid)
 	: FDataflowNode(InParam, InGuid)
 {
-	RegisterInputConnection(&Value);
+	static const FName TypeDependencyGroup("Main");
+
+	RegisterInputConnection(&Value)
+		.SetTypeDependencyGroup(TypeDependencyGroup);
+
+	RegisterOutputConnection(&Value, &Value)
+		.SetTypeDependencyGroup(TypeDependencyGroup);
 }
 
 void FDataflowPrintNode::Evaluate(UE::Dataflow::FContext& Context, const FDataflowOutput* Out) const
 {
+	if (Out->IsA(&Value))
+	{
+		SafeForwardInput(Context, &Value, &Value);
+	}
+
 	const FString InValue = "Result=" + GetValue(Context, &Value);
 	Context.Info(InValue, this, nullptr);
 }
diff --git a/Engine/Source/Runtime/Experimental/Dataflow/Core/Public/Dataflow/DataflowCoreNodes.h b/Engine/Source/Runtime/Experimental/Dataflow/Core/Public/Dataflow/DataflowCoreNodes.h
--- a/Engine/Source/Runtime/Experimental/Dataflow/Core/Public/Dataflow/DataflowCoreNodes.h
+++ b/Engine/Source/Runtime/Experimental/Dataflow/Core/Public/Dataflow/DataflowCoreNodes.h
@@ -107,7 +107,7 @@ public:
 	virtual void Evaluate(UE::Dataflow::FContext& Context, const FDataflowOutput* Out) const;
 
 public:
-	UPROPERTY(meta = (DataflowInput, DisplayName = "Value"))
+	UPROPERTY(meta = (DataflowInput, DataflowOutput, DataflowPassthrough = "Value", DisplayName = "Value"))
 	FDataflowStringConvertibleTypes Value;
 };
 
