Dataflow:
- Added AssetRandomSeed node

diff --git a/Engine/Plugins/Experimental/GeometryCollectionPlugin/Source/GeometryCollectionNodes/Private/Dataflow/GeometryCollectionAssetNodes.cpp b/Engine/Plugins/Experimental/GeometryCollectionPlugin/Source/GeometryCollectionNodes/Private/Dataflow/GeometryCollectionAssetNodes.cpp
--- a/Engine/Plugins/Experimental/GeometryCollectionPlugin/Source/GeometryCollectionNodes/Private/Dataflow/GeometryCollectionAssetNodes.cpp
+++ b/Engine/Plugins/Experimental/GeometryCollectionPlugin/Source/GeometryCollectionNodes/Private/Dataflow/GeometryCollectionAssetNodes.cpp
@@ -37,6 +37,7 @@ namespace UE::Dataflow
 		DATAFLOW_NODE_REGISTER_CREATION_FACTORY(FCreateGeometryCollectionFromSourcesDataflowNode);
 		DATAFLOW_NODE_REGISTER_CREATION_FACTORY(FGeometryCollectionToCollectionDataflowNode);
 		DATAFLOW_NODE_REGISTER_CREATION_FACTORY(FBlueprintToCollectionDataflowNode);
+		DATAFLOW_NODE_REGISTER_CREATION_FACTORY(FAssetRandomSeedDataflowNode);
 	}
 }
 
@@ -709,3 +710,33 @@ void FBlueprintToCollectionDataflowNode_v2::Evaluate(UE::Dataflow::FContext& Con
 	SetValue(Context, MoveTemp(OutInstancedMeshes), &InstancedMeshes);
 	SetValue(Context, MoveTemp(OutRootProxyMeshes), &RootProxyMeshes);
 }
+
+// ===========================================================================================================================
+
+FAssetRandomSeedDataflowNode::FAssetRandomSeedDataflowNode(const UE::Dataflow::FNodeParameters& InParam, FGuid InGuid)
+	: FDataflowNode(InParam, InGuid)
+{
+	RegisterInputConnection(&Seed);
+	RegisterOutputConnection(&RandomSeed);
+}
+
+void FAssetRandomSeedDataflowNode::Evaluate(UE::Dataflow::FContext& Context, const FDataflowOutput* Out) const
+{
+	if (Out->IsA(&RandomSeed))
+	{
+		const int32 InSeed = GetValue(Context, &Seed);
+
+		if (const UE::Dataflow::FEngineContext* EngineContext = Context.AsType<UE::Dataflow::FEngineContext>())
+		{
+			if (EngineContext->Owner)
+			{
+				const int32 HashVal = (int32)HashCombine(GetTypeHash(EngineContext->Owner->GetPathName()) + GetTypeHash(InSeed));
+				SetValue(Context, HashVal, &RandomSeed);
+
+				return;
+			}	
+		}
+
+		SetValue(Context, InSeed, &RandomSeed);
+	}
+}
diff --git a/Engine/Plugins/Experimental/GeometryCollectionPlugin/Source/GeometryCollectionNodes/Public/Dataflow/GeometryCollectionAssetNodes.h b/Engine/Plugins/Experimental/GeometryCollectionPlugin/Source/GeometryCollectionNodes/Public/Dataflow/GeometryCollectionAssetNodes.h
--- a/Engine/Plugins/Experimental/GeometryCollectionPlugin/Source/GeometryCollectionNodes/Public/Dataflow/GeometryCollectionAssetNodes.h
+++ b/Engine/Plugins/Experimental/GeometryCollectionPlugin/Source/GeometryCollectionNodes/Public/Dataflow/GeometryCollectionAssetNodes.h
@@ -442,6 +442,30 @@ private:
 	virtual void Evaluate(UE::Dataflow::FContext& Context, const FDataflowOutput* Out) const override;
 };
 
+/**
+* Create a random seed value from an asset
+*/
+USTRUCT(meta = (DataflowContext = "GeometryCollection", DataflowGeometryCollection))
+struct FAssetRandomSeedDataflowNode : public FDataflowNode
+{
+	GENERATED_USTRUCT_BODY()
+	DATAFLOW_NODE_DEFINE_INTERNAL(FAssetRandomSeedDataflowNode, "AssetRandomSeed", "GeometryCollection|Asset", "asset random seed")
+
+public:
+	FAssetRandomSeedDataflowNode(const UE::Dataflow::FNodeParameters& InParam, FGuid InGuid = FGuid::NewGuid());
+
+private:
+	/** Seed to combine with the seed generated from the asset path */
+	UPROPERTY(EditAnywhere, Category = "Random", meta = (DataflowInput));
+	int32 Seed = 0;
+
+	/** Random seed value */
+	UPROPERTY(meta = (DataflowOutput))
+	int32 RandomSeed;
+
+	virtual void Evaluate(UE::Dataflow::FContext& Context, const FDataflowOutput* Out) const override;
+};
+
 namespace UE::Dataflow
 {
 	void GeometryCollectionEngineAssetNodes();
