Dataflow:
- Added SpheresToMesh node

diff --git a/Engine/Plugins/Experimental/GeometryCollectionPlugin/Source/GeometryCollectionNodes/Private/Dataflow/GeometryCollectionDebugNodes.cpp b/Engine/Plugins/Experimental/GeometryCollectionPlugin/Source/GeometryCollectionNodes/Private/Dataflow/GeometryCollectionDebugNodes.cpp
--- a/Engine/Plugins/Experimental/GeometryCollectionPlugin/Source/GeometryCollectionNodes/Private/Dataflow/GeometryCollectionDebugNodes.cpp
+++ b/Engine/Plugins/Experimental/GeometryCollectionPlugin/Source/GeometryCollectionNodes/Private/Dataflow/GeometryCollectionDebugNodes.cpp
@@ -23,6 +23,7 @@ namespace UE::Dataflow
 		DATAFLOW_NODE_REGISTER_CREATION_FACTORY(FMeshToOBJStringDebugDataflowNode);
 		DATAFLOW_NODE_REGISTER_CREATION_FACTORY(FSphereCoveringCountSpheresNode);
 		DATAFLOW_NODE_REGISTER_CREATION_FACTORY(FWriteStringToFile);
+		DATAFLOW_NODE_REGISTER_CREATION_FACTORY(FSpheresToMeshDataflowNode);
 	}
 }
 
@@ -202,3 +203,71 @@ void FWriteStringToFile::Evaluate(UE::Dataflow::FContext& Context, const FDatafl
 		UE_LOG(LogChaos, Warning, TEXT("Failed to write to file %s:\n\n%s"), *InPath, *InContents);
 	}
 }
+
+/* --------------------------------------------------------------------------------------------------------------------------- */
+
+FSpheresToMeshDataflowNode::FSpheresToMeshDataflowNode(const UE::Dataflow::FNodeParameters& InParam, FGuid InGuid)
+	: FDataflowNode(InParam, InGuid)
+{
+	RegisterInputConnection(&Spheres);
+	RegisterOutputConnection(&Mesh);
+	RegisterOutputConnection(&Meshes);
+}
+
+void FSpheresToMeshDataflowNode::Evaluate(UE::Dataflow::FContext& Context, const FDataflowOutput* Out) const
+{
+	using namespace UE::Geometry;
+
+	if (Out->IsA(&Mesh))
+	{
+		TArray<FSphere> InSphereArr = GetValue(Context, &Spheres);
+
+		TObjectPtr<UDynamicMesh> NewMesh = NewObject<UDynamicMesh>();
+		NewMesh->Reset();
+
+		FDynamicMesh3 Accum;
+		FDynamicMeshEditor Editor(&Accum);
+		FMeshIndexMappings IndexMaps_Unused;
+		for (int32 SphereIdx = 0; SphereIdx < InSphereArr.Num(); ++SphereIdx)
+		{
+			FBoxSphereGenerator SphereGen;
+			SphereGen.EdgeVertices.A = SphereGen.EdgeVertices.B = SphereGen.EdgeVertices.C = FMath::Max(2, VerticesAlongEachSide);
+			SphereGen.Radius = InSphereArr[SphereIdx].W;
+			FVector Center = InSphereArr[SphereIdx].Center;
+
+			FDynamicMesh3 Sphere(&SphereGen.Generate());
+			Editor.AppendMesh(&Sphere, IndexMaps_Unused,
+				[Center](int VID, const FVector3d& Pos) { return Pos + Center; });
+			IndexMaps_Unused.Reset();
+		}
+		NewMesh->SetMesh(MoveTemp(Accum));
+
+		SetValue(Context, NewMesh, &Mesh);
+	}
+	else if (Out->IsA(&Meshes))
+	{
+		TArray<FSphere> InSphereArr = GetValue(Context, &Spheres);
+
+		TArray<TObjectPtr<UDynamicMesh>> NewMeshes;
+		NewMeshes.SetNum(Spheres.Num());
+
+		for (FSphere& Sphere : InSphereArr)
+		{
+			FBoxSphereGenerator SphereGen;
+			SphereGen.EdgeVertices.A = SphereGen.EdgeVertices.B = SphereGen.EdgeVertices.C = FMath::Max(2, VerticesAlongEachSide);
+			SphereGen.Radius = Sphere.W;
+			SphereGen.Box.Frame.Origin = Sphere.Center;
+
+			FDynamicMesh3 SphereMesh(&SphereGen.Generate());
+
+			TObjectPtr<UDynamicMesh> NewMesh = NewObject<UDynamicMesh>();
+			NewMesh->SetMesh(MoveTemp(SphereMesh));
+			NewMeshes.Add(NewMesh);
+		}
+
+		SetValue(Context, MoveTemp(NewMeshes), &Meshes);
+	}
+}
+
+/* --------------------------------------------------------------------------------------------------------------------------- */
+
diff --git a/Engine/Plugins/Experimental/GeometryCollectionPlugin/Source/GeometryCollectionNodes/Public/Dataflow/GeometryCollectionDebugNodes.h b/Engine/Plugins/Experimental/GeometryCollectionPlugin/Source/GeometryCollectionNodes/Public/Dataflow/GeometryCollectionDebugNodes.h
--- a/Engine/Plugins/Experimental/GeometryCollectionPlugin/Source/GeometryCollectionNodes/Public/Dataflow/GeometryCollectionDebugNodes.h
+++ b/Engine/Plugins/Experimental/GeometryCollectionPlugin/Source/GeometryCollectionNodes/Public/Dataflow/GeometryCollectionDebugNodes.h
@@ -172,6 +172,37 @@ public:
 	virtual void Evaluate(UE::Dataflow::FContext& Context, const FDataflowOutput* Out) const override;
 };
 
+// Convert a TArray<FShere> to a dynamic mesh
+USTRUCT(meta = (DataflowGeometryCollection))
+struct FSpheresToMeshDataflowNode : public FDataflowNode
+{
+	GENERATED_USTRUCT_BODY()
+	DATAFLOW_NODE_DEFINE_INTERNAL(FSpheresToMeshDataflowNode, "SpheresToMesh", "GeometryCollection|Mesh", "")
+
+private:
+	// Array of spheres to convert to a mesh/meshes
+	UPROPERTY(meta = (DataflowInput, DataflowIntrinsic))
+	TArray<FSphere> Spheres;
+
+	// Number of vertices to use along each axis when creating a mesh for each sphere
+	UPROPERTY(EditAnywhere, Category = Options, meta = (ClampMin = 2))
+	int32 VerticesAlongEachSide = 4;
+
+	/** Single mesh aggregating all the spheres together */
+	UPROPERTY(meta = (DataflowOutput))
+	TObjectPtr<UDynamicMesh> Mesh;
+
+	/** Array of meshes for each spheres found */
+	UPROPERTY(meta = (DataflowOutput))
+	TArray<TObjectPtr<UDynamicMesh>> Meshes;
+
+	virtual void Evaluate(UE::Dataflow::FContext& Context, const FDataflowOutput* Out) const override;
+
+public:
+	FSpheresToMeshDataflowNode(const UE::Dataflow::FNodeParameters& InParam, FGuid InGuid = FGuid::NewGuid());
+
+};
+
 namespace UE::Dataflow
 {
 	void GeometryCollectionDebugNodes();
