Geometry collection : when creating a dataflow from the asset browser, add a dataflow from a template

#Synced-CL 48041191

diff --git a/Engine/Plugins/Experimental/GeometryCollectionPlugin/Source/GeometryCollectionEditor/Private/GeometryCollection/GeometryCollectionFactory.cpp b/Engine/Plugins/Experimental/GeometryCollectionPlugin/Source/GeometryCollectionEditor/Private/GeometryCollection/GeometryCollectionFactory.cpp
--- a/Engine/Plugins/Experimental/GeometryCollectionPlugin/Source/GeometryCollectionEditor/Private/GeometryCollection/GeometryCollectionFactory.cpp
+++ b/Engine/Plugins/Experimental/GeometryCollectionPlugin/Source/GeometryCollectionEditor/Private/GeometryCollection/GeometryCollectionFactory.cpp
@@ -188,6 +188,15 @@ UObject* UGeometryCollectionFactory::FactoryCreateNew(UClass* Class, UObject* In
 		}
 	}
 
+	// Assign a default dataflow template if available 
+	if (const UDataflow* DataflowTemplateObject = LoadObject<UDataflow>(nullptr, TEXT("/GeometryCollectionPlugin/DF_GC_Template_StaticMesh.DF_GC_Template_StaticMesh")))
+	{
+		if (UDataflow* DataflowEmbeddedAsset = DuplicateObject<UDataflow>(DataflowTemplateObject, NewGeometryCollection, TEXT("EmbeddedDataflow")))
+		{
+			NewGeometryCollection->GetDataflowInstance().SetDataflowAsset(DataflowEmbeddedAsset);
+		}
+	}
+
 	NewGeometryCollection->Modify();
 	return NewGeometryCollection;
 }
