Dataflow:
- Removed Sphere, SphereArray, Box, BoxArray, Point, PointArray old way rendering from nodes

#Synced-CL 48273319

diff --git a/Engine/Plugins/Experimental/Dataflow/Source/DataflowNodes/Public/Dataflow/DataflowStaticMeshNodes.h b/Engine/Plugins/Experimental/Dataflow/Source/DataflowNodes/Public/Dataflow/DataflowStaticMeshNodes.h
--- a/Engine/Plugins/Experimental/Dataflow/Source/DataflowNodes/Public/Dataflow/DataflowStaticMeshNodes.h
+++ b/Engine/Plugins/Experimental/Dataflow/Source/DataflowNodes/Public/Dataflow/DataflowStaticMeshNodes.h
@@ -44,7 +44,6 @@ struct FGetStaticMeshBoundingBoxDataflowNode : public FDataflowNode
 {
 	GENERATED_USTRUCT_BODY()
 	DATAFLOW_NODE_DEFINE_INTERNAL(FGetStaticMeshBoundingBoxDataflowNode, "GetStaticMeshBoundingBox", "Static Mesh", "Bounds Size Dimensions Extents Center")
-	DATAFLOW_NODE_RENDER_TYPE("SurfaceRender", FName("FBox"), "BoundingBox")
 
 private:
 	/** Static Mesh to compute the bouning box from */
diff --git a/Engine/Plugins/Experimental/GeometryCollectionPlugin/Source/GeometryCollectionNodes/Public/Dataflow/GeometryCollectionFracturingNodes.h b/Engine/Plugins/Experimental/GeometryCollectionPlugin/Source/GeometryCollectionNodes/Public/Dataflow/GeometryCollectionFracturingNodes.h
--- a/Engine/Plugins/Experimental/GeometryCollectionPlugin/Source/GeometryCollectionNodes/Public/Dataflow/GeometryCollectionFracturingNodes.h
+++ b/Engine/Plugins/Experimental/GeometryCollectionPlugin/Source/GeometryCollectionNodes/Public/Dataflow/GeometryCollectionFracturingNodes.h
@@ -62,7 +62,6 @@ struct FUniformScatterPointsDataflowNode_v2 : public FDataflowNode
 {
 	GENERATED_USTRUCT_BODY()
 	DATAFLOW_NODE_DEFINE_INTERNAL(FUniformScatterPointsDataflowNode_v2, "UniformScatterPoints", "Generators|Point", "")
-	DATAFLOW_NODE_RENDER_TYPE("PointsRender", FName("TArray<FVector>"), "Points")
 
 public:
 	/** Minimum for the random range */
@@ -104,7 +103,6 @@ struct FClusterScatterPointsDataflowNode : public FDataflowNode
 {
 	GENERATED_USTRUCT_BODY()
 	DATAFLOW_NODE_DEFINE_INTERNAL(FClusterScatterPointsDataflowNode, "ClusterScatterPoints", "Generators|Point", "")
-	DATAFLOW_NODE_RENDER_TYPE("PointsRender", FName("TArray<FVector>"), "Points")
 
 private:
 
@@ -217,7 +215,6 @@ struct FRadialScatterPointsDataflowNode_v2 : public FDataflowNode
 {
 	GENERATED_USTRUCT_BODY()
 	DATAFLOW_NODE_DEFINE_INTERNAL(FRadialScatterPointsDataflowNode_v2, "RadialScatterPoints", "Generators|Point", "")
-	DATAFLOW_NODE_RENDER_TYPE("PointsRender", FName("TArray<FVector>"), "Points")
 
 public:
 	/** BoundingBox to generate points inside of */
@@ -313,7 +310,6 @@ struct FGridScatterPointsDataflowNode : public FDataflowNode
 {
 	GENERATED_USTRUCT_BODY()
 	DATAFLOW_NODE_DEFINE_INTERNAL(FGridScatterPointsDataflowNode, "GridScatterPoints", "Generators|Point", "")
-	DATAFLOW_NODE_RENDER_TYPE("PointsRender", FName("TArray<FVector>"), "Points")
 
 public:
 	/** Number of points in X direction */
@@ -377,7 +373,6 @@ struct FTransformPointsDataflowNode : public FDataflowNode
 {
 	GENERATED_USTRUCT_BODY()
 	DATAFLOW_NODE_DEFINE_INTERNAL(FTransformPointsDataflowNode, "TransformPoints", "Generators|Point", "")
-	DATAFLOW_NODE_RENDER_TYPE("PointsRender", FName("TArray<FVector>"), "Points")
 
 private:
 	UPROPERTY(meta = (DataflowOutput, DataflowInput, DataflowPassthrough = Points, DataflowIntrinsic))
@@ -400,7 +395,6 @@ struct FAppendPointsDataflowNode : public FDataflowNode
 {
 	GENERATED_USTRUCT_BODY()
 	DATAFLOW_NODE_DEFINE_INTERNAL(FAppendPointsDataflowNode, "AppendPoints", "Generators|Point", "")
-	DATAFLOW_NODE_RENDER_TYPE("PointsRender", FName("TArray<FVector>"), "Points")
 
 private:
 	UPROPERTY(meta = (DataflowInput, DataflowIntrinsic))
diff --git a/Engine/Plugins/Experimental/GeometryCollectionPlugin/Source/GeometryCollectionNodes/Public/Dataflow/GeometryCollectionMakeNodes.h b/Engine/Plugins/Experimental/GeometryCollectionPlugin/Source/GeometryCollectionNodes/Public/Dataflow/GeometryCollectionMakeNodes.h
--- a/Engine/Plugins/Experimental/GeometryCollectionPlugin/Source/GeometryCollectionNodes/Public/Dataflow/GeometryCollectionMakeNodes.h
+++ b/Engine/Plugins/Experimental/GeometryCollectionPlugin/Source/GeometryCollectionNodes/Public/Dataflow/GeometryCollectionMakeNodes.h
@@ -78,7 +78,6 @@ struct FMakePointsDataflowNode : public FDataflowNode
 {
 	GENERATED_USTRUCT_BODY()
 	DATAFLOW_NODE_DEFINE_INTERNAL(FMakePointsDataflowNode, "MakePoints", "Generators|Point", "")
-	DATAFLOW_NODE_RENDER_TYPE("PointsRender", FName("TArray<FVector>"), "Points")
 
 public:
 	UPROPERTY(EditAnywhere, Category = "Point")
@@ -124,7 +123,6 @@ struct FMakeBoxDataflowNode : public FDataflowNode
 {
 	GENERATED_USTRUCT_BODY()
 	DATAFLOW_NODE_DEFINE_INTERNAL(FMakeBoxDataflowNode, "MakeBox", "Generators|Box", "")
-	DATAFLOW_NODE_RENDER_TYPE("SurfaceRender",FName("FBox"), "Box")
 
 public:
 	UPROPERTY(EditAnywhere, Category = "Box", meta = (DisplayName = "Input Data Type"));
@@ -169,7 +167,6 @@ struct FMakeSphereDataflowNode : public FDataflowNode
 {
 	GENERATED_USTRUCT_BODY()
 	DATAFLOW_NODE_DEFINE_INTERNAL(FMakeSphereDataflowNode, "MakeSphere", "Generators|Sphere", "")
-	DATAFLOW_NODE_RENDER_TYPE("SurfaceRender", FName("FSphere"), "Sphere")
 
 public:
 	UPROPERTY(EditAnywhere, Category = "Sphere", meta = (DataflowInput));
diff --git a/Engine/Plugins/Experimental/GeometryCollectionPlugin/Source/GeometryCollectionNodes/Public/Dataflow/GeometryCollectionMeshNodes.h b/Engine/Plugins/Experimental/GeometryCollectionPlugin/Source/GeometryCollectionNodes/Public/Dataflow/GeometryCollectionMeshNodes.h
--- a/Engine/Plugins/Experimental/GeometryCollectionPlugin/Source/GeometryCollectionNodes/Public/Dataflow/GeometryCollectionMeshNodes.h
+++ b/Engine/Plugins/Experimental/GeometryCollectionPlugin/Source/GeometryCollectionNodes/Public/Dataflow/GeometryCollectionMeshNodes.h
@@ -498,7 +498,6 @@ struct FGetMeshBoundingBoxDataflowNode : public FDataflowNode
 {
 	GENERATED_USTRUCT_BODY()
 	DATAFLOW_NODE_DEFINE_INTERNAL(FGetMeshBoundingBoxDataflowNode, "GetMeshBoundingBox", "Mesh", "Bounds Size Dimensions Extents Center")
-	DATAFLOW_NODE_RENDER_TYPE("SurfaceRender", FName("FBox"), "BoundingBox")
 
 private:
 	/** dynamic mesh to compute the bouning box from */
diff --git a/Engine/Plugins/Experimental/GeometryCollectionPlugin/Source/GeometryCollectionNodes/Public/Dataflow/GeometryCollectionNodes.h b/Engine/Plugins/Experimental/GeometryCollectionPlugin/Source/GeometryCollectionNodes/Public/Dataflow/GeometryCollectionNodes.h
--- a/Engine/Plugins/Experimental/GeometryCollectionPlugin/Source/GeometryCollectionNodes/Public/Dataflow/GeometryCollectionNodes.h
+++ b/Engine/Plugins/Experimental/GeometryCollectionPlugin/Source/GeometryCollectionNodes/Public/Dataflow/GeometryCollectionNodes.h
@@ -171,7 +171,6 @@ struct FBoundingBoxDataflowNode : public FDataflowNode
 {
 	GENERATED_USTRUCT_BODY()
 	DATAFLOW_NODE_DEFINE_INTERNAL(FBoundingBoxDataflowNode, "GetCollectionBoundingBox", "Utilities|Box", "Collection Bounds Size Dimensions Extents Center")
-	DATAFLOW_NODE_RENDER_TYPE("SurfaceRender",FName("FBox"), "BoundingBox")
 
 public:
 	/** Collection to compute the bouning box from */
@@ -205,7 +204,6 @@ struct FBoundingSphereDataflowNode : public FDataflowNode
 {
 	GENERATED_USTRUCT_BODY()
 	DATAFLOW_NODE_DEFINE_INTERNAL(FBoundingSphereDataflowNode, "BoundingSphere", "Utilities|Sphere", "")
-	DATAFLOW_NODE_RENDER_TYPE("SurfaceRender", FName("FSphere"), "BoundingSphere")
 
 public:
 	UPROPERTY(meta = (DataflowInput))
@@ -348,7 +346,6 @@ struct FExpandBoundingSphereDataflowNode : public FDataflowNode
 {
 	GENERATED_USTRUCT_BODY()
 	DATAFLOW_NODE_DEFINE_INTERNAL(FExpandBoundingSphereDataflowNode, "BreakBoundingSphere", "Utilities|Sphere", "Expand Radius Center Volume")
-	DATAFLOW_NODE_RENDER_TYPE("PointRender", FName("FVector"), "Center")
 
 public:
 	UPROPERTY(meta = (DataflowInput))
@@ -1979,7 +1976,6 @@ struct FVisualizeTetrahedronsDataflowNode : public FDataflowNode
 	DATAFLOW_NODE_DEFINE_INTERNAL(FVisualizeTetrahedronsDataflowNode, "VisualizeTetrahedrons", "Flesh|Utilities", "")
 	DATAFLOW_NODE_RENDER_TYPE_START()
 	DATAFLOW_NODE_RENDER_TYPE_ADD("TetrahedronRender", FGeometryCollection::StaticType(), "Collection")
-	DATAFLOW_NODE_RENDER_TYPE_ADD("PointsRender", FName("TArray<FVector>"), "Vertices")
 	DATAFLOW_NODE_RENDER_TYPE_END()
 
 private:
@@ -2045,9 +2041,6 @@ struct FCollectionToPointsDataflowNode : public FDataflowNode
 {
 	GENERATED_USTRUCT_BODY()
 	DATAFLOW_NODE_DEFINE_INTERNAL(FCollectionToPointsDataflowNode, "CollectionToPoints", "GeometryCollection|Utilities", "")
-	DATAFLOW_NODE_RENDER_TYPE_START()
-	DATAFLOW_NODE_RENDER_TYPE_ADD("PointsRender", FName("TArray<FVector>"), "Points")
-	DATAFLOW_NODE_RENDER_TYPE_END()
 
 private:
 	/** Collection storing the points */
@@ -2080,9 +2073,6 @@ struct FSpheresToPointsDataflowNode : public FDataflowNode
 {
 	GENERATED_USTRUCT_BODY()
 	DATAFLOW_NODE_DEFINE_INTERNAL(FSpheresToPointsDataflowNode, "SpheresToPoints", "GeometryCollection|Utilities", "")
-	DATAFLOW_NODE_RENDER_TYPE_START()
-	DATAFLOW_NODE_RENDER_TYPE_ADD("PointsRender", FName("TArray<FVector>"), "Points")
-	DATAFLOW_NODE_RENDER_TYPE_END()
 
 private:
 	/** Input Spheres  */
diff --git a/Engine/Plugins/Experimental/GeometryCollectionPlugin/Source/GeometryCollectionNodes/Public/Dataflow/GeometryCollectionSamplingNodes.h b/Engine/Plugins/Experimental/GeometryCollectionPlugin/Source/GeometryCollectionNodes/Public/Dataflow/GeometryCollectionSamplingNodes.h
--- a/Engine/Plugins/Experimental/GeometryCollectionPlugin/Source/GeometryCollectionNodes/Public/Dataflow/GeometryCollectionSamplingNodes.h
+++ b/Engine/Plugins/Experimental/GeometryCollectionPlugin/Source/GeometryCollectionNodes/Public/Dataflow/GeometryCollectionSamplingNodes.h
@@ -40,7 +40,6 @@ struct FFilterPointSetWithMeshDataflowNode : public FDataflowNode
 {
 	GENERATED_USTRUCT_BODY()
 	DATAFLOW_NODE_DEFINE_INTERNAL(FFilterPointSetWithMeshDataflowNode, "FilterPointsWithMesh", "PointSampling", "")
-	DATAFLOW_NODE_RENDER_TYPE("PointsRender", FName("TArray<FVector>"), "SamplePoints")
 
 public:
 	FFilterPointSetWithMeshDataflowNode(const UE::Dataflow::FNodeParameters& InParam, FGuid InGuid = FGuid::NewGuid());
@@ -94,7 +93,6 @@ struct FUniformPointSamplingDataflowNode : public FDataflowNode
 {
 	GENERATED_USTRUCT_BODY()
 	DATAFLOW_NODE_DEFINE_INTERNAL(FUniformPointSamplingDataflowNode, "UniformPointSampling", "PointSampling", "")
-	DATAFLOW_NODE_RENDER_TYPE("PointsRender", FName("TArray<FVector>"), "SamplePoints")
 
 public:
 	/** Mesh to sample points on */
@@ -160,7 +158,6 @@ struct FNonUniformPointSamplingDataflowNode : public FDataflowNode
 {
 	GENERATED_USTRUCT_BODY()
 	DATAFLOW_NODE_DEFINE_INTERNAL(FNonUniformPointSamplingDataflowNode, "NonUniformPointSampling", "PointSampling", "")
-	DATAFLOW_NODE_RENDER_TYPE("PointsRender", FName("TArray<FVector>"), "SamplePoints")
 
 public:
 	/** Mesh to sample points on */
@@ -245,7 +242,6 @@ struct FVertexWeightedPointSamplingDataflowNode : public FDataflowNode
 {
 	GENERATED_USTRUCT_BODY()
 	DATAFLOW_NODE_DEFINE_INTERNAL(FVertexWeightedPointSamplingDataflowNode, "VertexWeightedPointSampling", "PointSampling", "")
-	DATAFLOW_NODE_RENDER_TYPE("PointsRender", FName("TArray<FVector>"), "SamplePoints")
 
 public:
 	/** Mesh to sample points on */
