Geometry collection : fix Get bounding boxes from collection node to be resilient to empty collections and non matching selections

#Synced-CL 48601912

diff --git a/Engine/Plugins/Experimental/GeometryCollectionPlugin/Source/GeometryCollectionNodes/Private/Dataflow/GeometryCollectionNodes.cpp b/Engine/Plugins/Experimental/GeometryCollectionPlugin/Source/GeometryCollectionNodes/Private/Dataflow/GeometryCollectionNodes.cpp
--- a/Engine/Plugins/Experimental/GeometryCollectionPlugin/Source/GeometryCollectionNodes/Private/Dataflow/GeometryCollectionNodes.cpp
+++ b/Engine/Plugins/Experimental/GeometryCollectionPlugin/Source/GeometryCollectionNodes/Private/Dataflow/GeometryCollectionNodes.cpp
@@ -450,46 +450,45 @@ void FHashVectorDataflowNode::Evaluate(UE::Dataflow::FContext& Context, const FD
 	}
 }
 
+//-----------------------------------------------------------------------------------------------
 
 void FGetBoundingBoxesFromCollectionDataflowNode::Evaluate(UE::Dataflow::FContext& Context, const FDataflowOutput* Out) const
 {
-	if (Out->IsA<TArray<FBox>>(&BoundingBoxes))
+	if (Out->IsA(&BoundingBoxes))
 	{
-		const FManagedArrayCollection& InCollection = GetValue<FManagedArrayCollection>(Context, &Collection);
-		const FDataflowTransformSelection& InTransformSelection = GetValue<FDataflowTransformSelection>(Context, &TransformSelection);
+		const FManagedArrayCollection& InCollection = GetValue(Context, &Collection);
+		const FDataflowTransformSelection& InTransformSelection = GetValue(Context, &TransformSelection);
 
-		GeometryCollection::Facades::FBoundsFacade BoundsFacade(InCollection);
-		const TManagedArray<FBox>& InBoundingBoxes = BoundsFacade.GetBoundingBoxes();
+		TArray<FBox> BoundingBoxesArr;
 
+		GeometryCollection::Facades::FBoundsFacade BoundsFacade(InCollection);
 		GeometryCollection::Facades::FCollectionTransformFacade TransformFacade(InCollection);
-
-		TArray<FBox> BoundingBoxesArr;
-		for (int32 Idx = 0; Idx < InBoundingBoxes.Num(); ++Idx)
+		if (BoundsFacade.IsValid() && TransformFacade.IsValid())
 		{
-			const FBox BoundingBoxInBoneSpace = InBoundingBoxes[Idx];
+			const bool bValidSelection = InTransformSelection.IsValidForCollection(InCollection);
 
-			// Transform from BoneSpace to CollectionSpace
-			const FTransform CollectionSpaceTransform = TransformFacade.ComputeCollectionSpaceTransform(Idx);
-			const FBox BoundingBoxInCollectionSpace = BoundingBoxInBoneSpace.TransformBy(CollectionSpaceTransform);
-
-			if (IsConnected<FDataflowTransformSelection>(&TransformSelection))
+			const TManagedArray<FBox>& InBoundingBoxes = BoundsFacade.GetBoundingBoxes();
+			for (int32 Idx = 0; Idx < InBoundingBoxes.Num(); ++Idx)
 			{
-				if (InTransformSelection.IsSelected(Idx))
+				const FBox BoundingBoxInBoneSpace = InBoundingBoxes[Idx];
+
+				// Transform from BoneSpace to CollectionSpace
+				const FTransform CollectionSpaceTransform = TransformFacade.ComputeCollectionSpaceTransform(Idx);
+				const FBox BoundingBoxInCollectionSpace = BoundingBoxInBoneSpace.TransformBy(CollectionSpaceTransform);
+
+				if (!bValidSelection || InTransformSelection.IsSelected(Idx))
 				{
 					BoundingBoxesArr.Add(BoundingBoxInCollectionSpace);
 				}
 			}
-			else
-			{
-				BoundingBoxesArr.Add(BoundingBoxInCollectionSpace);
-			}
 
+			SetValue(Context, MoveTemp(BoundingBoxesArr), &BoundingBoxes);
 		}
-
-		SetValue(Context, MoveTemp(BoundingBoxesArr), &BoundingBoxes);
 	}
 }
 
+//-----------------------------------------------------------------------------------------------
+
 void FGetRootIndexFromCollectionDataflowNode::Evaluate(UE::Dataflow::FContext& Context, const FDataflowOutput* Out) const
 {
 	if (Out->IsA(&RootIndex))
