Dataflow : add nodes for box and collection extend and volume stats
- Volume stats for collection
- Extent stats for bounding box
- Bounding box from collection with extra ratio information

#Synced-CL 48601912

diff --git a/Engine/Plugins/Experimental/GeometryCollectionPlugin/Source/GeometryCollectionNodes/Private/Dataflow/GeometryCollectionNodes.cpp b/Engine/Plugins/Experimental/GeometryCollectionPlugin/Source/GeometryCollectionNodes/Private/Dataflow/GeometryCollectionNodes.cpp
--- a/Engine/Plugins/Experimental/GeometryCollectionPlugin/Source/GeometryCollectionNodes/Private/Dataflow/GeometryCollectionNodes.cpp
+++ b/Engine/Plugins/Experimental/GeometryCollectionPlugin/Source/GeometryCollectionNodes/Private/Dataflow/GeometryCollectionNodes.cpp
@@ -56,6 +56,8 @@ namespace UE::Dataflow
 		DATAFLOW_NODE_REGISTER_CREATION_FACTORY(FPrintStringDataflowNode);
 		DATAFLOW_NODE_REGISTER_CREATION_FACTORY(FLogStringDataflowNode);
 		DATAFLOW_NODE_REGISTER_CREATION_FACTORY(FBoundingBoxDataflowNode);
+		DATAFLOW_NODE_REGISTER_CREATION_FACTORY(FBoundingBoxRatioDataflowNode);
+		DATAFLOW_NODE_REGISTER_CREATION_FACTORY(FBoundingBoxAverageExtentsDataflowNode);
 		DATAFLOW_NODE_REGISTER_CREATION_FACTORY(FExpandBoundingBoxDataflowNode);
 		DATAFLOW_NODE_REGISTER_CREATION_FACTORY(FGetBoxLengthsDataflowNode);
 		DATAFLOW_NODE_REGISTER_CREATION_FACTORY(FExpandVectorDataflowNode);
@@ -233,7 +235,7 @@ void FBoundingBoxDataflowNode::Evaluate(UE::Dataflow::FContext& Context, const F
 {
 	if (Out->IsA(&BoundingBox) || Out->IsA(&Center) || Out->IsA(&Dimensions))
 	{
-		const FManagedArrayCollection& InCollection = GetValue<FManagedArrayCollection>(Context, &Collection);
+		const FManagedArrayCollection& InCollection = GetValue(Context, &Collection);
 
 		GeometryCollection::Facades::FBoundsFacade BoundsFacade(InCollection);
 		const FBox& BoundingBoxInCollectionSpace = BoundsFacade.GetBoundingBoxInCollectionSpace();
@@ -246,6 +248,75 @@ void FBoundingBoxDataflowNode::Evaluate(UE::Dataflow::FContext& Context, const F
 
 ///////////////////////////////////////////////////////////////////////////////////////////////
 
+FBoundingBoxRatioDataflowNode::FBoundingBoxRatioDataflowNode(const UE::Dataflow::FNodeParameters& InParam, FGuid InGuid)
+	: FDataflowNode(InParam, InGuid)
+{
+	RegisterInputConnection(&Collection);
+	RegisterInputConnection(&ReferenceBoxSize);
+
+	RegisterOutputConnection(&BoundingBox);
+	RegisterOutputConnection(&Dimensions);
+	RegisterOutputConnection(&BoxRatio);
+}
+
+void FBoundingBoxRatioDataflowNode::Evaluate(UE::Dataflow::FContext& Context, const FDataflowOutput* Out) const
+{
+	if (Out->IsA(&BoundingBox) || Out->IsA(&BoxRatio) || Out->IsA(&Dimensions))
+	{
+		const FManagedArrayCollection& InCollection = GetValue(Context, &Collection);
+		const float InReferenceBoxSize = GetValue(Context, &ReferenceBoxSize);
+		const double ReferenceBoxVolume = InReferenceBoxSize * InReferenceBoxSize * InReferenceBoxSize;
+
+		GeometryCollection::Facades::FBoundsFacade BoundsFacade(InCollection);
+		const FBox& BoundingBoxInCollectionSpace = BoundsFacade.GetBoundingBoxInCollectionSpace();
+		const FVector OutBoxDimension = BoundingBoxInCollectionSpace.GetSize();
+		const double BoxVolume = OutBoxDimension.X * OutBoxDimension.Y * OutBoxDimension.Z;
+
+		const float OutBoxRatio = static_cast<float>(BoxVolume / FMath::Max(ReferenceBoxVolume, UE_SMALL_NUMBER));
+
+		SetValue(Context, BoundingBoxInCollectionSpace, &BoundingBox);
+		SetValue(Context, OutBoxDimension, &Dimensions);
+		SetValue(Context, OutBoxRatio, &BoxRatio);
+	}
+}
+
+///////////////////////////////////////////////////////////////////////////////////////////////
+
+FBoundingBoxAverageExtentsDataflowNode::FBoundingBoxAverageExtentsDataflowNode(const UE::Dataflow::FNodeParameters& InParam, FGuid InGuid)
+	: FDataflowNode(InParam, InGuid)
+{
+	RegisterInputConnection(&BoundingBox);
+
+	RegisterOutputConnection(&AverageExtents);
+	RegisterOutputConnection(&AverageHalfExtents);
+	RegisterOutputConnection(&MinLength);
+	RegisterOutputConnection(&MaxLength);
+	RegisterOutputConnection(&LengthRatio);
+}
+
+void FBoundingBoxAverageExtentsDataflowNode::Evaluate(UE::Dataflow::FContext& Context, const FDataflowOutput* Out) const
+{
+	if (Out->IsA(&AverageExtents) || Out->IsA(&AverageHalfExtents) || Out->IsA(&MinLength) || Out->IsA(&MaxLength) || Out->IsA(&LengthRatio))
+	{
+		const FBox& InBox = GetValue(Context, &BoundingBox);
+		const FVector BoxSize = InBox.GetSize();
+
+		const float OutAverageExtents = (BoxSize.X + BoxSize.Y + BoxSize.Z) / 3.f;
+		const float OutAverageHalfExtents = OutAverageExtents * 0.5f;
+		const float OutMinLength = BoxSize.GetMin();
+		const float OutMaxLength = BoxSize.GetMax();
+		const float OutLengthRatio = (OutMinLength / FMath::Max(OutMaxLength, UE_SMALL_NUMBER));
+
+		SetValue(Context, OutAverageExtents, &AverageExtents);
+		SetValue(Context, OutAverageHalfExtents, &AverageHalfExtents);
+		SetValue(Context, OutMinLength, &MinLength);
+		SetValue(Context, OutMaxLength, &MaxLength);
+		SetValue(Context, OutLengthRatio, &LengthRatio);
+	}
+}
+
+///////////////////////////////////////////////////////////////////////////////////////////////
+
 void FGetBoxLengthsDataflowNode::Evaluate(UE::Dataflow::FContext& Context, const FDataflowOutput* Out) const
 {
 	if (Out->IsA(&Lengths))
@@ -264,6 +335,18 @@ void FGetBoxLengthsDataflowNode::Evaluate(UE::Dataflow::FContext& Context, const
 	}
 }
 
+///////////////////////////////////////////////////////////////////////////////////////////////
+
+FExpandBoundingBoxDataflowNode::FExpandBoundingBoxDataflowNode(const UE::Dataflow::FNodeParameters& InParam, FGuid InGuid)
+	: FDataflowNode(InParam, InGuid)
+{
+	RegisterInputConnection(&BoundingBox);
+	RegisterOutputConnection(&Min);
+	RegisterOutputConnection(&Max);
+	RegisterOutputConnection(&Center);
+	RegisterOutputConnection(&HalfExtents);
+	RegisterOutputConnection(&Volume);
+}
 
 void FExpandBoundingBoxDataflowNode::Evaluate(UE::Dataflow::FContext& Context, const FDataflowOutput* Out) const
 {
@@ -291,6 +374,7 @@ void FExpandBoundingBoxDataflowNode::Evaluate(UE::Dataflow::FContext& Context, c
 	}
 }
 
+///////////////////////////////////////////////////////////////////////////////////////////////
 
 void FExpandVectorDataflowNode::Evaluate(UE::Dataflow::FContext& Context, const FDataflowOutput* Out) const
 {
diff --git a/Engine/Plugins/Experimental/GeometryCollectionPlugin/Source/GeometryCollectionNodes/Private/Dataflow/GeometryCollectionUtilityNodes.cpp b/Engine/Plugins/Experimental/GeometryCollectionPlugin/Source/GeometryCollectionNodes/Private/Dataflow/GeometryCollectionUtilityNodes.cpp
--- a/Engine/Plugins/Experimental/GeometryCollectionPlugin/Source/GeometryCollectionNodes/Private/Dataflow/GeometryCollectionUtilityNodes.cpp
+++ b/Engine/Plugins/Experimental/GeometryCollectionPlugin/Source/GeometryCollectionNodes/Private/Dataflow/GeometryCollectionUtilityNodes.cpp
@@ -24,6 +24,9 @@
 
 #include UE_INLINE_GENERATED_CPP_BY_NAME(GeometryCollectionUtilityNodes)
 
+
+#define LOCTEXT_NAMESPACE "GeometryCollectionUtilityNodes"
+
 namespace UE::Dataflow
 {
 
@@ -45,6 +48,7 @@ namespace UE::Dataflow
 		DATAFLOW_NODE_REGISTER_CREATION_FACTORY(FRecomputeNormalsInGeometryCollectionDataflowNode);
 		DATAFLOW_NODE_REGISTER_CREATION_FACTORY(FResampleGeometryCollectionDataflowNode);
 		DATAFLOW_NODE_REGISTER_CREATION_FACTORY(FValidateGeometryCollectionDataflowNode);
+		DATAFLOW_NODE_REGISTER_CREATION_FACTORY(FComputeVolumeStatsDataflowNode);
 	}
 }
 
@@ -960,7 +964,7 @@ FUpdateVolumeAttributesDataflowNode::FUpdateVolumeAttributesDataflowNode(const U
 	: FDataflowNode(InParam, InGuid)
 {
 	RegisterInputConnection(&Collection);
-	RegisterOutputConnection(&Collection);
+	RegisterOutputConnection(&Collection, &Collection);
 }
 
 void FUpdateVolumeAttributesDataflowNode::Evaluate(UE::Dataflow::FContext& Context, const FDataflowOutput* Out) const
@@ -976,6 +980,151 @@ void FUpdateVolumeAttributesDataflowNode::Evaluate(UE::Dataflow::FContext& Conte
 	}
 }
 
+/* --------------------------------------------------------------------------------------------------------------------------- */
+
+namespace ComputeVolumeStatsDataflowNode::Private
+{
+	static const FName VolumeAttributeName("Volume");
+	static const FName TransformToConvexIndicesAttributeName("TransformToConvexIndices");
+
+	static float ComputeGeometryVolume(const FManagedArrayCollection& InCollection, int32 TransformIndex)
+	{
+		float OutVolume = 0.f;
+		if (const TManagedArray<float>* VolumeAttributePtr = InCollection.FindAttribute<float>(VolumeAttributeName, FTransformCollection::TransformGroup))
+		{
+			if (VolumeAttributePtr->IsValidIndex(TransformIndex))
+			{
+				OutVolume = (*VolumeAttributePtr)[TransformIndex];
+			}
+		}
+		return OutVolume;
+	}
+
+	static float ComputeConvexVolume(FGeometryCollection& InCollection, int32 TransformIndex)
+	{
+		float ConvexVolume = 0;
+
+		// we want a fully non optimized encompassing convex, we don't care about overlaps as well
+		constexpr double FractionAllowRemove = 0.f;
+		constexpr double SimplificationDistanceThreshold = 0.0f; // no simplification
+		constexpr double CanExceedFraction = 1000000; // arbitrary large number since we don't want thegeneration to fail
+		constexpr EConvexOverlapRemoval OverlapRemovalMethod = EConvexOverlapRemoval::None; // no overlap removal
+		constexpr double OverlapRemovalShrinkPercent = 0.0f;
+
+		FGeometryCollectionConvexUtility::FGeometryCollectionConvexData ConvexData =
+			FGeometryCollectionConvexUtility::CreateNonOverlappingConvexHullData(
+				&InCollection,
+				FractionAllowRemove,
+				SimplificationDistanceThreshold,
+				CanExceedFraction,
+				OverlapRemovalMethod,
+				OverlapRemovalShrinkPercent
+			);
+
+		if (const TManagedArray<TSet<int32>>* TransformToConvexIndicesAttributePtr = InCollection.FindAttribute<TSet<int32>>(TransformToConvexIndicesAttributeName, FTransformCollection::TransformGroup))
+		{
+			if (TransformToConvexIndicesAttributePtr->IsValidIndex(TransformIndex))
+			{
+				// we should normally have only one convex but for safety let's sum up all the convex that may have been generated for this transform index
+				ConvexVolume = 0;
+
+				const TSet<int32>& ConvexIndices = (*TransformToConvexIndicesAttributePtr)[TransformIndex];
+				for (const int32 ConvexIndex : ConvexIndices)
+				{
+					if (ConvexData.ConvexHull.IsValidIndex(ConvexIndex))
+					{
+						if (const Chaos::FConvexPtr ConvexPtr = ConvexData.ConvexHull[ConvexIndex])
+						{
+							ConvexVolume += ConvexPtr->GetVolume();
+						}
+					}
+				}
+			}
+		}
+		return ConvexVolume;
+	}
+}
+
+FComputeVolumeStatsDataflowNode::FComputeVolumeStatsDataflowNode(const UE::Dataflow::FNodeParameters& InParam, FGuid InGuid)
+	: FDataflowNode(InParam, InGuid)
+{
+	RegisterInputConnection(&Collection);
+	RegisterInputConnection(&TransformIndex);
+
+	RegisterOutputConnection(&Collection, &Collection);
+	RegisterOutputConnection(&GeometryVolume);
+	RegisterOutputConnection(&ConvexVolume);
+	RegisterOutputConnection(&VolumeRatio);
+	RegisterOutputConnection(&VolumePercentage);
+}
+
+void FComputeVolumeStatsDataflowNode::Evaluate(UE::Dataflow::FContext& Context, const FDataflowOutput* Out) const
+{
+	if (Out->IsA(&Collection) || 
+		Out->IsA(&GeometryVolume) || 
+		Out->IsA(&ConvexVolume) || 
+		Out->IsA(&VolumeRatio) || 
+		Out->IsA(&VolumePercentage)
+		)
+	{
+		float OutGeometryVolume = 0.f;
+		float OutConvexVolume = 0.f;
+
+		const FManagedArrayCollection& InCollection = GetValue(Context, &Collection);
+		if (TUniquePtr<FGeometryCollection> GeomCollection = TUniquePtr<FGeometryCollection>(InCollection.NewCopy<FGeometryCollection>()))
+		{
+			const int32 NumTransforms = GeomCollection->NumElements(FGeometryCollection::TransformGroup);
+
+			Chaos::Facades::FCollectionHierarchyFacade HierarchyFacade(*GeomCollection);
+
+			const int32 RootIndex = HierarchyFacade.GetRootIndex();
+			int32 InTransformIndex = GetValue(Context, &TransformIndex, TransformIndex);
+			if (InTransformIndex < 0 || InTransformIndex >= NumTransforms)
+			{
+				InTransformIndex = RootIndex;
+				const FText Message = FText::Format(LOCTEXT("ComputeVolumeStats_InvalidTransformIndexUsingRoot", "Invalid Transform Index, trying to use root index : {0}"), FText::AsNumber(InTransformIndex));
+				Context.Warning(Message, this, Out);
+			}
+			if (InTransformIndex < 0 || InTransformIndex >= NumTransforms)
+			{
+				const FText Message = FText::Format(LOCTEXT("ComputeVolumeStats_InvalidTransformAndRootIndex", "Invalid Transform and root Index : {0}"), FText::AsNumber(InTransformIndex));
+				Context.Error(Message, this, Out);
+			}
+			else
+			{
+				// make sure we have a volume attribute 
+				FGeometryCollectionConvexUtility::SetVolumeAttributes(GeomCollection.Get());
+				if (!GeomCollection->HasAttribute(ComputeVolumeStatsDataflowNode::Private::VolumeAttributeName, FGeometryCollection::TransformGroup))
+				{
+					const FText Message = LOCTEXT("ComputeVolumeStats_NoVolumeAttribute", "Could not generate volume attribute");
+					Context.Error(Message, this, Out);
+				}
+				else
+				{
+					OutGeometryVolume = ComputeVolumeStatsDataflowNode::Private::ComputeGeometryVolume(*GeomCollection, InTransformIndex);
+					OutConvexVolume = ComputeVolumeStatsDataflowNode::Private::ComputeConvexVolume(*GeomCollection, InTransformIndex);
+				}
+			}
+		}
+		else
+		{
+			const FText Message = LOCTEXT("ComputeVolumeStats_GeometryCollectionCopyFail", "Failed to make a Geometry collection out of the input collection");
+			Context.Error(Message, this, Out);
+		}
+
+		const float OutVolumeRatio = OutGeometryVolume / FMath::Max(OutConvexVolume, UE_SMALL_NUMBER);
+		const float OutVolumePercentage = OutVolumeRatio * 100.0f;
+	
+		SafeForwardInput(Context, &Collection, &Collection);
+		SetValue(Context, OutGeometryVolume, &GeometryVolume);
+		SetValue(Context, OutConvexVolume, &ConvexVolume);
+		SetValue(Context, OutVolumeRatio, &VolumeRatio);
+		SetValue(Context, OutVolumePercentage, &VolumePercentage);
+	}
+}
+
+
+
 /* --------------------------------------------------------------------------------------------------------------------------- */
 
 FGetConvexHullVolumeDataflowNode::FGetConvexHullVolumeDataflowNode(const UE::Dataflow::FNodeParameters& InParam, FGuid InGuid)
@@ -1245,3 +1394,5 @@ void FValidateGeometryCollectionDataflowNode::Evaluate(UE::Dataflow::FContext& C
 		SetValue(Context, MoveTemp(InCollection), &Collection);
 	}
 }
+
+#undef LOCTEXT_NAMESPACE
\ No newline at end of file
diff --git a/Engine/Plugins/Experimental/GeometryCollectionPlugin/Source/GeometryCollectionNodes/Public/Dataflow/GeometryCollectionNodes.h b/Engine/Plugins/Experimental/GeometryCollectionPlugin/Source/GeometryCollectionNodes/Public/Dataflow/GeometryCollectionNodes.h
--- a/Engine/Plugins/Experimental/GeometryCollectionPlugin/Source/GeometryCollectionNodes/Public/Dataflow/GeometryCollectionNodes.h
+++ b/Engine/Plugins/Experimental/GeometryCollectionPlugin/Source/GeometryCollectionNodes/Public/Dataflow/GeometryCollectionNodes.h
@@ -195,9 +195,44 @@ public:
 };
 
 /**
- *
- * Description for this node
- *
+ * Get the bounding box a collection and compute its box ratio from a reference size
+ * The reference size is the length of a perfect cube representing a reference volume
+ * 
+ */
+USTRUCT(meta = (DataflowGeometryCollection))
+struct FBoundingBoxRatioDataflowNode : public FDataflowNode
+{
+	GENERATED_USTRUCT_BODY()
+	DATAFLOW_NODE_DEFINE_INTERNAL(FBoundingBoxRatioDataflowNode, "GetCollectionBoundingBoxRatio", "Utilities|Box", "Collection Reference Resolution Bounds Size Dimensions Extents Center")
+
+public:
+	FBoundingBoxRatioDataflowNode(const UE::Dataflow::FNodeParameters& InParam, FGuid InGuid = FGuid::NewGuid());
+
+private:
+	/** Collection to compute the bounding box and ratio from */
+	UPROPERTY(meta = (DataflowInput))
+	FManagedArrayCollection Collection;
+
+	/** Reference box size to use to compute the ratio */
+	UPROPERTY(EditAnywhere, Category = "Proximity", meta = (DataflowInput, ClampMin = "0", UIMin = "0", Units = cm))
+	float ReferenceBoxSize = 100;
+
+	/** Geometric bounding box of the input collection */
+	UPROPERTY(meta = (DataflowOutput))
+	FBox BoundingBox = FBox(ForceInit);
+
+	/** Dimensions of the resulting bounding box in centimers */
+	UPROPERTY(meta = (DataflowOutput))
+	FVector Dimensions = FVector::ZeroVector;
+
+	UPROPERTY(meta = (DataflowOutput))
+	float BoxRatio = 1.0;
+
+	virtual void Evaluate(UE::Dataflow::FContext& Context, const FDataflowOutput* Out) const override;
+};
+
+/**
+ * Compute the boudning sphere of the geometrry of a collection 
  */
 USTRUCT(meta = (DataflowGeometryCollection))
 struct FBoundingSphereDataflowNode : public FDataflowNode
@@ -289,12 +324,48 @@ public:
 
 };
 
+/**
+ * Get the bounding box average extents information
+ */
+USTRUCT(meta = (DataflowGeometryCollection))
+struct FBoundingBoxAverageExtentsDataflowNode : public FDataflowNode
+{
+	GENERATED_USTRUCT_BODY()
+	DATAFLOW_NODE_DEFINE_INTERNAL(FBoundingBoxAverageExtentsDataflowNode, "GetBoxAverageExtents", "Utilities|Box", "BoundingBox size extent ratio length min max")
+
+public:
+	FBoundingBoxAverageExtentsDataflowNode(const UE::Dataflow::FNodeParameters& InParam, FGuid InGuid = FGuid::NewGuid());
+
+private:
+	/** Bounding box to get the information from */
+	UPROPERTY(meta = (DataflowInput))
+	FBox BoundingBox = FBox(ForceInit);;
+
+	/** Average of the 3 axis extents */
+	UPROPERTY(meta = (DataflowOutput))
+	float AverageExtents = 0.f;
+
+	/** Average of the 3 axis half extents */
+	UPROPERTY(meta = (DataflowOutput))
+	float AverageHalfExtents = 0.f;
+
+	/** Minimum edge length of the box */
+	UPROPERTY(meta = (DataflowOutput))
+	float MinLength = 0.f;
 
+	/** Maximum edge length of the box */
+	UPROPERTY(meta = (DataflowOutput))
+	float MaxLength = 0.f;
+
+	/** Ratio of the min edge length over the max edge length */
+	UPROPERTY(meta = (DataflowOutput))
+	float LengthRatio = 1.f;
+
+	virtual void Evaluate(UE::Dataflow::FContext& Context, const FDataflowOutput* Out) const override;
+};
 
 /**
- *
- * Description for this node
- *
+ * Break a bounding box into individual components ( Min/Max corners, Center, HalfExtents and Volume )
  */
 USTRUCT()
 struct FExpandBoundingBoxDataflowNode : public FDataflowNode
@@ -303,37 +374,33 @@ struct FExpandBoundingBoxDataflowNode : public FDataflowNode
 	DATAFLOW_NODE_DEFINE_INTERNAL(FExpandBoundingBoxDataflowNode, "BreakBox", "Utilities|Box", "Bounds Bounding Expand Min Max Center Half Extents Volume")
 
 public:
+	/** Bounding box to get the information from */
 	UPROPERTY(meta = (DataflowInput, DisplayName = "Box"))
 	FBox BoundingBox = FBox(ForceInit);;
 
+	/** Minimum corner ofthe bounding box */
 	UPROPERTY(meta = (DataflowOutput))
 	FVector Min = FVector(0.0);
 
+	/** Maximum corner ofthe bounding box */
 	UPROPERTY(meta = (DataflowOutput))
 	FVector Max = FVector(0.0);
 
+	/** Center point of the bounding box */
 	UPROPERTY(meta = (DataflowOutput))
 	FVector Center = FVector(0.0);
 
+	/** Distances from the center to any corner along the three axis */
 	UPROPERTY(meta = (DataflowOutput))
 	FVector HalfExtents = FVector(0.0);
 
+	/** Total volume of the bounding box */
 	UPROPERTY(meta = (DataflowOutput))
 	float Volume = 0.0;
 
-	FExpandBoundingBoxDataflowNode(const UE::Dataflow::FNodeParameters& InParam, FGuid InGuid = FGuid::NewGuid())
-		: FDataflowNode(InParam, InGuid)
-	{
-		RegisterInputConnection(&BoundingBox);
-		RegisterOutputConnection(&Min);
-		RegisterOutputConnection(&Max);
-		RegisterOutputConnection(&Center);
-		RegisterOutputConnection(&HalfExtents);
-		RegisterOutputConnection(&Volume);
-	}
+	FExpandBoundingBoxDataflowNode(const UE::Dataflow::FNodeParameters& InParam, FGuid InGuid = FGuid::NewGuid());
 
 	virtual void Evaluate(UE::Dataflow::FContext& Context, const FDataflowOutput* Out) const override;
-
 };
 
 /**
diff --git a/Engine/Plugins/Experimental/GeometryCollectionPlugin/Source/GeometryCollectionNodes/Public/Dataflow/GeometryCollectionUtilityNodes.h b/Engine/Plugins/Experimental/GeometryCollectionPlugin/Source/GeometryCollectionNodes/Public/Dataflow/GeometryCollectionUtilityNodes.h
--- a/Engine/Plugins/Experimental/GeometryCollectionPlugin/Source/GeometryCollectionNodes/Public/Dataflow/GeometryCollectionUtilityNodes.h
+++ b/Engine/Plugins/Experimental/GeometryCollectionPlugin/Source/GeometryCollectionNodes/Public/Dataflow/GeometryCollectionUtilityNodes.h
@@ -768,10 +768,10 @@ USTRUCT(meta = (DataflowGeometryCollection))
 struct FUpdateVolumeAttributesDataflowNode : public FDataflowNode
 {
 	GENERATED_USTRUCT_BODY()
-	DATAFLOW_NODE_DEFINE_INTERNAL(FUpdateVolumeAttributesDataflowNode, "UpdateVolumeAttributes", "GeometryCollection|Utilities", "")
+	DATAFLOW_NODE_DEFINE_INTERNAL(FUpdateVolumeAttributesDataflowNode, "UpdateVolumeAttributes", "GeometryCollection|Utilities", "collection")
 
 public:
-	UPROPERTY(meta = (DataflowInput, DataflowOutput))
+	UPROPERTY(meta = (DataflowInput, DataflowOutput, DataflowPassthrough = Collection, DataflowIntrinsic))
 	FManagedArrayCollection Collection;
 
 	FUpdateVolumeAttributesDataflowNode(const UE::Dataflow::FNodeParameters& InParam, FGuid InGuid = FGuid::NewGuid());
@@ -779,6 +779,46 @@ public:
 	virtual void Evaluate(UE::Dataflow::FContext& Context, const FDataflowOutput* Out) const override;
 };
 
+/**
+ * Compute information about the volume of a collection specific transform or root
+ */
+USTRUCT(meta = (DataflowGeometryCollection))
+struct FComputeVolumeStatsDataflowNode : public FDataflowNode
+{
+	GENERATED_USTRUCT_BODY()
+	DATAFLOW_NODE_DEFINE_INTERNAL(FComputeVolumeStatsDataflowNode, "ComputeVolumeStatsData", "GeometryCollection|Utilities", "collection empty space percent size convex geometry")
+
+public:
+	FComputeVolumeStatsDataflowNode(const UE::Dataflow::FNodeParameters& InParam, FGuid InGuid = FGuid::NewGuid());
+
+private:
+	/** Collection to collect the stats from */
+	UPROPERTY(meta = (DataflowInput, DataflowOutput, DataflowPassthrough = Collection, DataflowIntrinsic))
+	FManagedArrayCollection Collection;
+
+	/** Selection to consider , if set to -1 or not a valid transform in the collection, the root will be considered */
+	UPROPERTY(meta = (DataflowInput))
+	int32 TransformIndex = INDEX_NONE;
+
+	/** Volume of the geometry in cm3 */
+	UPROPERTY(meta = (DataflowOutput));
+	float GeometryVolume = 0.f;
+
+	/** Volume of the single wrapping convex around the geometry in cm3 */
+	UPROPERTY(meta = (DataflowOutput));
+	float ConvexVolume = 0.f;
+
+	/** Ratio of the geometry volume in relation to the wrapping convex volume ( GeometryVolume / ConvexVolume ) */
+	UPROPERTY(meta = (DataflowOutput));
+	float VolumeRatio = 1.0f;
+
+	/** Percentage of the geometry volume in relation to the wrapping convex volume ( GeometryVolume / ConvexVolume * 100 ) */
+	UPROPERTY(meta = (DataflowOutput));
+	float VolumePercentage = 100.f;
+
+	virtual void Evaluate(UE::Dataflow::FContext& Context, const FDataflowOutput* Out) const override;
+};
+
 /**
  * Get the sum of volumes of the convex hulls on the selected nodes
  */
