Dataflow add the Empty space  cube size to the Volume stats node

#Synced-CL 48816559

diff --git a/Engine/Plugins/Experimental/GeometryCollectionPlugin/Source/GeometryCollectionNodes/Private/Dataflow/GeometryCollectionUtilityNodes.cpp b/Engine/Plugins/Experimental/GeometryCollectionPlugin/Source/GeometryCollectionNodes/Private/Dataflow/GeometryCollectionUtilityNodes.cpp
--- a/Engine/Plugins/Experimental/GeometryCollectionPlugin/Source/GeometryCollectionNodes/Private/Dataflow/GeometryCollectionUtilityNodes.cpp
+++ b/Engine/Plugins/Experimental/GeometryCollectionPlugin/Source/GeometryCollectionNodes/Private/Dataflow/GeometryCollectionUtilityNodes.cpp
@@ -1056,6 +1056,7 @@ FComputeVolumeStatsDataflowNode::FComputeVolumeStatsDataflowNode(const UE::Dataf
 	RegisterOutputConnection(&ConvexVolume);
 	RegisterOutputConnection(&VolumeRatio);
 	RegisterOutputConnection(&VolumePercentage);
+	RegisterOutputConnection(&EmptySpaceCubeSize);
 }
 
 void FComputeVolumeStatsDataflowNode::Evaluate(UE::Dataflow::FContext& Context, const FDataflowOutput* Out) const
@@ -1064,7 +1065,8 @@ void FComputeVolumeStatsDataflowNode::Evaluate(UE::Dataflow::FContext& Context,
 		Out->IsA(&GeometryVolume) || 
 		Out->IsA(&ConvexVolume) || 
 		Out->IsA(&VolumeRatio) || 
-		Out->IsA(&VolumePercentage)
+		Out->IsA(&VolumePercentage) ||
+		Out->IsA(&EmptySpaceCubeSize)
 		)
 	{
 		float OutGeometryVolume = 0.f;
@@ -1114,12 +1116,14 @@ void FComputeVolumeStatsDataflowNode::Evaluate(UE::Dataflow::FContext& Context,
 
 		const float OutVolumeRatio = OutGeometryVolume / FMath::Max(OutConvexVolume, UE_SMALL_NUMBER);
 		const float OutVolumePercentage = OutVolumeRatio * 100.0f;
-	
+		const float OutEmptySpaceCubeSize = FMath::Pow(FMath::Max(0, (OutConvexVolume - OutGeometryVolume)), 1.f / 3.f);
+
 		SafeForwardInput(Context, &Collection, &Collection);
 		SetValue(Context, OutGeometryVolume, &GeometryVolume);
 		SetValue(Context, OutConvexVolume, &ConvexVolume);
 		SetValue(Context, OutVolumeRatio, &VolumeRatio);
 		SetValue(Context, OutVolumePercentage, &VolumePercentage);
+		SetValue(Context, OutEmptySpaceCubeSize, &EmptySpaceCubeSize);
 	}
 }
 
diff --git a/Engine/Plugins/Experimental/GeometryCollectionPlugin/Source/GeometryCollectionNodes/Public/Dataflow/GeometryCollectionUtilityNodes.h b/Engine/Plugins/Experimental/GeometryCollectionPlugin/Source/GeometryCollectionNodes/Public/Dataflow/GeometryCollectionUtilityNodes.h
--- a/Engine/Plugins/Experimental/GeometryCollectionPlugin/Source/GeometryCollectionNodes/Public/Dataflow/GeometryCollectionUtilityNodes.h
+++ b/Engine/Plugins/Experimental/GeometryCollectionPlugin/Source/GeometryCollectionNodes/Public/Dataflow/GeometryCollectionUtilityNodes.h
@@ -816,6 +816,11 @@ private:
 	UPROPERTY(meta = (DataflowOutput));
 	float VolumePercentage = 100.f;
 
+	/** Cube size of the empty space volume (cubic root of the volume = Edge size of a cube of similar volume) */
+	UPROPERTY(meta = (DataflowOutput));
+	float EmptySpaceCubeSize = 100.f;
+
+
 	virtual void Evaluate(UE::Dataflow::FContext& Context, const FDataflowOutput* Out) const override;
 };
 
