Add required sphere covering input to more dataflow nodes + add node to generate a sphere covering

#Synced-CL 49177856

diff --git a/Engine/Plugins/Experimental/GeometryCollectionPlugin/Source/GeometryCollectionNodes/Private/Dataflow/GeometryCollectionUtilityNodes.cpp b/Engine/Plugins/Experimental/GeometryCollectionPlugin/Source/GeometryCollectionNodes/Private/Dataflow/GeometryCollectionUtilityNodes.cpp
--- a/Engine/Plugins/Experimental/GeometryCollectionPlugin/Source/GeometryCollectionNodes/Private/Dataflow/GeometryCollectionUtilityNodes.cpp
+++ b/Engine/Plugins/Experimental/GeometryCollectionPlugin/Source/GeometryCollectionNodes/Private/Dataflow/GeometryCollectionUtilityNodes.cpp
@@ -11,6 +11,7 @@
 #include "GeometryCollection/GeometryCollectionConvexUtility.h"
 #include "GeometryCollection/Facades/CollectionHierarchyFacade.h"
 #include "GeometryCollection/Facades/CollectionTransformSelectionFacade.h"
+#include "GeometryCollectionToDynamicMesh.h"
 #include "Dataflow/DataflowDebugDrawInterface.h"
 #include "Dataflow/DataflowDebugDraw.h"
 #include "Dataflow/DataflowDynamicMeshDebugDrawMesh.h"
@@ -33,6 +34,7 @@ namespace UE::Dataflow
 	void GeometryCollectionUtilityNodes()
 	{
 		DATAFLOW_NODE_REGISTER_CREATION_FACTORY(FMakeDataflowConvexDecompositionSettingsNode);
+		DATAFLOW_NODE_REGISTER_CREATION_FACTORY(FCreateNegativeSpaceSphereCovering);
 		DATAFLOW_NODE_REGISTER_CREATION_FACTORY(FCreateLeafConvexHullsDataflowNode);
 		DATAFLOW_NODE_REGISTER_CREATION_FACTORY(FSimplifyConvexHullsDataflowNode);
 		DATAFLOW_NODE_REGISTER_CREATION_FACTORY(FCreateNonOverlappingConvexHullsDataflowNode);
@@ -473,6 +475,7 @@ FGenerateClusterConvexHullsFromLeafHullsDataflowNode::FGenerateClusterConvexHull
 	RegisterInputConnection(&ErrorTolerance).SetCanHidePin(true).SetPinIsHidden(true);
 	RegisterInputConnection(&OptionalSelectionFilter);
 	RegisterInputConnection(&bAllowMergingLeafHulls).SetCanHidePin(true).SetPinIsHidden(true);
+	RegisterInputConnection(&RequireNegativeSpaceCovering).SetCanHidePin(true).SetPinIsHidden(true);
 	RegisterInputConnection(&bProtectNegativeSpace).SetCanHidePin(true).SetPinIsHidden(true);
 	RegisterInputConnection(&TargetNumSamples).SetCanHidePin(true).SetPinIsHidden(true);
 	RegisterInputConnection(&MinSampleSpacing).SetCanHidePin(true).SetPinIsHidden(true);
@@ -522,6 +525,13 @@ void FGenerateClusterConvexHullsFromLeafHullsDataflowNode::Evaluate(UE::Dataflow
 				bHasNegativeSpace = UE::FractureEngine::Convex::ComputeConvexHullsNegativeSpace(*GeomCollection, NegativeSpace, NegativeSpaceSettings, bHasSelectionFilter, SelectionArray);
 			}
 
+			const FDataflowSphereCovering& RequireCovering = GetValue(Context, &RequireNegativeSpaceCovering);
+			if (RequireCovering.Spheres.Num() > 0)
+			{
+				bHasNegativeSpace = true;
+				NegativeSpace.Append(RequireCovering.Spheres);
+			}
+
 			const int32 InConvexCount = GetValue(Context, &ConvexCount);
 			const double InErrorToleranceInCm = GetValue(Context, &ErrorTolerance);
 			FGeometryCollectionConvexUtility::FClusterConvexHullSettings HullMergeSettings(InConvexCount, InErrorToleranceInCm, bPreferExternalCollisionShapes);
@@ -598,6 +608,103 @@ void FGenerateClusterConvexHullsFromLeafHullsDataflowNode::DebugDraw(UE::Dataflo
 
 /* --------------------------------------------------------------------------------------------------------------------------- */
 
+FCreateNegativeSpaceSphereCovering::FCreateNegativeSpaceSphereCovering(const UE::Dataflow::FNodeParameters& InParam, FGuid InGuid)
+	: FDataflowNode(InParam, InGuid)
+{
+	RegisterInputConnection(&Collection);
+	RegisterInputConnection(&OptionalSelectionFilter);
+	RegisterInputConnection(&TargetNumSamples).SetCanHidePin(true).SetPinIsHidden(true);
+	RegisterInputConnection(&MinSampleSpacing).SetCanHidePin(true).SetPinIsHidden(true);
+	RegisterInputConnection(&NegativeSpaceTolerance).SetCanHidePin(true).SetPinIsHidden(true);
+	RegisterInputConnection(&MinRadius).SetCanHidePin(true).SetPinIsHidden(true);
+
+	RegisterOutputConnection(&SphereCovering);
+}
+
+void FCreateNegativeSpaceSphereCovering::Evaluate(UE::Dataflow::FContext& Context, const FDataflowOutput* Out) const
+{
+	if (Out->IsA(&SphereCovering))
+	{
+		const FManagedArrayCollection& InCollection = GetValue<FManagedArrayCollection>(Context, &Collection);
+		FDataflowSphereCovering ResultSphereCovering;
+
+		UE::Geometry::FNegativeSpaceSampleSettings NegativeSpaceSettings;
+		NegativeSpaceSettings.TargetNumSamples = GetValue(Context, &TargetNumSamples);
+		NegativeSpaceSettings.MinRadius = GetValue(Context, &MinRadius);
+		NegativeSpaceSettings.ReduceRadiusMargin = GetValue(Context, &NegativeSpaceTolerance);
+		NegativeSpaceSettings.MinSpacing = GetValue(Context, &MinSampleSpacing);
+		NegativeSpaceSettings.SampleMethod = ConvertNegativeSpaceSampleMethodDataflowEnum(SampleMethod);
+		NegativeSpaceSettings.bRequireSearchSampleCoverage = bRequireSearchSampleCoverage;
+		NegativeSpaceSettings.bOnlyConnectedToHull = bOnlyConnectedToHull;
+		NegativeSpaceSettings.Sanitize();
+
+		FDataflowTransformSelection InTransformSelection = GetValue(Context, &OptionalSelectionFilter);
+
+		// If optional selection input not connected select everything by default
+		if (!IsConnected(&OptionalSelectionFilter))
+		{
+			InTransformSelection.InitializeFromCollection(InCollection, true);
+		}
+		if (InTransformSelection.AnySelected())
+		{
+			TArray<int32> TransformSelectionArray = InTransformSelection.AsArray();
+
+			GeometryCollection::Facades::FCollectionTransformSelectionFacade SelectionFacade(InCollection);
+			Chaos::Facades::FCollectionHierarchyFacade HierarchyFacade(InCollection);
+			TArray<int32> LeafSelectionArray = TransformSelectionArray;
+			SelectionFacade.ConvertSelectionToRigidNodes(LeafSelectionArray);
+
+			UE::Geometry::FGeometryCollectionToDynamicMeshes CollectionToMeshes;
+			UE::Geometry::FGeometryCollectionToDynamicMeshes::FToMeshOptions ToMeshOptions;
+			ToMeshOptions.bWeldVertices = false;
+			ToMeshOptions.bSaveIsolatedVertices = false;
+			if (CollectionToMeshes.InitFromTransformSelection(InCollection, LeafSelectionArray, ToMeshOptions)
+				&& !CollectionToMeshes.Meshes.IsEmpty())
+			{
+				UE::Geometry::FDynamicMesh3 CombinedMesh = MoveTemp(*CollectionToMeshes.Meshes[0].Mesh);
+				for (int32 MeshIdx = 1; MeshIdx < CollectionToMeshes.Meshes.Num(); ++MeshIdx)
+				{
+					CombinedMesh.AppendWithOffsets(*CollectionToMeshes.Meshes[MeshIdx].Mesh);
+				}
+
+				UE::Geometry::FDynamicMeshAABBTree3 Tree(&CombinedMesh, true);
+				UE::Geometry::TFastWindingTree<UE::Geometry::FDynamicMesh3> Winding(&Tree, true);
+				constexpr bool bHasFlippedTriangles = false;
+				ResultSphereCovering.Spheres.AddNegativeSpace(Winding, NegativeSpaceSettings, bHasFlippedTriangles);
+			}
+		}
+
+		SetValue(Context, MoveTemp(ResultSphereCovering), &SphereCovering);
+	}
+}
+
+#if WITH_EDITOR
+bool FCreateNegativeSpaceSphereCovering::CanDebugDrawViewMode(const FName& ViewModeName) const
+{
+	return ViewModeName == UE::Dataflow::FDataflowConstruction3DViewMode::Name;
+}
+
+void FCreateNegativeSpaceSphereCovering::DebugDraw(UE::Dataflow::FContext& Context, IDataflowDebugDrawInterface& DataflowRenderingInterface, const FDebugDrawParameters& DebugDrawParameters) const
+{
+	using namespace UE::Geometry;
+
+	if ((DebugDrawParameters.bNodeIsSelected || DebugDrawParameters.bNodeIsPinned))
+	{
+		if (SphereCoveringDebugDrawRenderSettings.bDisplaySphereCovering)
+		{
+			if (const FDataflowOutput* SphereCoveringOutput = FindOutput(&SphereCovering))
+			{
+				const FDataflowSphereCovering& OutSpheres = SphereCoveringOutput->GetValue(Context, SphereCovering);
+
+				UE::Dataflow::Convex::SphereCoveringDebugDrawProc(DataflowRenderingInterface, OutSpheres, SphereCoveringDebugDrawRenderSettings);
+			}
+		}
+	}
+}
+#endif
+
+/* --------------------------------------------------------------------------------------------------------------------------- */
+
 FGenerateClusterConvexHullsFromChildrenHullsDataflowNode::FGenerateClusterConvexHullsFromChildrenHullsDataflowNode(const UE::Dataflow::FNodeParameters& InParam, FGuid InGuid)
 	: FDataflowNode(InParam, InGuid)
 {
@@ -606,6 +713,7 @@ FGenerateClusterConvexHullsFromChildrenHullsDataflowNode::FGenerateClusterConvex
 	RegisterInputConnection(&ErrorTolerance).SetCanHidePin(true).SetPinIsHidden(true);
 	RegisterInputConnection(&OptionalSelectionFilter);
 	RegisterInputConnection(&bAllowMergingLeafHulls).SetCanHidePin(true).SetPinIsHidden(true);
+	RegisterInputConnection(&RequireNegativeSpaceCovering).SetCanHidePin(true).SetPinIsHidden(true);
 	RegisterInputConnection(&bProtectNegativeSpace).SetCanHidePin(true).SetPinIsHidden(true);
 	RegisterInputConnection(&TargetNumSamples).SetCanHidePin(true).SetPinIsHidden(true);
 	RegisterInputConnection(&MinSampleSpacing).SetCanHidePin(true).SetPinIsHidden(true);
@@ -655,6 +763,13 @@ void FGenerateClusterConvexHullsFromChildrenHullsDataflowNode::Evaluate(UE::Data
 				bHasNegativeSpace = UE::FractureEngine::Convex::ComputeConvexHullsNegativeSpace(*GeomCollection, NegativeSpace, NegativeSpaceSettings, bHasSelectionFilter, SelectionArray);
 			}
 
+			const FDataflowSphereCovering& RequireCovering = GetValue(Context, &RequireNegativeSpaceCovering);
+			if (RequireCovering.Spheres.Num() > 0)
+			{
+				bHasNegativeSpace = true;
+				NegativeSpace.Append(RequireCovering.Spheres);
+			}
+
 			const int32 InConvexCount = GetValue(Context, &ConvexCount);
 			const double InErrorToleranceInCm = GetValue(Context, &ErrorTolerance);
 			FGeometryCollectionConvexUtility::FClusterConvexHullSettings HullMergeSettings(InConvexCount, InErrorToleranceInCm, bPreferExternalCollisionShapes);
@@ -739,6 +854,7 @@ FMergeConvexHullsDataflowNode::FMergeConvexHullsDataflowNode(const UE::Dataflow:
 	RegisterInputConnection(&ErrorTolerance).SetCanHidePin(true).SetPinIsHidden(true);
 	RegisterInputConnection(&OptionalSelectionFilter);
 	RegisterInputConnection(&bProtectNegativeSpace).SetCanHidePin(true).SetPinIsHidden(true);
+	RegisterInputConnection(&RequireNegativeSpaceCovering).SetCanHidePin(true).SetPinIsHidden(true);
 	RegisterInputConnection(&TargetNumSamples).SetCanHidePin(true).SetPinIsHidden(true);
 	RegisterInputConnection(&MinSampleSpacing).SetCanHidePin(true).SetPinIsHidden(true);
 	RegisterInputConnection(&NegativeSpaceTolerance).SetCanHidePin(true).SetPinIsHidden(true);
@@ -906,6 +1022,12 @@ void FMergeConvexHullsDataflowNode::Evaluate(UE::Dataflow::FContext& Context, co
 		{
 			bHasPrecomputedNegativeSpace = UE::FractureEngine::Convex::ComputeConvexHullsNegativeSpace(InCollection, NegativeSpace, NegativeSpaceSettings, bHasSelectionFilter, SelectionArray, false);
 		}
+		const FDataflowSphereCovering& RequireCovering = GetValue(Context, &RequireNegativeSpaceCovering);
+		if (RequireCovering.Spheres.Num() > 0)
+		{
+			bHasPrecomputedNegativeSpace = true;
+			NegativeSpace.Append(RequireCovering.Spheres);
+		}
 
 		const int32 InMaxConvexCount = GetValue(Context, &MaxConvexCount);
 		const double InErrorToleranceInCm = GetValue(Context, &ErrorTolerance);
diff --git a/Engine/Plugins/Experimental/GeometryCollectionPlugin/Source/GeometryCollectionNodes/Public/Dataflow/GeometryCollectionUtilityNodes.h b/Engine/Plugins/Experimental/GeometryCollectionPlugin/Source/GeometryCollectionNodes/Public/Dataflow/GeometryCollectionUtilityNodes.h
--- a/Engine/Plugins/Experimental/GeometryCollectionPlugin/Source/GeometryCollectionNodes/Public/Dataflow/GeometryCollectionUtilityNodes.h
+++ b/Engine/Plugins/Experimental/GeometryCollectionPlugin/Source/GeometryCollectionNodes/Public/Dataflow/GeometryCollectionUtilityNodes.h
@@ -356,6 +356,75 @@ enum class ENegativeSpaceSampleMethodDataflowEnum : uint8
 	NavigableVoxelSearch
 };
 
+/**
+ *
+ * Generates a negative space sphere covering for a given geometry collection or subset of bones in the collection
+ *
+ */
+USTRUCT(meta = (DataflowGeometryCollection))
+struct FCreateNegativeSpaceSphereCovering : public FDataflowNode
+{
+	GENERATED_USTRUCT_BODY()
+	DATAFLOW_NODE_DEFINE_INTERNAL(FCreateNegativeSpaceSphereCovering, "CreateNegativeSpaceSphereCovering", "GeometryCollection|Utilities", "")
+
+private:
+	UPROPERTY(meta = (DataflowInput))
+	FManagedArrayCollection Collection;
+
+	UPROPERTY(meta = (DataflowOutput))
+	FDataflowSphereCovering SphereCovering;
+
+	/** Optional transform selection -- if not provided, covering for all leaf geometry will be computed. */
+	UPROPERTY(meta = (DataflowInput))
+	FDataflowTransformSelection OptionalSelectionFilter;
+
+
+	/** Method to use to find and sample negative space */
+	UPROPERTY(EditAnywhere, Category = NegativeSpace)
+	ENegativeSpaceSampleMethodDataflowEnum SampleMethod = ENegativeSpaceSampleMethodDataflowEnum::NavigableVoxelSearch;
+
+	/** Whether to require that all candidate locations identified by Voxel Search are covered by negative space samples, up to the specified Min Sample Spacing. Only applies to Voxel Search. */
+	UPROPERTY(EditAnywhere, Category = NegativeSpace, meta = (EditCondition = "SampleMethod != ENegativeSpaceSampleMethodDataflowEnum::Uniform"))
+	bool bRequireSearchSampleCoverage = true;
+
+	/** When performing Voxel Search, only look for negative space that is connected out to the convex hull. This removes inaccessable internal negative space from consideration. Only applies to Voxel Search. */
+	UPROPERTY(EditAnywhere, Category = NegativeSpace, meta = (EditCondition = "SampleMethod != ENegativeSpaceSampleMethodDataflowEnum::Uniform"))
+	bool bOnlyConnectedToHull = true;
+
+	/** Approximate number of spheres to consider when covering negative space. Only applicable with the Uniform Sample Method or if Require Search Sample Coverage is disabled. */
+	UPROPERTY(EditAnywhere, Category = NegativeSpace, meta = (DataflowInput, ClampMin = 1, EditCondition = "SampleMethod == ENegativeSpaceSampleMethodDataflowEnum::Uniform || !bRequireSearchSampleCoverage"))
+	int32 TargetNumSamples = 50;
+
+	/** Minimum desired spacing between spheres; if > 0, will attempt not to place sphere centers closer than this */
+	UPROPERTY(EditAnywhere, Category = NegativeSpace, meta = (DataflowInput, ClampMin = 0, Units = cm))
+	double MinSampleSpacing = 0.0;
+
+	/** Amount of space to leave between input geometry and protected negative space */
+	UPROPERTY(EditAnywhere, Category = NegativeSpace, meta = (DataflowInput, ClampMin = .01, UIMin = .1, Units = cm))
+	double NegativeSpaceTolerance = 2.0;
+
+	/** Spheres smaller than this are not included in the negative space */
+	UPROPERTY(EditAnywhere, Category = NegativeSpace, meta = (DataflowInput, ClampMin = 0, Units = cm))
+	double MinRadius = 10.0;
+
+	UPROPERTY(EditAnywhere, Category = "Debug Draw - Sphere Covering", meta = (ShowOnlyInnerProperties))
+	FDataflowNodeSphereCoveringDebugDrawSettings SphereCoveringDebugDrawRenderSettings;
+
+
+	virtual void Evaluate(UE::Dataflow::FContext& Context, const FDataflowOutput* Out) const override;
+#if WITH_EDITOR
+	virtual bool CanDebugDraw() const override
+	{
+		return true;
+	}
+	virtual bool CanDebugDrawViewMode(const FName& ViewModeName) const override;
+	virtual void DebugDraw(UE::Dataflow::FContext& Context, IDataflowDebugDrawInterface& DataflowRenderingInterface, const FDebugDrawParameters& DebugDrawParameters) const override;
+#endif
+
+public:
+	FCreateNegativeSpaceSphereCovering(const UE::Dataflow::FNodeParameters& InParam, FGuid InGuid = FGuid::NewGuid());
+};
+
 /**
  *
  * Generates cluster convex hulls for leafs hulls
@@ -411,7 +480,11 @@ private:
 	UPROPERTY(EditAnywhere, Category = "Convex", meta = (DataflowInput))
 	bool bAllowMergingLeafHulls = false;
 
-	/** Whether to use a sphere cover to define negative space that should not be covered by convex hulls */
+	/** Negative space that must be used. Applies in addition to any automatically-detected covering added via bProtectNegativeSpace. */
+	UPROPERTY(meta = (DataflowInput))
+	FDataflowSphereCovering RequireNegativeSpaceCovering;
+
+	/** Whether to generate a sphere cover to define negative space that should not be covered by convex hulls */
 	UPROPERTY(EditAnywhere, Category = NegativeSpace, meta = (DataflowInput))
 	bool bProtectNegativeSpace = false;
 
@@ -524,7 +597,11 @@ private:
 	UPROPERTY(EditAnywhere, Category = "Convex", meta = (DataflowInput))
 	bool bAllowMergingLeafHulls = false;
 
-	/** Whether to use a sphere cover to define negative space that should not be covered by convex hulls */
+	/** Negative space that must be used. Applies in addition to any automatically-detected covering added via bProtectNegativeSpace. */
+	UPROPERTY(meta = (DataflowInput))
+	FDataflowSphereCovering RequireNegativeSpaceCovering;
+
+	/** Whether to generate a sphere cover to define negative space that should not be covered by convex hulls */
 	UPROPERTY(EditAnywhere, Category = NegativeSpace, meta = (DataflowInput))
 	bool bProtectNegativeSpace = false;
 
@@ -699,7 +776,11 @@ private:
 	UPROPERTY(EditAnywhere, Category = "Convex", meta = (Units = cm, EditCondition = "MergeProximityFilter != EConvexHullProximityFilter::None"))
 	float MergeProximityDistanceThreshold = .1f;
 
-	/** Whether to use a sphere cover to define negative space that should not be covered by convex hulls */
+	/** Negative space that must be used. Applies in addition to any automatically-detected covering added via bProtectNegativeSpace. */
+	UPROPERTY(meta = (DataflowInput))
+	FDataflowSphereCovering RequireNegativeSpaceCovering;
+
+	/** Whether to generate a sphere cover to define negative space that should not be covered by convex hulls */
 	UPROPERTY(EditAnywhere, Category = NegativeSpace, meta = (DataflowInput))
 	bool bProtectNegativeSpace = false;
 
