Add missing curve selection support for the dataflow select by attr node

#Synced-CL 49646132

diff --git a/Engine/Plugins/Experimental/GeometryCollectionPlugin/Source/GeometryCollectionNodes/Private/Dataflow/GeometryCollectionSelectionNodes.cpp b/Engine/Plugins/Experimental/GeometryCollectionPlugin/Source/GeometryCollectionNodes/Private/Dataflow/GeometryCollectionSelectionNodes.cpp
--- a/Engine/Plugins/Experimental/GeometryCollectionPlugin/Source/GeometryCollectionNodes/Private/Dataflow/GeometryCollectionSelectionNodes.cpp
+++ b/Engine/Plugins/Experimental/GeometryCollectionPlugin/Source/GeometryCollectionNodes/Private/Dataflow/GeometryCollectionSelectionNodes.cpp
@@ -1488,7 +1488,8 @@ void FCollectionSelectionByAttrDataflowNode::Evaluate(UE::Dataflow::FContext& Co
 		Out->IsA(&FaceSelection) ||
 		Out->IsA(&TransformSelection) ||
 		Out->IsA(&GeometrySelection) ||
-		Out->IsA(&MaterialSelection))
+		Out->IsA(&MaterialSelection) ||
+		Out->IsA(&CurveSelection))
 	{
 		const FManagedArrayCollection& InCollection = GetValue(Context, &Collection);
 		FCollectionAttributeKey InAttributeKey = GetValue(Context, &AttributeKey);
