Dataflow: add vertex selection by box, sphere and plane nodes
- also fix crashing in the way we select vertices

#Synced-CL 49754362

diff --git a/Engine/Plugins/Experimental/GeometryCollectionPlugin/Source/GeometryCollectionNodes/Private/Dataflow/GeometryCollectionSelectionNodes.cpp b/Engine/Plugins/Experimental/GeometryCollectionPlugin/Source/GeometryCollectionNodes/Private/Dataflow/GeometryCollectionSelectionNodes.cpp
--- a/Engine/Plugins/Experimental/GeometryCollectionPlugin/Source/GeometryCollectionNodes/Private/Dataflow/GeometryCollectionSelectionNodes.cpp
+++ b/Engine/Plugins/Experimental/GeometryCollectionPlugin/Source/GeometryCollectionNodes/Private/Dataflow/GeometryCollectionSelectionNodes.cpp
@@ -8,6 +8,7 @@
 #endif
 
 #include "Dataflow/DataflowDebugDrawInterface.h"
+#include "Dataflow/DataflowSimpleDebugDrawMesh.h"
 
 #include "Engine/StaticMesh.h"
 #include "GeometryCollection/GeometryCollectionObject.h"
@@ -41,6 +42,44 @@
 
 namespace UE::Dataflow
 {
+	namespace SelectionNodes::Private
+	{
+		static void DebugDrawBox(const FBox& InBox, const FTransform& InBoxTransform, IDataflowDebugDrawInterface& DataflowRenderingInterface)
+		{
+			DataflowRenderingInterface.SetLineWidth(1.0);
+			DataflowRenderingInterface.SetWireframe(true);
+			DataflowRenderingInterface.SetWorldPriority();
+			DataflowRenderingInterface.SetColor(FLinearColor::Red);
+
+			const FVector TransformedCenter = InBox.GetCenter() + InBoxTransform.GetTranslation();
+			const FVector ScaledExtent = InBox.GetExtent() * InBoxTransform.GetScale3D();
+			DataflowRenderingInterface.DrawBox(ScaledExtent, InBoxTransform.GetRotation(), TransformedCenter, 1.0);
+		}
+
+		static void DebugDrawSphere(const FSphere& InSphere, IDataflowDebugDrawInterface& DataflowRenderingInterface)
+		{
+			DataflowRenderingInterface.SetLineWidth(1.0);
+			DataflowRenderingInterface.SetWireframe(true);
+			DataflowRenderingInterface.SetWorldPriority();
+			DataflowRenderingInterface.SetColor(FLinearColor::Red);
+
+			DataflowRenderingInterface.DrawSphere(InSphere.Center, InSphere.W);
+		}
+
+		static void DebugDrawPlane(const FTransform& InPlaneTransform, float Size, int32 VtxPerEdge, IDataflowDebugDrawInterface& DataflowRenderingInterface)
+		{
+			DataflowRenderingInterface.SetLineWidth(1.0);
+			DataflowRenderingInterface.SetShaded(false);
+			DataflowRenderingInterface.SetWireframe(true);
+			DataflowRenderingInterface.SetWorldPriority();
+			DataflowRenderingInterface.SetColor(FLinearColor::Red);
+
+			FSimpleDebugDrawMesh Mesh;
+			Mesh.MakeRectangleMesh(FVector::ZeroVector, Size, Size, VtxPerEdge, VtxPerEdge);
+			Mesh.TransformVertices(InPlaneTransform);
+			DataflowRenderingInterface.DrawMesh(Mesh);
+		}
+	}
 
 	void GeometryCollectionSelectionNodes()
 	{
@@ -67,6 +106,9 @@ namespace UE::Dataflow
 		DATAFLOW_NODE_REGISTER_CREATION_FACTORY(FCollectionTransformSelectionInBoxDataflowNode);
 		DATAFLOW_NODE_REGISTER_CREATION_FACTORY(FCollectionTransformSelectionInSphereDataflowNode);
 		DATAFLOW_NODE_REGISTER_CREATION_FACTORY(FCollectionTransformSelectionByFloatAttrDataflowNode);
+		DATAFLOW_NODE_REGISTER_CREATION_FACTORY(FCollectionVertexSelectionByBoxDataflowNode);
+		DATAFLOW_NODE_REGISTER_CREATION_FACTORY(FCollectionVertexSelectionBySphereDataflowNode);
+		DATAFLOW_NODE_REGISTER_CREATION_FACTORY(FCollectionVertexSelectionByPlaneDataflowNode);
 		DATAFLOW_NODE_REGISTER_CREATION_FACTORY(FSelectFloatArrayIndicesInRangeDataflowNode);
 		DATAFLOW_NODE_REGISTER_CREATION_FACTORY(FCollectionTransformSelectionByIntAttrDataflowNode);
 		DATAFLOW_NODE_REGISTER_CREATION_FACTORY(FCollectionVertexSelectionCustomDataflowNode);
@@ -845,22 +887,180 @@ void FCollectionTransformSelectionInBoxDataflowNode::DebugDraw(UE::Dataflow::FCo
 {
 	if ((DebugDrawParameters.bNodeIsSelected || DebugDrawParameters.bNodeIsPinned))
 	{
-		const FBox& InBox = GetValue<FBox>(Context, &Box);
-		const FTransform& InTransform = GetValue<FTransform>(Context, &Transform);
+		const FBox& InBox = GetValue(Context, &Box);
+		const FTransform& InTransform = GetValue(Context, &Transform);
 
-		DataflowRenderingInterface.SetLineWidth(1.0);
-		DataflowRenderingInterface.SetWireframe(true);
-		DataflowRenderingInterface.SetWorldPriority();
-		DataflowRenderingInterface.SetColor(FLinearColor::Red);
+		UE::Dataflow::SelectionNodes::Private::DebugDrawBox(InBox, InTransform, DataflowRenderingInterface);
+	}
+}
+
+#endif
+
+//////////////////////////////////////////////////////////////////////////////////////////////////
+
+FCollectionVertexSelectionByBoxDataflowNode::FCollectionVertexSelectionByBoxDataflowNode(const UE::Dataflow::FNodeParameters& InParam, FGuid InGuid)
+	: FDataflowNode(InParam, InGuid)
+{
+	RegisterInputConnection(&Collection);
+	RegisterInputConnection(&Box);
+	RegisterInputConnection(&BoxTransform);
+
+	RegisterOutputConnection(&Collection, &Collection);
+	RegisterOutputConnection(&VertexSelection);
+}
+
+void FCollectionVertexSelectionByBoxDataflowNode::Evaluate(UE::Dataflow::FContext& Context, const FDataflowOutput* Out) const
+{
+	if (Out->IsA(&Collection))
+	{
+		SafeForwardInput(Context, &Collection, &Collection);
+	}
+	else if (Out->IsA(&VertexSelection))
+	{
+		const FManagedArrayCollection& InCollection = GetValue(Context, &Collection);
+		const FBox& InBox = GetValue(Context, &Box);
+		const FTransform& InBoxTransform = GetValue(Context, &BoxTransform);
+
+		GeometryCollection::Facades::FCollectionTransformSelectionFacade TransformSelectionFacade(InCollection);
+		TArray<bool> SelectedVertices;
+		TransformSelectionFacade.SelectVerticesInTransformedBox(InBox, InBoxTransform, SelectedVertices);
+
+		FDataflowVertexSelection OutVertexSelection;
+		OutVertexSelection.InitializeFromCollection(InCollection, false);
+		OutVertexSelection.SetFromArray(SelectedVertices);
+
+		SetValue(Context, MoveTemp(OutVertexSelection), &VertexSelection);
+	}
+}
+
+#if WITH_EDITOR
+
+bool FCollectionVertexSelectionByBoxDataflowNode::CanDebugDrawViewMode(const FName& ViewModeName) const
+{
+	return ViewModeName == UE::Dataflow::FDataflowConstruction3DViewMode::Name;
+}
+
+void FCollectionVertexSelectionByBoxDataflowNode::DebugDraw(UE::Dataflow::FContext& Context, IDataflowDebugDrawInterface& DataflowRenderingInterface, const FDebugDrawParameters& DebugDrawParameters) const
+{
+	if ((DebugDrawParameters.bNodeIsSelected || DebugDrawParameters.bNodeIsPinned))
+	{
+		const FBox& InBox = GetValue(Context, &Box);
+		const FTransform& InBoxTransform = GetValue(Context, &BoxTransform);
+		UE::Dataflow::SelectionNodes::Private::DebugDrawBox(InBox, InBoxTransform, DataflowRenderingInterface);
+	}
+}
+
+#endif
+
+//////////////////////////////////////////////////////////////////////////////////////////////////
+
+FCollectionVertexSelectionBySphereDataflowNode::FCollectionVertexSelectionBySphereDataflowNode(const UE::Dataflow::FNodeParameters& InParam, FGuid InGuid)
+	: FDataflowNode(InParam, InGuid)
+{
+	RegisterInputConnection(&Collection);
+	RegisterInputConnection(&Sphere);
+
+	RegisterOutputConnection(&Collection, &Collection);
+	RegisterOutputConnection(&VertexSelection);
+}
+
+void FCollectionVertexSelectionBySphereDataflowNode::Evaluate(UE::Dataflow::FContext& Context, const FDataflowOutput* Out) const
+{
+	if (Out->IsA(&Collection))
+	{
+		SafeForwardInput(Context, &Collection, &Collection);
+	}
+	else if (Out->IsA(&VertexSelection))
+	{
+		const FManagedArrayCollection& InCollection = GetValue(Context, &Collection);
+		const FSphere& InSphere = GetValue(Context, &Sphere);
+
+		GeometryCollection::Facades::FCollectionTransformSelectionFacade TransformSelectionFacade(InCollection);
+		TArray<bool> SelectedVertices;
+		TransformSelectionFacade.SelectVerticesInSphere(InSphere, SelectedVertices);
+
+		FDataflowVertexSelection OutVertexSelection;
+		OutVertexSelection.InitializeFromCollection(InCollection, false);
+		OutVertexSelection.SetFromArray(SelectedVertices);
+
+		SetValue(Context, MoveTemp(OutVertexSelection), &VertexSelection);
+	}
+}
+
+#if WITH_EDITOR
+
+bool FCollectionVertexSelectionBySphereDataflowNode::CanDebugDrawViewMode(const FName& ViewModeName) const
+{
+	return ViewModeName == UE::Dataflow::FDataflowConstruction3DViewMode::Name;
+}
+
+void FCollectionVertexSelectionBySphereDataflowNode::DebugDraw(UE::Dataflow::FContext& Context, IDataflowDebugDrawInterface& DataflowRenderingInterface, const FDebugDrawParameters& DebugDrawParameters) const
+{
+	if ((DebugDrawParameters.bNodeIsSelected || DebugDrawParameters.bNodeIsPinned))
+	{
+		const FSphere& InSphere = GetValue(Context, &Sphere);
+		UE::Dataflow::SelectionNodes::Private::DebugDrawSphere(InSphere, DataflowRenderingInterface);
+	}
+}
+
+#endif
+
+//////////////////////////////////////////////////////////////////////////////////////////////////
 
-		const FVector TransformedCenter = InBox.GetCenter() + InTransform.GetTranslation();
-		const FVector ScaledExtent = InBox.GetExtent() * InTransform.GetScale3D();
-		DataflowRenderingInterface.DrawBox(ScaledExtent, InTransform.GetRotation(), TransformedCenter, 1.0);
+FCollectionVertexSelectionByPlaneDataflowNode::FCollectionVertexSelectionByPlaneDataflowNode(const UE::Dataflow::FNodeParameters& InParam, FGuid InGuid)
+	: FDataflowNode(InParam, InGuid)
+{
+	RegisterInputConnection(&Collection);
+	RegisterInputConnection(&PlaneTransform);
+
+	RegisterOutputConnection(&Collection, &Collection);
+	RegisterOutputConnection(&VertexSelection);
+}
+
+void FCollectionVertexSelectionByPlaneDataflowNode::Evaluate(UE::Dataflow::FContext& Context, const FDataflowOutput* Out) const
+{
+	if (Out->IsA(&Collection))
+	{
+		SafeForwardInput(Context, &Collection, &Collection);
+	}
+	else if (Out->IsA(&VertexSelection))
+	{
+		const FManagedArrayCollection& InCollection = GetValue(Context, &Collection);
+		const FTransform& InPlaneTransform = GetValue(Context, &PlaneTransform);
+
+		GeometryCollection::Facades::FCollectionTransformSelectionFacade TransformSelectionFacade(InCollection);
+		TArray<bool> SelectedVertices;
+
+		const FPlane InPlane(InPlaneTransform.GetLocation(), InPlaneTransform.GetUnitAxis(EAxis::Z));
+		TransformSelectionFacade.SelectVerticesOnPlaneSide(InPlane, bPositiveSide, SelectedVertices);
+
+		FDataflowVertexSelection OutVertexSelection;
+		OutVertexSelection.InitializeFromCollection(InCollection, false);
+		OutVertexSelection.SetFromArray(SelectedVertices);
+
+		SetValue(Context, MoveTemp(OutVertexSelection), &VertexSelection);
+	}
+}
+
+#if WITH_EDITOR
+
+bool FCollectionVertexSelectionByPlaneDataflowNode::CanDebugDrawViewMode(const FName& ViewModeName) const
+{
+	return ViewModeName == UE::Dataflow::FDataflowConstruction3DViewMode::Name;
+}
+
+void FCollectionVertexSelectionByPlaneDataflowNode::DebugDraw(UE::Dataflow::FContext& Context, IDataflowDebugDrawInterface& DataflowRenderingInterface, const FDebugDrawParameters& DebugDrawParameters) const
+{
+	if ((DebugDrawParameters.bNodeIsSelected || DebugDrawParameters.bNodeIsPinned))
+	{
+		const FTransform& InPlaneTransform = GetValue(Context, &PlaneTransform);
+		UE::Dataflow::SelectionNodes::Private::DebugDrawPlane(InPlaneTransform, PlaneSize, PlaneResolution, DataflowRenderingInterface);
 	}
 }
 
 #endif
 
+
 //////////////////////////////////////////////////////////////////////////////////////////////////
 
 void FCollectionTransformSelectionInSphereDataflowNode::Evaluate(UE::Dataflow::FContext& Context, const FDataflowOutput* Out) const
diff --git a/Engine/Plugins/Experimental/GeometryCollectionPlugin/Source/GeometryCollectionNodes/Public/Dataflow/GeometryCollectionSelectionNodes.h b/Engine/Plugins/Experimental/GeometryCollectionPlugin/Source/GeometryCollectionNodes/Public/Dataflow/GeometryCollectionSelectionNodes.h
--- a/Engine/Plugins/Experimental/GeometryCollectionPlugin/Source/GeometryCollectionNodes/Public/Dataflow/GeometryCollectionSelectionNodes.h
+++ b/Engine/Plugins/Experimental/GeometryCollectionPlugin/Source/GeometryCollectionNodes/Public/Dataflow/GeometryCollectionSelectionNodes.h
@@ -996,7 +996,6 @@ public:
 #endif //WITH_EDITOR
 };
 
-
 /**
  *
  * Selects bones if their Vertices/BoundingBox/Centroid in a sphere
@@ -1052,6 +1051,129 @@ public:
 };
 
 
+/**
+ *
+ * Selects vertices they are contained in a transformed box
+ *
+ */
+USTRUCT(meta = (DataflowGeometryCollection))
+struct FCollectionVertexSelectionByBoxDataflowNode : public FDataflowNode
+{
+	GENERATED_USTRUCT_BODY()
+	DATAFLOW_NODE_DEFINE_INTERNAL(FCollectionVertexSelectionByBoxDataflowNode, "CollectionVertexSelectByBox", "GeometryCollection|Selection|Vertex", "In BoundingBox Points Transformed")
+
+public:
+	FCollectionVertexSelectionByBoxDataflowNode(const UE::Dataflow::FNodeParameters& InParam, FGuid InGuid = FGuid::NewGuid());
+
+private:
+	/** GeometryCollection for the selection */
+	UPROPERTY(meta = (DataflowInput, DataflowOutput, DataflowPassthrough = "Collection", DataflowIntrinsic))
+	FManagedArrayCollection Collection;
+
+	/** Box to contain Vertices */
+	UPROPERTY(EditAnywhere, Category = "Select", meta = (DataflowInput))
+	FBox Box = FBox(FVector(-50), FVector(+50));
+
+	/** Transform for the box */
+	UPROPERTY(EditAnywhere, Category = "Select", meta = (DataflowInput))
+	FTransform BoxTransform;
+
+	/** selected vertices */
+	UPROPERTY(meta = (DataflowOutput))
+	FDataflowVertexSelection VertexSelection;
+
+	virtual void Evaluate(UE::Dataflow::FContext& Context, const FDataflowOutput* Out) const override;
+
+#if WITH_EDITOR
+	virtual bool CanDebugDrawViewMode(const FName& ViewModeName) const override;
+	virtual void DebugDraw(UE::Dataflow::FContext& Context, IDataflowDebugDrawInterface& DataflowRenderingInterface, const FDebugDrawParameters& DebugDrawParameters) const override;
+#endif //WITH_EDITOR
+};
+
+
+/**
+ *
+ * Selects vertices they are contained in a sphere
+ *
+ */
+USTRUCT(meta = (DataflowGeometryCollection))
+struct FCollectionVertexSelectionBySphereDataflowNode : public FDataflowNode
+{
+	GENERATED_USTRUCT_BODY()
+	DATAFLOW_NODE_DEFINE_INTERNAL(FCollectionVertexSelectionBySphereDataflowNode, "CollectionVertexSelectBySphere", "GeometryCollection|Selection|Vertex", "In Points Radius Center")
+
+public:
+	FCollectionVertexSelectionBySphereDataflowNode(const UE::Dataflow::FNodeParameters& InParam, FGuid InGuid = FGuid::NewGuid());
+
+private:
+	/** GeometryCollection for the selection */
+	UPROPERTY(meta = (DataflowInput, DataflowOutput, DataflowPassthrough = "Collection", DataflowIntrinsic))
+	FManagedArrayCollection Collection;
+
+	/** sphere to select vertices */
+	UPROPERTY(EditAnywhere, Category = "Select", meta = (DataflowInput))
+	FSphere Sphere = FSphere(FVector(0, 0, 0), 50.0);
+
+	/** selected vertices */
+	UPROPERTY(meta = (DataflowOutput))
+	FDataflowVertexSelection VertexSelection;
+
+	virtual void Evaluate(UE::Dataflow::FContext& Context, const FDataflowOutput* Out) const override;
+
+#if WITH_EDITOR
+	virtual bool CanDebugDrawViewMode(const FName& ViewModeName) const override;
+	virtual void DebugDraw(UE::Dataflow::FContext& Context, IDataflowDebugDrawInterface& DataflowRenderingInterface, const FDebugDrawParameters& DebugDrawParameters) const override;
+#endif //WITH_EDITOR
+};
+
+/**
+ *
+ * Selects vertices they are contained on a place side 
+ *
+ */
+USTRUCT(meta = (DataflowGeometryCollection))
+struct FCollectionVertexSelectionByPlaneDataflowNode : public FDataflowNode
+{
+	GENERATED_USTRUCT_BODY()
+	DATAFLOW_NODE_DEFINE_INTERNAL(FCollectionVertexSelectionByPlaneDataflowNode, "CollectionVertexSelectByPlane", "GeometryCollection|Selection|Vertex", "Side Points Normal")
+
+public:
+	FCollectionVertexSelectionByPlaneDataflowNode(const UE::Dataflow::FNodeParameters& InParam, FGuid InGuid = FGuid::NewGuid());
+
+private:
+	/** GeometryCollection for the selection */
+	UPROPERTY(meta = (DataflowInput, DataflowOutput, DataflowPassthrough = "Collection", DataflowIntrinsic))
+	FManagedArrayCollection Collection;
+
+	/** plane to select vertices */
+	UPROPERTY(EditAnywhere, Category = "Select", meta = (DataflowInput))
+	FTransform PlaneTransform = FTransform(FRotator(0, 90, 0));
+
+	/** true to use the positive side of the plane */
+	UPROPERTY(EditAnywhere, Category = "Select")
+	bool bPositiveSide = true;
+
+	/** selected vertices */
+	UPROPERTY(meta = (DataflowOutput))
+	FDataflowVertexSelection VertexSelection;
+
+	/** Size of the debug draw plane in cm */
+	UPROPERTY(EditAnywhere, Category = "DebugDraw", meta=(Units=cm, ClampMin = "0.001", UIMin = "0.001"))
+	float PlaneSize = 200;
+
+	/** number of vertices along an edge of the debug draw plane */
+	UPROPERTY(EditAnywhere, Category = "DebugDraw", meta = (ClampMin = "2", UIMin = "2"))
+	int32 PlaneResolution = 20;
+
+
+	virtual void Evaluate(UE::Dataflow::FContext& Context, const FDataflowOutput* Out) const override;
+
+#if WITH_EDITOR
+	virtual bool CanDebugDrawViewMode(const FName& ViewModeName) const override;
+	virtual void DebugDraw(UE::Dataflow::FContext& Context, IDataflowDebugDrawInterface& DataflowRenderingInterface, const FDebugDrawParameters& DebugDrawParameters) const override;
+#endif //WITH_EDITOR
+};
+
 /**
  *
  * Selects bones by a float attribute
diff --git a/Engine/Source/Runtime/Experimental/Chaos/Private/GeometryCollection/Facades/CollectionMeshFacade.cpp b/Engine/Source/Runtime/Experimental/Chaos/Private/GeometryCollection/Facades/CollectionMeshFacade.cpp
--- a/Engine/Source/Runtime/Experimental/Chaos/Private/GeometryCollection/Facades/CollectionMeshFacade.cpp
+++ b/Engine/Source/Runtime/Experimental/Chaos/Private/GeometryCollection/Facades/CollectionMeshFacade.cpp
@@ -106,13 +106,20 @@ namespace GeometryCollection::Facades
 		const TManagedArray<int32>& VertexCounts = VertexCountAttribute.Get();
 
 		TArray<int32> VertexIndices;
-
-		const int32 VertexIndexStart = VertexStarts[TransformToGeometryIndicies[BoneIdx]];
-		for (int32 Offset = 0; Offset < VertexCounts[TransformToGeometryIndicies[BoneIdx]]; ++Offset)
+		if (TransformToGeometryIndicies.IsValidIndex(BoneIdx))
 		{
-			VertexIndices.Add(VertexIndexStart + Offset);
-		}
+			const int32 GeometryIndex = TransformToGeometryIndicies[BoneIdx];
+			if (VertexStarts.IsValidIndex(GeometryIndex) && VertexCounts.IsValidIndex(GeometryIndex))
+			{
+				VertexIndices.Reserve(VertexCounts[GeometryIndex]);
 
+				const int32 VertexIndexStart = VertexStarts[GeometryIndex];
+				for (int32 Offset = 0; Offset < VertexCounts[GeometryIndex]; ++Offset)
+				{
+					VertexIndices.Add(VertexIndexStart + Offset);
+				}
+			}
+		}
 		return VertexIndices;
 	}
 
@@ -123,7 +130,15 @@ namespace GeometryCollection::Facades
 		const TManagedArray<int32>& VertexCounts = VertexCountAttribute.Get();
 		const TManagedArray<FVector3f>& Vertices = VertexAttribute.Get();
 
-		return TArrayView<const FVector3f>(Vertices.GetData() + VertexStarts[TransformToGeometryIndicies[BoneIdx]], VertexCounts[TransformToGeometryIndicies[BoneIdx]]);
+		if (TransformToGeometryIndicies.IsValidIndex(BoneIdx))
+		{
+			const int32 GeometryIndex = TransformToGeometryIndicies[BoneIdx];
+			if (VertexStarts.IsValidIndex(GeometryIndex) && VertexCounts.IsValidIndex(GeometryIndex))
+			{
+				return TArrayView<const FVector3f>(Vertices.GetData() + VertexStarts[GeometryIndex], VertexCounts[GeometryIndex]);
+			}
+		}
+		return {};
 	}
 
 	const TArray<int32> FCollectionMeshFacade::GetFaceIndices(int32 BoneIdx) const
@@ -134,12 +149,20 @@ namespace GeometryCollection::Facades
 
 		TArray<int32> FaceIndicies;
 
-		const int32 FaceIndexStart = FaceStarts[TransformToGeometryIndicies[BoneIdx]];
-		for (int32 Offset = 0; Offset < FaceCounts[TransformToGeometryIndicies[BoneIdx]]; ++Offset)
+		if (TransformToGeometryIndicies.IsValidIndex(BoneIdx))
 		{
-			FaceIndicies.Add(FaceIndexStart + Offset);
-		}
+			const int32 GeometryIndex = TransformToGeometryIndicies[BoneIdx];
+			if (FaceStarts.IsValidIndex(GeometryIndex) && FaceCounts.IsValidIndex(GeometryIndex))
+			{
+				FaceIndicies.Reserve(FaceCounts[GeometryIndex]);
 
+				const int32 FaceIndexStart = FaceStarts[GeometryIndex];
+				for (int32 Offset = 0; Offset < FaceCounts[GeometryIndex]; ++Offset)
+				{
+					FaceIndicies.Add(FaceIndexStart + Offset);
+				}
+			}
+		}
 		return FaceIndicies;
 	}
 
@@ -153,16 +176,23 @@ namespace GeometryCollection::Facades
 		TManagedArray<FVector3f>& TangentUs = TangentUAttribute.Modify();
 		TManagedArray<FVector3f>& TangentVs = TangentVAttribute.Modify();
 
-		const int32 VertexIndexStart = VertexStarts[TransformToGeometryIndices[TransformIdx]];
-
-		for (int32 Offset = 0; Offset < VertexCounts[TransformToGeometryIndices[TransformIdx]]; ++Offset)
+		if (TransformToGeometryIndices.IsValidIndex(TransformIdx))
 		{
-			const int32 VertexIdx = VertexIndexStart + Offset;
-
-			Vertices[VertexIdx] = (FVector3f)InTransform.TransformPosition((FVector)Vertices[VertexIdx]);
-			Normals[VertexIdx] = (FVector3f)InTransform.TransformVector((FVector)Normals[VertexIdx]);
-			TangentUs[VertexIdx] = (FVector3f)InTransform.TransformVector((FVector)TangentUs[VertexIdx]);
-			TangentVs[VertexIdx] = (FVector3f)InTransform.TransformVector((FVector)TangentVs[VertexIdx]);
+			const int32 GeometryIndex = TransformToGeometryIndices[TransformIdx];
+			if (VertexStarts.IsValidIndex(GeometryIndex) && VertexCounts.IsValidIndex(GeometryIndex))
+			{
+				const int32 VertexIndexStart = VertexStarts[TransformToGeometryIndices[TransformIdx]];
+
+				for (int32 Offset = 0; Offset < VertexCounts[TransformToGeometryIndices[TransformIdx]]; ++Offset)
+				{
+					const int32 VertexIdx = VertexIndexStart + Offset;
+					// VertexIdx is valid indices for te arrays below because those attributes all belong to the Vertices group of the collection
+					Vertices[VertexIdx] = (FVector3f)InTransform.TransformPosition((FVector)Vertices[VertexIdx]);
+					Normals[VertexIdx] = (FVector3f)InTransform.TransformVector((FVector)Normals[VertexIdx]);
+					TangentUs[VertexIdx] = (FVector3f)InTransform.TransformVector((FVector)TangentUs[VertexIdx]);
+					TangentVs[VertexIdx] = (FVector3f)InTransform.TransformVector((FVector)TangentVs[VertexIdx]);
+				}
+			}
 		}
 	}
 
diff --git a/Engine/Source/Runtime/Experimental/Chaos/Private/GeometryCollection/Facades/CollectionTransformSelectionFacade.cpp b/Engine/Source/Runtime/Experimental/Chaos/Private/GeometryCollection/Facades/CollectionTransformSelectionFacade.cpp
--- a/Engine/Source/Runtime/Experimental/Chaos/Private/GeometryCollection/Facades/CollectionTransformSelectionFacade.cpp
+++ b/Engine/Source/Runtime/Experimental/Chaos/Private/GeometryCollection/Facades/CollectionTransformSelectionFacade.cpp
@@ -771,6 +771,103 @@ namespace GeometryCollection::Facades
 		return SiblingGroups;
 	}
 
+	void FCollectionTransformSelectionFacade::SelectVerticesInTransformedBox(const FBox& InBox, const FTransform& InBoxTransform, TArray<bool>& OutVertexSelection) const
+	{
+		const FCollectionMeshFacade MeshFacade(ConstCollection);
+		const FCollectionTransformFacade TransformFacade(ConstCollection);
+
+		OutVertexSelection.Reset();
+		if (!MeshFacade.IsValid() || !TransformFacade.IsValid())
+		{
+			return;
+		}
+		const int32 NumVertices = MeshFacade.VertexAttribute.Num();
+		OutVertexSelection.Reserve(NumVertices);
+
+		const TArray<FTransform> CollectionSpaceTransforms = TransformFacade.ComputeCollectionSpaceTransforms();
+
+		for (int32 VertexIdx = 0; VertexIdx < NumVertices; ++VertexIdx)
+		{
+			const int32 TransformIdx = MeshFacade.BoneMapAttribute.IsValidIndex(VertexIdx) ? MeshFacade.BoneMapAttribute[VertexIdx] : INDEX_NONE;
+
+			bool bIsSelected = false;
+			if (CollectionSpaceTransforms.IsValidIndex(TransformIdx))
+			{
+				const FVector VertexInBoneSpace = FVector(MeshFacade.VertexAttribute[VertexIdx]);
+				const FVector VertexInCollectionSpace = CollectionSpaceTransforms[TransformIdx].TransformPosition(VertexInBoneSpace);
+				const FVector VertexInBoxSpace = InBoxTransform.InverseTransformPosition(VertexInCollectionSpace);
+
+				bIsSelected = InBox.IsInside(VertexInBoxSpace);
+			}
+			OutVertexSelection.Add(bIsSelected);
+		}
+	}
+	
+	void FCollectionTransformSelectionFacade::SelectVerticesInSphere(const FSphere& InSphere, TArray<bool>& OutVertexSelection) const
+	{
+		const FCollectionMeshFacade MeshFacade(ConstCollection);
+		const FCollectionTransformFacade TransformFacade(ConstCollection);
+
+		OutVertexSelection.Reset();
+		if (!MeshFacade.IsValid() || !TransformFacade.IsValid())
+		{
+			return;
+		}
+		const int32 NumVertices = MeshFacade.VertexAttribute.Num();
+		OutVertexSelection.Reserve(NumVertices);
+
+		const TArray<FTransform> CollectionSpaceTransforms = TransformFacade.ComputeCollectionSpaceTransforms();
+
+		for (int32 VertexIdx = 0; VertexIdx < NumVertices; ++VertexIdx)
+		{
+			const int32 TransformIdx = MeshFacade.BoneMapAttribute.IsValidIndex(VertexIdx) ? MeshFacade.BoneMapAttribute[VertexIdx] : INDEX_NONE;
+
+			bool bIsSelected = false;
+			if (CollectionSpaceTransforms.IsValidIndex(TransformIdx))
+			{
+				const FVector VertexInBoneSpace = FVector(MeshFacade.VertexAttribute[VertexIdx]);
+				const FVector VertexInCollectionSpace = CollectionSpaceTransforms[TransformIdx].TransformPosition(VertexInBoneSpace);
+
+				bIsSelected = InSphere.IsInside(VertexInCollectionSpace);
+			}
+			OutVertexSelection.Add(bIsSelected);
+		}
+	}
+
+	void FCollectionTransformSelectionFacade::SelectVerticesOnPlaneSide(const FPlane& InPlane, bool bUsePositiveSide, TArray<bool>& OutVertexSelection) const
+	{
+		const FCollectionMeshFacade MeshFacade(ConstCollection);
+		const FCollectionTransformFacade TransformFacade(ConstCollection);
+
+		OutVertexSelection.Reset();
+		if (!MeshFacade.IsValid() || !TransformFacade.IsValid())
+		{
+			return;
+		}
+		const int32 NumVertices = MeshFacade.VertexAttribute.Num();
+		OutVertexSelection.Reserve(NumVertices);
+
+		const TArray<FTransform> CollectionSpaceTransforms = TransformFacade.ComputeCollectionSpaceTransforms();
+
+		for (int32 VertexIdx = 0; VertexIdx < NumVertices; ++VertexIdx)
+		{
+			const int32 TransformIdx = MeshFacade.BoneMapAttribute.IsValidIndex(VertexIdx)? MeshFacade.BoneMapAttribute[VertexIdx]: INDEX_NONE;
+
+			bool bIsSelected = false;
+			if (CollectionSpaceTransforms.IsValidIndex(TransformIdx))
+			{
+				const FVector VertexInBoneSpace = FVector(MeshFacade.VertexAttribute[VertexIdx]);
+				const FVector VertexInCollectionSpace = CollectionSpaceTransforms[TransformIdx].TransformPosition(VertexInBoneSpace);
+
+				const double DistanceFromPlane = InPlane.PlaneDot(VertexInCollectionSpace);
+				bIsSelected = bUsePositiveSide
+					? (DistanceFromPlane >= 0)
+					: (DistanceFromPlane <= 0);
+			}
+			OutVertexSelection.Add(bIsSelected);
+		}
+	}
+
 	TArray<int32> FCollectionTransformSelectionFacade::SelectVerticesInBox(const FBox& InBox, const FTransform& InBoxTransform, bool bAllVerticesInBox) const
 	{
 		TArray<int32> OutSelection;
@@ -784,6 +881,8 @@ namespace GeometryCollection::Facades
 
 			for (int32 TransformIdx = 0; TransformIdx < NumTransforms; ++TransformIdx)
 			{
+				const FTransform CollectionSpaceTransform = TransformFacade.ComputeCollectionSpaceTransform(TransformIdx);
+
 				int32 NumVerticesInsideBox = 0;
 
 				const TArrayView<const FVector3f> VertexPositions = MeshFacade.GetVertexPositions(TransformIdx);
@@ -792,10 +891,9 @@ namespace GeometryCollection::Facades
 					const FVector VertexInBoneSpace(Vertex);
 
 					// Transform from BoneSpace to CollectionSpace
-					const FTransform CollectionSpaceTransform = TransformFacade.ComputeCollectionSpaceTransform(TransformIdx);
 					const FVector VertexInCollectionSpace = CollectionSpaceTransform.TransformPosition(VertexInBoneSpace);
 
-					// Transform with specified transform
+					// Transform in Box space 
 					const FVector VertexInBoxSpace = InBoxTransform.InverseTransformPosition(VertexInCollectionSpace);
 
 					if (!InBox.IsInside(VertexInBoxSpace))
diff --git a/Engine/Source/Runtime/Experimental/Chaos/Public/GeometryCollection/Facades/CollectionTransformSelectionFacade.h b/Engine/Source/Runtime/Experimental/Chaos/Public/GeometryCollection/Facades/CollectionTransformSelectionFacade.h
--- a/Engine/Source/Runtime/Experimental/Chaos/Public/GeometryCollection/Facades/CollectionTransformSelectionFacade.h
+++ b/Engine/Source/Runtime/Experimental/Chaos/Public/GeometryCollection/Facades/CollectionTransformSelectionFacade.h
@@ -132,6 +132,24 @@ namespace GeometryCollection::Facades
 		/** Return a map from parent of bone to the selected bones with that parent */
 		CHAOS_API TMap<int32, TArray<int32>> GetClusteredSelections(const TArray<int32>& InSelection) const;
 
+		/**
+		* Fill an array of bool that contains whether or not a vertex is contained in a transformed box
+		* The transformed box must be in collection space
+		*/
+		CHAOS_API void SelectVerticesInTransformedBox(const FBox& InBox, const FTransform& InBoxTransform, TArray<bool>& OutVertexSelection) const;
+
+		/**
+		* Fill an array of bool that contains whether or not a vertex is contained in a sphere
+		* The sphere must be in collection space
+		*/
+		CHAOS_API void SelectVerticesInSphere(const FSphere& InSphere, TArray<bool>& OutVertexSelection) const;
+
+		/**
+		* Fill an array of bool that contains whether or not a vertex is on a specific side of a plane 
+		* The plane must be defined in collection space
+		*/
+		CHAOS_API void SelectVerticesOnPlaneSide(const FPlane& InPlane, bool bUsePositiveSide, TArray<bool>& OutVertexSelection) const;
+
 		/** 
 		* Returns indices of bones where all the vertices or at least one vertex of the bone is inside InBox 
 		*/
