Dataflow Material nodes : node to set add material parameter nodes and create material instances

#Synced-CL 49892100

#Synced-CL 49892100

#Synced-CL 49892100

diff --git a/Engine/Plugins/Experimental/GeometryCollectionPlugin/Source/GeometryCollectionNodes/Private/Dataflow/GeometryCollectionMaterialInterfaceNodes.cpp b/Engine/Plugins/Experimental/GeometryCollectionPlugin/Source/GeometryCollectionNodes/Private/Dataflow/GeometryCollectionMaterialInterfaceNodes.cpp
--- a/Engine/Plugins/Experimental/GeometryCollectionPlugin/Source/GeometryCollectionNodes/Private/Dataflow/GeometryCollectionMaterialInterfaceNodes.cpp
+++ b/Engine/Plugins/Experimental/GeometryCollectionPlugin/Source/GeometryCollectionNodes/Private/Dataflow/GeometryCollectionMaterialInterfaceNodes.cpp
@@ -9,8 +9,12 @@
 #include "Materials/Material.h"
 #include "Materials/MaterialInstanceDynamic.h"
 #include "Materials/MaterialInstanceConstant.h"
+#include "Misc/PackageName.h"
+#include "Engine/Texture.h"
 #include "Engine/Texture2D.h"
 
+#include "AssetRegistry/AssetRegistryModule.h"
+
 #include UE_INLINE_GENERATED_CPP_BY_NAME(GeometryCollectionMaterialInterfaceNodes)
 
 namespace UE::Dataflow
@@ -24,6 +28,11 @@ namespace UE::Dataflow
 		DATAFLOW_NODE_REGISTER_CREATION_FACTORY(FAssignMaterialInterfaceToCollectionDataflowNode);
 		DATAFLOW_NODE_REGISTER_CREATION_FACTORY(FMaterialInterfaceTextureOverrideDataflowNode);
 
+		DATAFLOW_NODE_REGISTER_CREATION_FACTORY(FMaterialInterfaceSetFloatParameterDataflowNode);
+		DATAFLOW_NODE_REGISTER_CREATION_FACTORY(FMaterialInterfaceSetVectorParameterDataflowNode);
+		DATAFLOW_NODE_REGISTER_CREATION_FACTORY(FMaterialInterfaceSetTexture2DParameterDataflowNode);
+		DATAFLOW_NODE_REGISTER_CREATION_FACTORY(FMaterialInterfaceTerminalDataflowNode);
+
 		// Deprecated nodes
 		DATAFLOW_NODE_REGISTER_CREATION_FACTORY(FGetMaterialInterfaceArraySizeDataflowNode);
 		DATAFLOW_NODE_REGISTER_CREATION_FACTORY(FGetFromMaterialInterfaceArrayDataflowNode);
@@ -414,3 +423,193 @@ void FMaterialInterfaceTextureOverrideDataflowNode::Evaluate(UE::Dataflow::FCont
 	SafeForwardInput(Context, &Material, &Material);
 }
 
+
+//////////////////////////////////////////////////////////////////////////////////////////////////////////////////////////////////////////////////////
+
+FMaterialInterfaceSetFloatParameterDataflowNode::FMaterialInterfaceSetFloatParameterDataflowNode(const UE::Dataflow::FNodeParameters& InParam, FGuid InGuid)
+	: FDataflowNode(InParam, InGuid)
+{
+	RegisterInputConnection(&MaterialParameters);
+	RegisterInputConnection(&ParameterName);
+	RegisterInputConnection(&ParameterValue);
+	RegisterOutputConnection(&MaterialParameters, &MaterialParameters);
+}
+
+void FMaterialInterfaceSetFloatParameterDataflowNode::Evaluate(UE::Dataflow::FContext& Context, const FDataflowOutput* Out) const
+{
+	if (Out->IsA(&MaterialParameters))
+	{
+		FDataflowMaterialParameters OutMaterialParameters = GetValue(Context, &MaterialParameters);
+
+		FDataflowMaterialFloatParameter FloatParameter;
+		FloatParameter.Name = FName(GetValue(Context, &ParameterName));
+		FloatParameter.Value = static_cast<float>(GetValue(Context, &ParameterValue));
+
+		if (!FloatParameter.Name.IsNone())
+		{
+			OutMaterialParameters.FloatParameters.Add(FloatParameter);
+		}
+
+		SetValue(Context, MoveTemp(OutMaterialParameters), &MaterialParameters);
+	}
+}
+
+//////////////////////////////////////////////////////////////////////////////////////////////////////////////////////////////////////////////////////
+
+FMaterialInterfaceSetVectorParameterDataflowNode::FMaterialInterfaceSetVectorParameterDataflowNode(const UE::Dataflow::FNodeParameters& InParam, FGuid InGuid)
+	: FDataflowNode(InParam, InGuid)
+{
+	RegisterInputConnection(&MaterialParameters);
+	RegisterInputConnection(&ParameterName);
+	RegisterInputConnection(&ParameterValue);
+	RegisterOutputConnection(&MaterialParameters, &MaterialParameters);
+}
+
+void FMaterialInterfaceSetVectorParameterDataflowNode::Evaluate(UE::Dataflow::FContext& Context, const FDataflowOutput* Out) const
+{
+	if (Out->IsA(&MaterialParameters))
+	{
+		FDataflowMaterialParameters OutMaterialParameters = GetValue(Context, &MaterialParameters);
+
+		FDataflowMaterialVectorParameter VectorParameter;
+		VectorParameter.Name = FName(GetValue(Context, &ParameterName));
+		VectorParameter.Value = GetValue(Context, &ParameterValue);
+
+		if (!VectorParameter.Name.IsNone())
+		{
+			OutMaterialParameters.VectorParameters.Add(VectorParameter);
+		}
+
+		SetValue(Context, MoveTemp(OutMaterialParameters), &MaterialParameters);
+	}
+}
+
+//////////////////////////////////////////////////////////////////////////////////////////////////////////////////////////////////////////////////////
+
+FMaterialInterfaceSetTexture2DParameterDataflowNode::FMaterialInterfaceSetTexture2DParameterDataflowNode(const UE::Dataflow::FNodeParameters& InParam, FGuid InGuid)
+	: FDataflowNode(InParam, InGuid)
+{
+	RegisterInputConnection(&MaterialParameters);
+	RegisterInputConnection(&ParameterName);
+	RegisterInputConnection(&ParameterValue);
+	RegisterOutputConnection(&MaterialParameters, &MaterialParameters);
+}
+
+void FMaterialInterfaceSetTexture2DParameterDataflowNode::Evaluate(UE::Dataflow::FContext& Context, const FDataflowOutput* Out) const
+{
+	if (Out->IsA(&MaterialParameters))
+	{
+		FDataflowMaterialParameters OutMaterialParameters = GetValue(Context, &MaterialParameters);
+
+		FDataflowMaterialTextureParameter TextureParameter;
+		TextureParameter.Name = FName(GetValue(Context, &ParameterName));
+		TextureParameter.Value = GetValue(Context, &ParameterValue);
+
+		if (!TextureParameter.Name.IsNone())
+		{
+			OutMaterialParameters.TextureParameters.Add(TextureParameter);
+		}
+
+		SetValue(Context, MoveTemp(OutMaterialParameters), &MaterialParameters);
+	}
+}
+
+//////////////////////////////////////////////////////////////////////////////////////////////////////////////////////////////////////////////////////
+
+FMaterialInterfaceTerminalDataflowNode::FMaterialInterfaceTerminalDataflowNode(const UE::Dataflow::FNodeParameters& InParam, FGuid InGuid)
+	: FDataflowTerminalNode(InParam, InGuid)
+{
+	RegisterInputConnection(&Material);
+	RegisterInputConnection(&AssetPath);
+	RegisterInputConnection(&MaterialParameters);
+	RegisterOutputConnection(&Material, &Material);
+}
+
+void FMaterialInterfaceTerminalDataflowNode::Evaluate(UE::Dataflow::FContext& Context) const
+{
+	// Asset is created and stored by the SetAssetValue
+	// Evaluate only return an existing one or nullptr
+	const FString& InAssetPath = GetValue(Context, &AssetPath);
+	TObjectPtr<UMaterialInterface> OutMaterial = ::LoadObject<UMaterialInterface>(nullptr, InAssetPath);
+	SetValue(Context, OutMaterial, &Material);
+}
+
+void FMaterialInterfaceTerminalDataflowNode::SetAssetValue(TObjectPtr<UObject> Asset, UE::Dataflow::FContext& Context) const
+{
+	TObjectPtr<UMaterialInterface> NullMaterialPtr = nullptr;
+
+#if WITH_EDITOR
+	const TObjectPtr<UMaterialInterface> ParentMaterial = GetValue(Context, &Material);
+	if (ParentMaterial == nullptr)
+	{
+		Context.Warning(TEXT("Parent Material is null, material instance cannot be created"), this);
+		SetValue(Context, NullMaterialPtr, &Material);
+		return;
+	}
+
+	const FDataflowMaterialParameters& InMaterialParameters = GetValue(Context, &MaterialParameters);
+	const FString& InAssetPath = GetValue(Context, &AssetPath);
+	if (!FPackageName::IsValidObjectPath(InAssetPath))
+	{
+		Context.Warning(FString::Format(TEXT("Asset Path input is not a valid path : {0}"), { InAssetPath }), this);
+		SetValue(Context, NullMaterialPtr, &Material);
+		return;
+	}
+
+	const FString PackageName = FPackageName::ObjectPathToPackageName(InAssetPath);
+
+	UPackage* Package = Cast<UPackage>(FindPackage(nullptr, *PackageName));
+	if (Package == nullptr)
+	{
+		Package = CreatePackage(*PackageName);
+	}
+	if (Package == nullptr)
+	{
+		Context.Warning(FString::Format(TEXT("Failed to find or create package {0}"), { PackageName }), this);
+		SetValue(Context, NullMaterialPtr, &Material);
+		return;
+	}
+
+	const FString AssetName = FPackageName::GetLongPackageAssetName(PackageName);
+	UMaterialInstanceConstant* MaterialInstance = NewObject<UMaterialInstanceConstant>(Package, UMaterialInstanceConstant::StaticClass(), FName(AssetName), RF_Public | RF_Standalone | RF_Transactional);
+	if (!MaterialInstance)
+	{
+		Context.Warning(FString::Format(TEXT("Material instance {0} with cannot be created"), { InAssetPath }), this);
+		SetValue(Context, NullMaterialPtr, &Material);
+		return;
+	}
+
+	// Setup the new material instance
+	MaterialInstance->SetParentEditorOnly(ParentMaterial);
+
+	for (const FDataflowMaterialFloatParameter& FloatParameter : InMaterialParameters.FloatParameters)
+	{
+		MaterialInstance->SetScalarParameterValueEditorOnly(*FloatParameter.Name.ToString(), FloatParameter.Value);
+	}
+
+	for (const FDataflowMaterialVectorParameter& VectorParameter : InMaterialParameters.VectorParameters)
+	{
+		MaterialInstance->SetVectorParameterValueEditorOnly(*VectorParameter.Name.ToString(), FLinearColor(VectorParameter.Value));
+	}
+
+	for (const FDataflowMaterialTextureParameter& TextureParameter : InMaterialParameters.TextureParameters)
+	{
+		MaterialInstance->SetTextureParameterValueEditorOnly(*TextureParameter.Name.ToString(), TextureParameter.Value.Get());
+	}
+
+	// make sure the asset is set properly
+	MaterialInstance->MarkPackageDirty();
+	FAssetRegistryModule::AssetCreated(MaterialInstance);
+
+	// Finalize the new material asset creation
+	MaterialInstance->PreEditChange(nullptr);
+	MaterialInstance->PostEditChange();
+
+	TObjectPtr<UMaterialInterface> OutMaterialInstance = MaterialInstance;
+	SetValue(Context, OutMaterialInstance, &Material);
+
+#else
+	Context.Error(TEXT("Creation of material instance only supported in Editor"), this);
+	SetValue(Context, NullMaterialPtr, &Material);
+#endif
+}
diff --git a/Engine/Plugins/Experimental/GeometryCollectionPlugin/Source/GeometryCollectionNodes/Public/Dataflow/GeometryCollectionMaterialInterfaceNodes.h b/Engine/Plugins/Experimental/GeometryCollectionPlugin/Source/GeometryCollectionNodes/Public/Dataflow/GeometryCollectionMaterialInterfaceNodes.h
--- a/Engine/Plugins/Experimental/GeometryCollectionPlugin/Source/GeometryCollectionNodes/Public/Dataflow/GeometryCollectionMaterialInterfaceNodes.h
+++ b/Engine/Plugins/Experimental/GeometryCollectionPlugin/Source/GeometryCollectionNodes/Public/Dataflow/GeometryCollectionMaterialInterfaceNodes.h
@@ -4,12 +4,14 @@
 
 #include "CoreMinimal.h"
 #include "Dataflow/DataflowNode.h"
+#include "Dataflow/DataflowTerminalNode.h"
 #include "GeometryCollection/ManagedArrayCollection.h"
 
 #include "GeometryCollectionMaterialInterfaceNodes.generated.h"
 
 class UMaterialInterface;
 class UTexture2D;
+class UTexture;
 
 /**
 * Make a array from a user defined list of material objects
@@ -247,7 +249,186 @@ private:
 	//~ End FDataflowNode interface
 };
 
+//////////////////////////////////////////////////////////////////////////////////////////////////////////////////
 
+USTRUCT()
+struct FDataflowMaterialFloatParameter
+{
+	GENERATED_BODY()
+
+public:
+	UPROPERTY()
+	FName Name;
+
+	UPROPERTY()
+	float Value = 0;
+};
+
+USTRUCT()
+struct FDataflowMaterialVectorParameter
+{
+	GENERATED_BODY()
+
+public:
+	UPROPERTY()
+	FName Name;
+
+	UPROPERTY()
+	FVector4 Value = FVector4(0,0,0,0);
+};
+
+USTRUCT()
+struct FDataflowMaterialTextureParameter
+{
+	GENERATED_BODY()
+
+public:
+	UPROPERTY()
+	FName Name;
+
+	UPROPERTY()
+	TObjectPtr<UTexture> Value = nullptr;
+};
+
+USTRUCT()
+struct FDataflowMaterialParameters
+{
+	GENERATED_BODY()
+
+public:
+	UPROPERTY()
+	TArray<FDataflowMaterialFloatParameter> FloatParameters;
+
+	UPROPERTY()
+	TArray<FDataflowMaterialVectorParameter> VectorParameters;
+
+	UPROPERTY()
+	TArray<FDataflowMaterialTextureParameter> TextureParameters;
+};
+
+//////////////////////////////////////////////////////////////////////////////////////////////////////////////////////
+
+/**
+ * Add a material float parameter to a parameters object
+ * Material parameters are used to create material instances
+ */
+USTRUCT(Meta = (MeshResizing, Experimental))
+struct FMaterialInterfaceSetFloatParameterDataflowNode : public FDataflowNode
+{
+	GENERATED_USTRUCT_BODY()
+	DATAFLOW_NODE_DEFINE_INTERNAL(FMaterialInterfaceSetFloatParameterDataflowNode, "SetMaterialFloatParameter", "Materials", "Instance Value")
+
+public:
+	FMaterialInterfaceSetFloatParameterDataflowNode(const UE::Dataflow::FNodeParameters& InParam, FGuid InGuid = FGuid::NewGuid());
+
+private:
+
+	UPROPERTY(meta = (DataflowInput, DataflowOutput, DataflowPassthrough = "MaterialParameters"))
+	FDataflowMaterialParameters MaterialParameters;
+
+	UPROPERTY(EditAnywhere, Category = "Parameter", meta = (DataflowInput))
+	FDataflowStringTypes ParameterName;
+
+	UPROPERTY(EditAnywhere, Category = "Parameter", meta = (DataflowInput))
+	FDataflowNumericTypes ParameterValue;
+
+	//~ Begin FDataflowNode interface
+	virtual void Evaluate(UE::Dataflow::FContext& Context, const FDataflowOutput* Out) const override;
+	//~ End FDataflowNode interface
+};
+
+//////////////////////////////////////////////////////////////////////////////////////////////////////////////////////
+
+/**
+ * Add a material vector4 parameter to a parameters object
+ * Material parameters are used to create material instances
+ */
+USTRUCT(Meta = (MeshResizing, Experimental))
+struct FMaterialInterfaceSetVectorParameterDataflowNode : public FDataflowNode
+{
+	GENERATED_USTRUCT_BODY()
+	DATAFLOW_NODE_DEFINE_INTERNAL(FMaterialInterfaceSetVectorParameterDataflowNode, "SetMaterialVectorParameter", "Materials", "Instance Value")
+
+public:
+	FMaterialInterfaceSetVectorParameterDataflowNode(const UE::Dataflow::FNodeParameters& InParam, FGuid InGuid = FGuid::NewGuid());
+
+private:
+
+	UPROPERTY(meta = (DataflowInput, DataflowOutput, DataflowPassthrough = "MaterialParameters"))
+	FDataflowMaterialParameters MaterialParameters;
+
+	UPROPERTY(EditAnywhere, Category = "Parameter", meta = (DataflowInput))
+	FDataflowStringTypes ParameterName;
+
+	UPROPERTY(EditAnywhere, Category = "Parameter", meta = (DataflowInput))
+	FDataflowVectorTypes ParameterValue;
+
+	//~ Begin FDataflowNode interface
+	virtual void Evaluate(UE::Dataflow::FContext& Context, const FDataflowOutput* Out) const override;
+	//~ End FDataflowNode interface
+};
+
+//////////////////////////////////////////////////////////////////////////////////////////////////////////////////////
+
+/**
+ * Add a material texture parameter to a parameters object
+ * Material parameters are used to create material instances
+ */
+USTRUCT(Meta = (MeshResizing, Experimental))
+struct FMaterialInterfaceSetTexture2DParameterDataflowNode : public FDataflowNode
+{
+	GENERATED_USTRUCT_BODY()
+	DATAFLOW_NODE_DEFINE_INTERNAL(FMaterialInterfaceSetTexture2DParameterDataflowNode, "SetMaterialTexture2DParameter", "Materials", "Instance Value")
+
+public:
+	FMaterialInterfaceSetTexture2DParameterDataflowNode(const UE::Dataflow::FNodeParameters& InParam, FGuid InGuid = FGuid::NewGuid());
+
+private:
+
+	UPROPERTY(meta = (DataflowInput, DataflowOutput, DataflowPassthrough = "MaterialParameters"))
+	FDataflowMaterialParameters MaterialParameters;
+
+	UPROPERTY(EditAnywhere, Category = "Parameter", meta = (DataflowInput))
+	FDataflowStringTypes ParameterName;
+
+	UPROPERTY(EditAnywhere, Category = "Parameter", meta = (DataflowInput))
+	TObjectPtr<UTexture2D> ParameterValue;
+
+	//~ Begin FDataflowNode interface
+	virtual void Evaluate(UE::Dataflow::FContext& Context, const FDataflowOutput* Out) const override;
+	//~ End FDataflowNode interface
+};
+
+//////////////////////////////////////////////////////////////////////////////////////////////////////////////////////
+
+/**
+ * Make a material instance from another material 
+ */
+USTRUCT(Meta = (MeshResizing, Experimental))
+struct FMaterialInterfaceTerminalDataflowNode : public FDataflowTerminalNode
+{
+	GENERATED_USTRUCT_BODY()
+	DATAFLOW_NODE_DEFINE_INTERNAL(FMaterialInterfaceTerminalDataflowNode, "MaterialTerminalNode", "Terminal", "Instance Value Parameters")
+
+public:
+	FMaterialInterfaceTerminalDataflowNode(const UE::Dataflow::FNodeParameters& InParam, FGuid InGuid = FGuid::NewGuid());
+
+private:
+
+	UPROPERTY(meta = (DataflowInput, DataflowOutput, DataflowPassthrough = "Material"))
+	TObjectPtr<UMaterialInterface> Material;
+
+	UPROPERTY(EditAnywhere, Category = "Asset", meta = (DataflowInput))
+	FString AssetPath;
+
+	UPROPERTY(EditAnywhere, Category = "Parameters", meta = (DataflowInput))
+	FDataflowMaterialParameters MaterialParameters;
+
+	//~ Begin FDataflowNode interface
+	virtual void Evaluate(UE::Dataflow::FContext& Context) const override;
+	virtual void SetAssetValue(TObjectPtr<UObject> Asset, UE::Dataflow::FContext& Context) const override;
+	//~ End FDataflowNode interface
+};
 
 namespace UE::Dataflow
 {
