Add initial medial axis skeletonization dataflow node

#Synced-CL 50036218

diff --git a/Engine/Plugins/Experimental/GeometryCollectionPlugin/Source/GeometryCollectionNodes/GeometryCollectionNodes.Build.cs b/Engine/Plugins/Experimental/GeometryCollectionPlugin/Source/GeometryCollectionNodes/GeometryCollectionNodes.Build.cs
--- a/Engine/Plugins/Experimental/GeometryCollectionPlugin/Source/GeometryCollectionNodes/GeometryCollectionNodes.Build.cs
+++ b/Engine/Plugins/Experimental/GeometryCollectionPlugin/Source/GeometryCollectionNodes/GeometryCollectionNodes.Build.cs
@@ -35,7 +35,8 @@ namespace UnrealBuildTool.Rules
 					"DynamicMesh",
 					"MeshConversion",
 					"MeshConversionEngineTypes",
-					"FractureEngine"
+					"FractureEngine",
+					"GeometryAlgorithms"
 				}
 			);
 
diff --git a/Engine/Plugins/Experimental/GeometryCollectionPlugin/Source/GeometryCollectionNodes/Private/Dataflow/GeometryCollectionSamplingNodes.cpp b/Engine/Plugins/Experimental/GeometryCollectionPlugin/Source/GeometryCollectionNodes/Private/Dataflow/GeometryCollectionSamplingNodes.cpp
--- a/Engine/Plugins/Experimental/GeometryCollectionPlugin/Source/GeometryCollectionNodes/Private/Dataflow/GeometryCollectionSamplingNodes.cpp
+++ b/Engine/Plugins/Experimental/GeometryCollectionPlugin/Source/GeometryCollectionNodes/Private/Dataflow/GeometryCollectionSamplingNodes.cpp
@@ -5,6 +5,12 @@
 #include "UDynamicMesh.h"
 #include "Spatial/FastWinding.h"
 
+#include "Dataflow/DataflowDebugDrawInterface.h"
+#include "Dataflow/DataflowDebugDraw.h"
+#if WITH_EDITOR
+#include "Dataflow/DataflowRenderingViewMode.h"
+#endif
+
 #include UE_INLINE_GENERATED_CPP_BY_NAME(GeometryCollectionSamplingNodes)
 
 namespace UE::Dataflow
@@ -16,7 +22,34 @@ namespace UE::Dataflow
 		DATAFLOW_NODE_REGISTER_CREATION_FACTORY(FNonUniformPointSamplingDataflowNode);
 		DATAFLOW_NODE_REGISTER_CREATION_FACTORY(FVertexWeightedPointSamplingDataflowNode);
 		DATAFLOW_NODE_REGISTER_CREATION_FACTORY(FFilterPointSetWithMeshDataflowNode);
+		DATAFLOW_NODE_REGISTER_CREATION_FACTORY(FMeshMedialSkeletonSamplingDataflowNode);
 	}
+
+	static void MedialSkeletonDebugDraw(IDataflowDebugDrawInterface& DataflowRenderingInterface, const UE::Geometry::MedialAxis::FMedialSkeleton& ToDrawSkel)
+	{
+		using namespace UE::Geometry;
+
+		DataflowRenderingInterface.SetWorldPriority();
+		DataflowRenderingInterface.SetColor(FLinearColor::Green);
+		for (int32 ClusterIdx = 0; ClusterIdx < ToDrawSkel.ClusterNeighbors.Num(); ++ClusterIdx)
+		{
+			if (!ensure(ToDrawSkel.Spheres.IsValidIndex(ClusterIdx)))
+			{
+				continue;
+			}
+			FVector3d A = ToDrawSkel.Spheres[ClusterIdx].Center;
+			for (int32 NbrIdx : ToDrawSkel.ClusterNeighbors[ClusterIdx])
+			{
+				if (!ensure(ToDrawSkel.Spheres.IsValidIndex(NbrIdx)))
+				{
+					continue;
+				}
+				FVector3d B = ToDrawSkel.Spheres[NbrIdx].Center;
+				DataflowRenderingInterface.DrawLine(A, B);
+			}
+		}
+	}
+
 }
 
 FFilterPointSetWithMeshDataflowNode::FFilterPointSetWithMeshDataflowNode(const UE::Dataflow::FNodeParameters& InParam, FGuid InGuid)
@@ -346,3 +379,76 @@ void FVertexWeightedPointSamplingDataflowNode::Evaluate(UE::Dataflow::FContext&
 }
 
 
+FMeshMedialSkeletonSamplingDataflowNode::FMeshMedialSkeletonSamplingDataflowNode(const UE::Dataflow::FNodeParameters& InParam, FGuid InGuid)
+	: FDataflowNode(InParam, InGuid)
+{
+	RegisterInputConnection(&TargetMesh);
+	RegisterInputConnection(&MaxSpheres).SetCanHidePin(true);
+	RegisterInputConnection(&MinErrorDelta).SetCanHidePin(true).SetPinIsHidden(true);
+	RegisterInputConnection(&MinClusterErrorToSplit).SetCanHidePin(true).SetPinIsHidden(true);
+	RegisterInputConnection(&MinClusterSizeToKeep).SetCanHidePin(true).SetPinIsHidden(true);
+	RegisterInputConnection(&bSplitClustersIfEdgesIntersectSurface).SetCanHidePin(true).SetPinIsHidden(true);
+	RegisterInputConnection(&ErrorScaleNearEdgeSurfaceIntersection).SetCanHidePin(true).SetPinIsHidden(true);
+	RegisterInputConnection(&PosErrorWt).SetCanHidePin(true).SetPinIsHidden(true);
+
+	RegisterOutputConnection(&MedialSkeleton);
+	RegisterOutputConnection(&MedialSpheres);
+}
+
+void FMeshMedialSkeletonSamplingDataflowNode::Evaluate(UE::Dataflow::FContext& Context, const FDataflowOutput* Out) const
+{
+	if (Out->IsA(&MedialSkeleton) ||
+		Out->IsA(&MedialSpheres))
+	{
+		FDataflowMedialSkeleton OutMedialSkeleton;
+		UE::Geometry::MedialAxis::FMedialSkeleton& Skeleton = OutMedialSkeleton.Skeleton;
+
+		if (TObjectPtr<UDynamicMesh> InTargetMesh = GetValue(Context, &TargetMesh))
+		{
+			const UE::Geometry::FDynamicMesh3& InDynTargetMesh = InTargetMesh->GetMeshRef();
+
+			if (InDynTargetMesh.TriangleCount() > 0)
+			{
+				UE::Geometry::TMeshAABBTree3<FDynamicMesh3> Spatial(&InDynTargetMesh);
+				UE::Geometry::TFastWindingTree<FDynamicMesh3> FWTree(&Spatial);
+				UE::Geometry::MedialAxis::FSkeletonViaSampling SkelSample;
+				SkelSample.MaxSpheres = GetValue(Context, &MaxSpheres);
+				SkelSample.MinErrorDelta = GetValue(Context, &MinErrorDelta);
+				SkelSample.MinClusterErrorToSplit = GetValue(Context, &MinClusterErrorToSplit);
+				SkelSample.MinClusterSizeToKeep = GetValue(Context, &MinClusterSizeToKeep);
+				SkelSample.bSplitClustersIfEdgesIntersectSurface = GetValue(Context, &bSplitClustersIfEdgesIntersectSurface);
+				SkelSample.ErrorScaleNearEdgeSurfaceIntersection = GetValue(Context, &ErrorScaleNearEdgeSurfaceIntersection);
+				SkelSample.PosErrorWt = GetValue(Context, &PosErrorWt);
+
+				Skeleton = SkelSample.ComputeSkeleton(Spatial, 1e-4, &FWTree);
+			}
+		}
+
+		TArray<FSphere> Spheres = Skeleton.Spheres;
+		SetValue(Context, MoveTemp(OutMedialSkeleton), &MedialSkeleton);
+		SetValue(Context, MoveTemp(Spheres), &MedialSpheres);
+	}
+}
+
+#if WITH_EDITOR
+bool FMeshMedialSkeletonSamplingDataflowNode::CanDebugDrawViewMode(const FName& ViewModeName) const
+{
+	return ViewModeName == UE::Dataflow::FDataflowConstruction3DViewMode::Name;
+}
+
+void FMeshMedialSkeletonSamplingDataflowNode::DebugDraw(UE::Dataflow::FContext& Context, IDataflowDebugDrawInterface& DataflowRenderingInterface, const FDebugDrawParameters& DebugDrawParameters) const
+{
+	using namespace UE::Geometry;
+
+	if ((DebugDrawParameters.bNodeIsSelected || DebugDrawParameters.bNodeIsPinned))
+	{
+		if (const FDataflowOutput* SphereCoveringOutput = FindOutput(&MedialSkeleton))
+		{
+			const UE::Geometry::MedialAxis::FMedialSkeleton& ToDrawSkeleton = SphereCoveringOutput->GetValue(Context, MedialSkeleton).Skeleton;
+
+			UE::Dataflow::MedialSkeletonDebugDraw(DataflowRenderingInterface, ToDrawSkeleton);
+		}
+	}
+}
+#endif
+
diff --git a/Engine/Plugins/Experimental/GeometryCollectionPlugin/Source/GeometryCollectionNodes/Public/Dataflow/GeometryCollectionSamplingNodes.h b/Engine/Plugins/Experimental/GeometryCollectionPlugin/Source/GeometryCollectionNodes/Public/Dataflow/GeometryCollectionSamplingNodes.h
--- a/Engine/Plugins/Experimental/GeometryCollectionPlugin/Source/GeometryCollectionNodes/Public/Dataflow/GeometryCollectionSamplingNodes.h
+++ b/Engine/Plugins/Experimental/GeometryCollectionPlugin/Source/GeometryCollectionNodes/Public/Dataflow/GeometryCollectionSamplingNodes.h
@@ -3,8 +3,11 @@
 #pragma once 
 
 #include "Dataflow/DataflowEngine.h"
+#include "Dataflow/DataflowDebugDraw.h"
 #include "FractureEngineSampling.h"
 
+#include "Skeletonization/MeshMedialAxisSampling.h"
+
 #include "GeometryCollectionSamplingNodes.generated.h"
 
 #if UE_ENABLE_INCLUDE_ORDER_DEPRECATED_IN_5_5
@@ -328,6 +331,81 @@ public:
 };
 
 
+//~ Simple wrapper class to work with medial skeletons in dataflow
+// A set of spheres and associated connectivity corresponding to a sampled medial skeleton of a target mesh
+USTRUCT()
+struct FDataflowMedialSkeleton
+{
+	GENERATED_USTRUCT_BODY()
+
+public:
+	UE::Geometry::MedialAxis::FMedialSkeleton Skeleton;
+};
+
+
+/**
+ *
+ * Sample the medial skeleton of a target mesh
+ *
+ */
+USTRUCT(meta = (DataflowGeometryCollection))
+struct FMeshMedialSkeletonSamplingDataflowNode : public FDataflowNode
+{
+	GENERATED_USTRUCT_BODY()
+	DATAFLOW_NODE_DEFINE_INTERNAL(FMeshMedialSkeletonSamplingDataflowNode, "MeshMedialSkeletonSampling", "MedialSkeletonSampling", "")
+
+private:
+	/** Mesh to sample skeleton on */
+	UPROPERTY(meta = (DataflowInput, DataflowIntrinsic))
+	TObjectPtr<UDynamicMesh> TargetMesh;
+
+	UPROPERTY(meta = (DataflowOutput))
+	FDataflowMedialSkeleton MedialSkeleton;
+
+	UPROPERTY(meta = (DataflowOutput))
+	TArray<FSphere> MedialSpheres;
+
+	// Stop at this sphere count
+	UPROPERTY(EditAnywhere, Category = Options, meta = (DataflowInput, ClampMin = 1))
+	int32 MaxSpheres = 1000;
+
+	// Do not split a cluster if its max error is below this threshold
+	UPROPERTY(EditAnywhere, Category = Options, meta = (DataflowInput, UIMin = .1, ClampMin = 0.))
+	float MinClusterErrorToSplit = 10.f;
+
+	// Stop when the average squared error at vertices decreases by less than this w/ an added sphere addition
+	UPROPERTY(EditAnywhere, Category = Options, meta = (DataflowInput, UIMin = .1, ClampMin = 0.))
+	float MinErrorDelta = 1.f;
+
+	// Clusters with fewer vertices than this will be discarded
+	UPROPERTY(EditAnywhere, Category = Options, meta = (DataflowInput, ClampMin = 1, UIMax = 10))
+	int32 MinClusterSizeToKeep = 4;
+
+	// Whether to test medial skeleton edges for intersection w/ the input surface, and try to refine where these intersections are found
+	UPROPERTY(EditAnywhere, Category = Options, meta = (DataflowInput))
+	bool bSplitClustersIfEdgesIntersectSurface = true;
+
+	// Factor to scale vertex error near skeleton-edge/input-surface intersections
+	UPROPERTY(EditAnywhere, Category = Options, AdvancedDisplay, meta = (DataflowInput, ClampMin = 1, UIMax = 100))
+	float ErrorScaleNearEdgeSurfaceIntersection = 10.f;
+
+	// Weight for the position error term. Relative to plane error term, which has weight 1.
+	UPROPERTY(EditAnywhere, Category = Options, AdvancedDisplay, meta = (DataflowInput, UIMin = .1, ClampMin = .001, UIMax = 2))
+	float PosErrorWt = .2f;
+
+	virtual void Evaluate(UE::Dataflow::FContext& Context, const FDataflowOutput* Out) const override;
+
+#if WITH_EDITOR
+	virtual bool CanDebugDrawViewMode(const FName& ViewModeName) const override;
+	virtual void DebugDraw(UE::Dataflow::FContext& Context, IDataflowDebugDrawInterface& DataflowRenderingInterface, const FDebugDrawParameters& DebugDrawParameters) const override;
+#endif
+
+public:
+	FMeshMedialSkeletonSamplingDataflowNode(const UE::Dataflow::FNodeParameters& InParam, FGuid InGuid = FGuid::NewGuid());
+};
+
+
+
 namespace UE::Dataflow
 {
 	void GeometryCollectionSamplingNodes();
