don't clear the gutter alpha channel in the fracture autouv tool's bake; copy pixels to the gutter unchanged to reduce chance of visible seams
for equivalent dataflow nodes, keep original behavior (clearing gutter alpha) as default, but add option to not clear it

#Synced-CL 50262759

diff --git a/Engine/Plugins/Experimental/ChaosEditor/Source/FractureEditor/Private/FractureToolUV.cpp b/Engine/Plugins/Experimental/ChaosEditor/Source/FractureEditor/Private/FractureToolUV.cpp
--- a/Engine/Plugins/Experimental/ChaosEditor/Source/FractureEditor/Private/FractureToolUV.cpp
+++ b/Engine/Plugins/Experimental/ChaosEditor/Source/FractureEditor/Private/FractureToolUV.cpp
@@ -606,7 +606,7 @@ void UFractureToolAutoUV::BakeTextureForComponent(UGeometryCollectionComponent*
 	BakeOp->AttribSettings.Curvature_VoxelRes = AutoUVSettings->VoxelResolution;
 	BakeOp->AttribSettings.Curvature_ThicknessFactor = AutoUVSettings->ThicknessFactor;
 	BakeOp->AttribSettings.Curvature_MaxValue = AutoUVSettings->MaxCurvature;
-	BakeOp->AttribSettings.ClearGutterChannel = 3; // default clear the gutters for the alpha channel, so it shows more clearly the island boundaries
+	BakeOp->AttribSettings.ClearGutterChannel = AutoUVSettings->bClearGutterAlpha ? 3 : INDEX_NONE;
 
 	FImageType ImageResult;
 	int Result = RunCancellableGeometryCollectionOpGeneric<FBakeTextureOp, FImageType>(
diff --git a/Engine/Plugins/Experimental/ChaosEditor/Source/FractureEditor/Private/FractureToolUV.h b/Engine/Plugins/Experimental/ChaosEditor/Source/FractureEditor/Private/FractureToolUV.h
--- a/Engine/Plugins/Experimental/ChaosEditor/Source/FractureEditor/Private/FractureToolUV.h
+++ b/Engine/Plugins/Experimental/ChaosEditor/Source/FractureEditor/Private/FractureToolUV.h
@@ -150,6 +150,10 @@ public:
 	UPROPERTY(EditAnywhere, Category = MapSettings, meta = (ClampMin = "1", ClampMax = "10", UIMax = "6"))
 	int32 GutterSize = 2;
 
+	/** Whether to set gutter pixels to transparent, or to copy them unchanged from the closest interior pixel value. */
+	UPROPERTY(EditAnywhere, Category = MapSettings)
+	bool bClearGutterAlpha = false;
+
 	/** The resulting automatically-generated texture map */
 	UPROPERTY(VisibleAnywhere, Category = MapSettings)
 	TObjectPtr<UTexture2D> Result;
diff --git a/Engine/Plugins/Experimental/GeometryCollectionPlugin/Source/GeometryCollectionNodes/Private/Dataflow/GeometryCollectionTextureNodes.cpp b/Engine/Plugins/Experimental/GeometryCollectionPlugin/Source/GeometryCollectionNodes/Private/Dataflow/GeometryCollectionTextureNodes.cpp
--- a/Engine/Plugins/Experimental/GeometryCollectionPlugin/Source/GeometryCollectionNodes/Private/Dataflow/GeometryCollectionTextureNodes.cpp
+++ b/Engine/Plugins/Experimental/GeometryCollectionPlugin/Source/GeometryCollectionNodes/Private/Dataflow/GeometryCollectionTextureNodes.cpp
@@ -34,6 +34,7 @@ FBakeTextureFromCollectionDataflowNode::FBakeTextureFromCollectionDataflowNode(c
 	RegisterInputConnection(&SmoothingIterations).SetCanHidePin(true).SetPinIsHidden(true);
 	RegisterInputConnection(&ThicknessFactor).SetCanHidePin(true).SetPinIsHidden(true);
 	RegisterInputConnection(&MaxCurvature).SetCanHidePin(true).SetPinIsHidden(true);
+	RegisterInputConnection(&bClearGutterAlpha).SetCanHidePin(true).SetPinIsHidden(true);
 	
 	RegisterOutputConnection(&Collection, &Collection);
 	RegisterOutputConnection(&Image);
@@ -88,7 +89,7 @@ void FBakeTextureFromCollectionDataflowNode::Evaluate(UE::Dataflow::FContext& Co
 		BakeAttributeSettings.Curvature_VoxelRes = GetValue(Context, &VoxelResolution);
 		BakeAttributeSettings.Curvature_ThicknessFactor = GetValue(Context, &ThicknessFactor);
 		BakeAttributeSettings.Curvature_MaxValue = GetValue(Context, &MaxCurvature);
-		BakeAttributeSettings.ClearGutterChannel = 3; // default clear the gutters for the alpha channel, so it shows more clearly the island boundaries
+		BakeAttributeSettings.ClearGutterChannel = GetValue(Context, &bClearGutterAlpha) ? 3 : INDEX_NONE;
 
 		UE::Geometry::TImageBuilder<FVector4f> ImageBuilder;
 		ImageBuilder.SetDimensions({ InResolution, InResolution });
diff --git a/Engine/Plugins/Experimental/GeometryCollectionPlugin/Source/GeometryCollectionNodes/Public/Dataflow/GeometryCollectionTextureNodes.h b/Engine/Plugins/Experimental/GeometryCollectionPlugin/Source/GeometryCollectionNodes/Public/Dataflow/GeometryCollectionTextureNodes.h
--- a/Engine/Plugins/Experimental/GeometryCollectionPlugin/Source/GeometryCollectionNodes/Public/Dataflow/GeometryCollectionTextureNodes.h
+++ b/Engine/Plugins/Experimental/GeometryCollectionPlugin/Source/GeometryCollectionNodes/Public/Dataflow/GeometryCollectionTextureNodes.h
@@ -77,6 +77,10 @@ private:
 	UPROPERTY(EditAnywhere, Category = ChannelOptions, meta = (DataflowInput, ClampMin = 1, UIMin = 1))
 	int32 GutterSize = 1;
 
+	// If true, will clear the alpha channel for gutter pixels, so they are transparent and the UV boundaries are clearer. If false, gutter pixel values will be copied as-is from the nearest interior pixel.
+	UPROPERTY(EditAnywhere, Category = ChannelOptions, meta = (DataflowInput))
+	bool bClearGutterAlpha = true;
+
 	/** Index of the added UV channel */
 	UPROPERTY(EditAnywhere, Category = ChannelOptions, meta = (DataflowInput, DataflowOutput, DataflowPassthrough = "UVChannel", DataflowUVChannel, ClampMin = 0, ClampMax = 7, UIMin = 0, UIMax = 7))
 	int32 UVChannel = 0;
