Dataflow : fix material and animBank terminal nodes display names

#Synced-CL 50850263

diff --git a/Engine/Plugins/Experimental/ChaosRigidAsset/Source/ChaosRigidAssetNodes/Private/AnimationTerminalNodes.h b/Engine/Plugins/Experimental/ChaosRigidAsset/Source/ChaosRigidAssetNodes/Private/AnimationTerminalNodes.h
--- a/Engine/Plugins/Experimental/ChaosRigidAsset/Source/ChaosRigidAssetNodes/Private/AnimationTerminalNodes.h
+++ b/Engine/Plugins/Experimental/ChaosRigidAsset/Source/ChaosRigidAssetNodes/Private/AnimationTerminalNodes.h
@@ -45,7 +45,7 @@ USTRUCT(meta = (DataflowTerminal))
 struct FDataflowAnimBankAssetTerminalNode : public FDataflowTerminalNode
 {
 	GENERATED_BODY()
-	DATAFLOW_NODE_DEFINE_INTERNAL(FDataflowAnimBankAssetTerminalNode, "AnimBankTerminalNode", "Terminal", "Animation Keys")
+	DATAFLOW_NODE_DEFINE_INTERNAL(FDataflowAnimBankAssetTerminalNode, "AnimBankTerminal", "Terminal", "Animation Keys")
 
 public:
 	FDataflowAnimBankAssetTerminalNode(const UE::Dataflow::FNodeParameters& InParam, FGuid InGuid = FGuid::NewGuid());
diff --git a/Engine/Plugins/Experimental/GeometryCollectionPlugin/Source/GeometryCollectionNodes/Public/Dataflow/GeometryCollectionMaterialInterfaceNodes.h b/Engine/Plugins/Experimental/GeometryCollectionPlugin/Source/GeometryCollectionNodes/Public/Dataflow/GeometryCollectionMaterialInterfaceNodes.h
--- a/Engine/Plugins/Experimental/GeometryCollectionPlugin/Source/GeometryCollectionNodes/Public/Dataflow/GeometryCollectionMaterialInterfaceNodes.h
+++ b/Engine/Plugins/Experimental/GeometryCollectionPlugin/Source/GeometryCollectionNodes/Public/Dataflow/GeometryCollectionMaterialInterfaceNodes.h
@@ -408,7 +408,7 @@ USTRUCT(Meta = (MeshResizing, Experimental))
 struct FMaterialInterfaceTerminalDataflowNode : public FDataflowTerminalNode
 {
 	GENERATED_USTRUCT_BODY()
-	DATAFLOW_NODE_DEFINE_INTERNAL(FMaterialInterfaceTerminalDataflowNode, "MaterialTerminalNode", "Terminal", "Instance Value Parameters")
+	DATAFLOW_NODE_DEFINE_INTERNAL(FMaterialInterfaceTerminalDataflowNode, "MaterialTerminal", "Terminal", "Instance Value Parameters")
 
 public:
 	FMaterialInterfaceTerminalDataflowNode(const UE::Dataflow::FNodeParameters& InParam, FGuid InGuid = FGuid::NewGuid());
