add dataflow node to convert medial skeleton + mesh to a skinned dynamic mesh

#Synced-CL 50894192

diff --git a/Engine/Plugins/Experimental/GeometryCollectionPlugin/Source/GeometryCollectionNodes/GeometryCollectionNodes.Build.cs b/Engine/Plugins/Experimental/GeometryCollectionPlugin/Source/GeometryCollectionNodes/GeometryCollectionNodes.Build.cs
--- a/Engine/Plugins/Experimental/GeometryCollectionPlugin/Source/GeometryCollectionNodes/GeometryCollectionNodes.Build.cs
+++ b/Engine/Plugins/Experimental/GeometryCollectionPlugin/Source/GeometryCollectionNodes/GeometryCollectionNodes.Build.cs
@@ -36,7 +36,8 @@ namespace UnrealBuildTool.Rules
 					"MeshConversion",
 					"MeshConversionEngineTypes",
 					"FractureEngine",
-					"GeometryAlgorithms"
+					"GeometryAlgorithms",
+					"SkeletalMeshDescription"
 				}
 			);
 
@@ -47,6 +48,7 @@ namespace UnrealBuildTool.Rules
 					{
 						"DataflowEditor",
 						"UnrealEd",
+						"SkeletalMeshUtilitiesCommon", // for skeletal mesh conversion
 					}
 					);
 			}
diff --git a/Engine/Plugins/Experimental/GeometryCollectionPlugin/Source/GeometryCollectionNodes/Private/Dataflow/GeometryCollectionNodesPlugin.cpp b/Engine/Plugins/Experimental/GeometryCollectionPlugin/Source/GeometryCollectionNodes/Private/Dataflow/GeometryCollectionNodesPlugin.cpp
--- a/Engine/Plugins/Experimental/GeometryCollectionPlugin/Source/GeometryCollectionNodes/Private/Dataflow/GeometryCollectionNodesPlugin.cpp
+++ b/Engine/Plugins/Experimental/GeometryCollectionPlugin/Source/GeometryCollectionNodes/Private/Dataflow/GeometryCollectionNodesPlugin.cpp
@@ -39,6 +39,7 @@
 #include "Dataflow/SetVertexColorFromFloatArrayNode.h"
 #include "Dataflow/SetVertexColorFromVertexIndicesNode.h"
 #include "Dataflow/SetVertexColorFromVertexSelectionNode.h"
+#include "Dataflow/MeshMedialSkeletonConversionNodes.h"
 #include "Dataflow/DataflowNodeAndPinTypeColors.h"
 #include "Dataflow/DataflowCategoryRegistry.h"
 #include "GeometryCollection/GeometryCollectionObject.h"
@@ -67,6 +68,7 @@ void IGeometryCollectionNodesPlugin::StartupModule()
 	UE::Dataflow::GeometryCollectionArrayNodes();
 	UE::Dataflow::GeometryCollectionDebugNodes();
 	UE::Dataflow::GeometryCollectionSamplingNodes();
+	UE::Dataflow::RegisterMeshMedialSkeletonConversionNodes();
 	UE::Dataflow::GeometryCollectionRemoveMeshOverlapsNode();
 	UE::Dataflow::RegisterGeometryCollectionUVNodes();
 	UE::Dataflow::RegisterGeometryCollectionTextureNodes();
diff --git a/Engine/Plugins/Experimental/GeometryCollectionPlugin/Source/GeometryCollectionNodes/Private/Dataflow/MeshMedialSkeletonConversionNodes.cpp b/Engine/Plugins/Experimental/GeometryCollectionPlugin/Source/GeometryCollectionNodes/Private/Dataflow/MeshMedialSkeletonConversionNodes.cpp
new file mode 100644
--- /dev/null
+++ b/Engine/Plugins/Experimental/GeometryCollectionPlugin/Source/GeometryCollectionNodes/Private/Dataflow/MeshMedialSkeletonConversionNodes.cpp
@@ -0,0 +1,97 @@
+// Copyright Epic Games, Inc. All Rights Reserved.
+
+#include "Dataflow/MeshMedialSkeletonConversionNodes.h"
+
+#include "DynamicMesh/DynamicBoneAttribute.h"
+#include "DynamicMesh/DynamicMeshAttributeSet.h"
+#include "DynamicMesh/DynamicVertexSkinWeightsAttribute.h"
+#include "UDynamicMesh.h"
+
+#include "SkeletalMeshAttributes.h" // for FSkeletalMeshAttributes::DefaultSkinWeightProfileName
+
+
+#include UE_INLINE_GENERATED_CPP_BY_NAME(MeshMedialSkeletonConversionNodes)
+
+namespace UE::Dataflow
+{
+
+	void RegisterMeshMedialSkeletonConversionNodes()
+	{
+		DATAFLOW_NODE_REGISTER_CREATION_FACTORY(FSkinMeshViaMedialSkeleton);
+	}
+
+}
+
+
+
+namespace MeshMedialSkeletonConversionLocals
+{
+
+
+	void SetSkeletonConversionOptions(UE::Geometry::MedialAxis::FMedialSkeletonToTreeSkeletonOptions& OutOptions,
+		EDataflowMedialSkeletonConversionEdgeWeightMethod EdgeWeightMethod,
+		EDataflowMedialSkeletonConversionMergeDisconnectedMethod MergeDisconnectedMethod)
+	{
+		using namespace UE::Geometry::MedialAxis;
+		switch (EdgeWeightMethod)
+		{
+		case EDataflowMedialSkeletonConversionEdgeWeightMethod::ArrayOrder:
+			OutOptions.EdgeWeightMethod = FMedialSkeletonToTreeSkeletonOptions::EEdgeWeightMethod::ArrayOrder;
+			break;
+		case EDataflowMedialSkeletonConversionEdgeWeightMethod::AvgRadius:
+			OutOptions.EdgeWeightMethod = FMedialSkeletonToTreeSkeletonOptions::EEdgeWeightMethod::AvgRadius;
+			break;
+		case EDataflowMedialSkeletonConversionEdgeWeightMethod::EdgeLength:
+			OutOptions.EdgeWeightMethod = FMedialSkeletonToTreeSkeletonOptions::EEdgeWeightMethod::EdgeLength;
+			break;
+		}
+		switch (MergeDisconnectedMethod)
+		{
+		case EDataflowMedialSkeletonConversionMergeDisconnectedMethod::AddTopLevelRoot:
+			OutOptions.MergeDisconnectedMethod = FMedialSkeletonToTreeSkeletonOptions::EMergeDisconnectedMethod::AddTopLevelRoot;
+			break;
+		case EDataflowMedialSkeletonConversionMergeDisconnectedMethod::ConnectClosestBones:
+			OutOptions.MergeDisconnectedMethod = FMedialSkeletonToTreeSkeletonOptions::EMergeDisconnectedMethod::ConnectClosestBones;
+			break;
+		}
+		
+	}
+}
+
+
+FSkinMeshViaMedialSkeleton::FSkinMeshViaMedialSkeleton(const UE::Dataflow::FNodeParameters& InParam, FGuid InGuid)
+	: FDataflowNode(InParam, InGuid)
+{
+	RegisterInputConnection(&Mesh);
+	RegisterInputConnection(&MedialSkeleton);
+	RegisterInputConnection(&RootClusterIndex).SetCanHidePin(true).SetPinIsHidden(true);
+
+	RegisterOutputConnection(&Mesh, &Mesh);
+}
+
+void FSkinMeshViaMedialSkeleton::Evaluate(UE::Dataflow::FContext& Context, const FDataflowOutput* Out) const
+{
+	using namespace UE::Geometry;
+
+	if (Out->IsA(&Mesh))
+	{
+		TObjectPtr<UDynamicMesh> InMesh = GetValue(Context, &Mesh);
+		const FDataflowMedialSkeleton& InMedialSkeleton = GetValue(Context, &MedialSkeleton);
+		TObjectPtr<UDynamicMesh> OutMesh = NewObject<UDynamicMesh>();
+
+		if (InMesh)
+		{
+			// Create a new mesh copy and add the weights
+			OutMesh->SetMesh(InMesh->GetMeshRef());
+			FDynamicMesh3& DynMesh = OutMesh->GetMeshRef();
+
+			UE::Geometry::MedialAxis::FMedialSkeletonToTreeSkeletonOptions Options;
+			MeshMedialSkeletonConversionLocals::SetSkeletonConversionOptions(Options, EdgeWeightMethod, MergeDisconnectedMethod);
+			const int32 UseRootIndex = GetValue(Context, &RootClusterIndex);
+			Options.ToDynamicMeshSkeleton(InMedialSkeleton.Skeleton, DynMesh, FSkeletalMeshAttributes::DefaultSkinWeightProfileName, Options, UseRootIndex);
+	
+		}
+		SetValue(Context, OutMesh, &Mesh);
+	}
+}
+
diff --git a/Engine/Plugins/Experimental/GeometryCollectionPlugin/Source/GeometryCollectionNodes/Private/Dataflow/MeshMedialSkeletonConversionNodes.h b/Engine/Plugins/Experimental/GeometryCollectionPlugin/Source/GeometryCollectionNodes/Private/Dataflow/MeshMedialSkeletonConversionNodes.h
new file mode 100644
--- /dev/null
+++ b/Engine/Plugins/Experimental/GeometryCollectionPlugin/Source/GeometryCollectionNodes/Private/Dataflow/MeshMedialSkeletonConversionNodes.h
@@ -0,0 +1,74 @@
+// Copyright Epic Games, Inc. All Rights Reserved.
+
+#pragma once 
+
+#include "CoreMinimal.h"
+#include "Dataflow/DataflowDebugDraw.h"
+#include "Dataflow/DataflowEngine.h"
+#include "Dataflow/GeometryCollectionSamplingNodes.h"
+#include "GeometryCollection/ManagedArrayCollection.h"
+
+
+#include "MeshMedialSkeletonConversionNodes.generated.h"
+
+UENUM()
+enum class EDataflowMedialSkeletonConversionEdgeWeightMethod : uint8
+{
+	// Favor adding shorter edges first. (Note: Typically not desirable for an animation skeleton.)
+	EdgeLength,
+	// Favor edges between clusters that are earlier in the medial skeleton array. Because medial spheres are added incrementally in the highest error locations, their order can approximate importance.
+	ArrayOrder,
+	// Favor adding edges between larger medial spheres first
+	AvgRadius
+};
+
+UENUM()
+enum class EDataflowMedialSkeletonConversionMergeDisconnectedMethod : uint8
+{
+	// For disconnected components, connect the closest pair of nodes in each disconnected component and the root component
+	ConnectClosestBones,
+	// Add a top-level root node, and connect the selected root and all disconnected components to this node
+	AddTopLevelRoot
+};
+
+
+USTRUCT(meta = (DataflowGeometryCollection))
+struct FSkinMeshViaMedialSkeleton : public FDataflowNode
+{
+	GENERATED_USTRUCT_BODY()
+	DATAFLOW_NODE_DEFINE_INTERNAL(FSkinMeshViaMedialSkeleton, "SkinMeshViaMedialSkeleton", "MedialSkeletonSampling", "Skeletal Mesh Conversion")
+
+private:
+	/** Mesh to add an animation skeleton to */
+	UPROPERTY(meta = (DataflowInput, DataflowOutput, DataflowIntrinsic, DataflowPassthrough = "Mesh"))
+	TObjectPtr<UDynamicMesh> Mesh;
+
+	/** Medial skeleton of input mesh */
+	UPROPERTY(meta = (DataflowInput, DataflowIntrinsic))
+	FDataflowMedialSkeleton MedialSkeleton;
+
+	UPROPERTY(EditAnywhere, Category = "Options")
+	EDataflowMedialSkeletonConversionEdgeWeightMethod EdgeWeightMethod = EDataflowMedialSkeletonConversionEdgeWeightMethod::ArrayOrder;
+
+	UPROPERTY(EditAnywhere, Category = "Options")
+	EDataflowMedialSkeletonConversionMergeDisconnectedMethod MergeDisconnectedMethod = EDataflowMedialSkeletonConversionMergeDisconnectedMethod::ConnectClosestBones;
+
+	/** If non-negative, the index of the medial axis cluster (sphere) to use as the root of the Skeleton */
+	UPROPERTY(EditAnywhere, Category = "Options", meta = (DataflowInput))
+	int32 RootClusterIndex = INDEX_NONE;
+
+	virtual void Evaluate(UE::Dataflow::FContext& Context, const FDataflowOutput* Out) const override;
+
+
+public:
+	FSkinMeshViaMedialSkeleton(const UE::Dataflow::FNodeParameters& InParam, FGuid InGuid = FGuid::NewGuid());
+};
+
+
+namespace UE::Dataflow
+{
+	void RegisterMeshMedialSkeletonConversionNodes();
+}
+
+
+
