add dynamic mesh->skeletal mesh dataflow terminal node

#Synced-CL 50928841

diff --git a/Engine/Plugins/Experimental/GeometryCollectionPlugin/Source/GeometryCollectionNodes/Private/Dataflow/GeometryCollectionNodesPlugin.cpp b/Engine/Plugins/Experimental/GeometryCollectionPlugin/Source/GeometryCollectionNodes/Private/Dataflow/GeometryCollectionNodesPlugin.cpp
--- a/Engine/Plugins/Experimental/GeometryCollectionPlugin/Source/GeometryCollectionNodes/Private/Dataflow/GeometryCollectionNodesPlugin.cpp
+++ b/Engine/Plugins/Experimental/GeometryCollectionPlugin/Source/GeometryCollectionNodes/Private/Dataflow/GeometryCollectionNodesPlugin.cpp
@@ -40,6 +40,7 @@
 #include "Dataflow/SetVertexColorFromVertexIndicesNode.h"
 #include "Dataflow/SetVertexColorFromVertexSelectionNode.h"
 #include "Dataflow/MeshMedialSkeletonConversionNodes.h"
+#include "Dataflow/MeshToSkeletalMeshTerminalNode.h"
 #include "Dataflow/DataflowNodeAndPinTypeColors.h"
 #include "Dataflow/DataflowCategoryRegistry.h"
 #include "GeometryCollection/GeometryCollectionObject.h"
@@ -68,12 +69,13 @@ void IGeometryCollectionNodesPlugin::StartupModule()
 	UE::Dataflow::GeometryCollectionArrayNodes();
 	UE::Dataflow::GeometryCollectionDebugNodes();
 	UE::Dataflow::GeometryCollectionSamplingNodes();
-	UE::Dataflow::RegisterMeshMedialSkeletonConversionNodes();
 	UE::Dataflow::GeometryCollectionRemoveMeshOverlapsNode();
 	UE::Dataflow::RegisterGeometryCollectionUVNodes();
 	UE::Dataflow::RegisterGeometryCollectionTextureNodes();
 	UE::Dataflow::RegisterGeometryCollectionMaterialInterfaceNodes();
 	UE::Dataflow::RegisterGeometryStaticMeshToCollectionNodes();
+	UE::Dataflow::RegisterMeshMedialSkeletonConversionNodes();
+	UE::Dataflow::RegisterMeshToSkeletalMeshTerminalNode();
 	DATAFLOW_NODE_REGISTER_CREATION_FACTORY(FCreateColorArrayFromFloatArrayDataflowNode);
 	DATAFLOW_NODE_REGISTER_CREATION_FACTORY(FGeometryCollectionVertexScalarToVertexIndicesNode);
 	DATAFLOW_NODE_REGISTER_CREATION_FACTORY(FSetVertexColorFromFloatArrayDataflowNode);
diff --git a/Engine/Plugins/Experimental/GeometryCollectionPlugin/Source/GeometryCollectionNodes/Private/Dataflow/MeshToSkeletalMeshTerminalNode.cpp b/Engine/Plugins/Experimental/GeometryCollectionPlugin/Source/GeometryCollectionNodes/Private/Dataflow/MeshToSkeletalMeshTerminalNode.cpp
new file mode 100644
--- /dev/null
+++ b/Engine/Plugins/Experimental/GeometryCollectionPlugin/Source/GeometryCollectionNodes/Private/Dataflow/MeshToSkeletalMeshTerminalNode.cpp
@@ -0,0 +1,346 @@
+// Copyright Epic Games, Inc. All Rights Reserved.
+
+#include "Dataflow/MeshToSkeletalMeshTerminalNode.h"
+
+#include "BoneWeights.h"
+#include "DynamicMeshToMeshDescription.h"
+#include "UDynamicMesh.h"
+#include "Rendering/SkeletalMeshModel.h"
+#include "SkeletalMeshAttributes.h"
+
+#include "AssetRegistry/AssetData.h"
+#include "AssetRegistry/AssetRegistryModule.h"
+#include "Misc/PackageName.h"
+
+#include "Engine/SkeletalMesh.h"
+#include "Animation/Skeleton.h"
+#include "Engine/SkinnedAssetCommon.h"
+
+#include "Dataflow/DataflowDebugDrawInterface.h"
+#include "Dataflow/DataflowDebugDraw.h"
+#if WITH_EDITOR
+#include "StaticToSkeletalMeshConverter.h"
+#include "Dataflow/DataflowRenderingViewMode.h"
+#include "Dataflow/DataflowSkeletalMeshNodes.h"
+#endif
+
+#include UE_INLINE_GENERATED_CPP_BY_NAME(MeshToSkeletalMeshTerminalNode)
+
+namespace UE::Dataflow
+{
+
+	void RegisterMeshToSkeletalMeshTerminalNode()
+	{
+		DATAFLOW_NODE_REGISTER_CREATION_FACTORY(FMeshToSkeletalMeshTerminalNode);
+	}
+
+}
+
+
+namespace MeshToSkeletalMeshAssetLocals
+{
+#if WITH_EDITOR
+	// Helper to convert a mesh description to a skeletal mesh.
+	// editor-only due to FStaticToSkeletalMeshConverter usage
+	void MeshDescriptionToSkeletalMesh(const FMeshDescription& MeshDescription, FSkeletalMeshAttributes& Attributes, const TArray<TObjectPtr<UMaterialInterface>>& InMaterials, USkeletalMesh* OutSkeletalMesh, USkeleton* OutSkeleton)
+	{
+		// Reset the skeletal mesh
+		OutSkeletalMesh->SetNumSourceModels(0);
+		OutSkeletalMesh->GetImportedModel()->LODModels.Reset();
+		OutSkeletalMesh->ReleaseResources();
+		OutSkeletalMesh->ReleaseSkinWeightProfileResources();
+
+
+		// Copy bones to skeleton
+		if (Attributes.GetNumBones() &&
+			// Note FReferenceSkeletonModifier does not allow us to remove a root bone, so we can't clear an existing skeleton here; need to be passed a fresh/empty skeleton
+			ensureMsgf(OutSkeleton->GetReferenceSkeleton().GetRefBoneInfo().IsEmpty(), TEXT("Expected an empty skeleton to fill, but skeleton already had bones")))
+		{
+			FReferenceSkeletonModifier Modifier(OutSkeleton);
+			// For now, we assume the bone hierarchy is consistent and can construct a well-formed ref skeleton.
+			FSkeletalMeshAttributesShared::FBoneNameAttributesConstRef BoneNames(Attributes.GetBoneNames());
+			FSkeletalMeshAttributesShared::FBoneParentIndexAttributesConstRef BoneParents(Attributes.GetBoneParentIndices());
+			FSkeletalMeshAttributesShared::FBonePoseAttributesConstRef BonePoses(Attributes.GetBonePoses());
+
+			for (int32 BoneIndex = 0; BoneIndex < Attributes.GetNumBones(); BoneIndex++)
+			{
+				Modifier.Add(FMeshBoneInfo(BoneNames.Get(BoneIndex), BoneNames.Get(BoneIndex).ToString(), BoneParents.Get(BoneIndex)),
+					BonePoses.Get(BoneIndex));
+			}
+		}
+
+		// Set skeletal materials from InMaterials
+		TArray<FSkeletalMaterial> SkeletalMaterials;
+		for (int32 MatIndex = 0; MatIndex < InMaterials.Num(); ++MatIndex)
+		{
+			const FName SlotName = FName("Material", MatIndex);
+
+			FSkeletalMaterial NewMaterial;
+			NewMaterial.MaterialInterface = InMaterials.IsValidIndex(MatIndex) ? InMaterials[MatIndex] : nullptr;
+			NewMaterial.MaterialSlotName = SlotName;
+			NewMaterial.ImportedMaterialSlotName = SlotName;
+			SkeletalMaterials.Add(NewMaterial);
+		}
+
+		TArray<const FMeshDescription*> MeshDescriptions = { &MeshDescription };
+
+		// Update the skeletal mesh from the description
+		FStaticToSkeletalMeshConverter::FInitializationParams Parameters;
+		Parameters.Materials = SkeletalMaterials;
+		Parameters.bRecomputeNormals = false;
+		Parameters.bRecomputeTangents = false;
+
+		FStaticToSkeletalMeshConverter::InitializeSkeletalMeshFromMeshDescriptions(
+			OutSkeletalMesh,
+			MeshDescriptions,
+			OutSkeleton->GetReferenceSkeleton(),
+			Parameters);
+
+		OutSkeletalMesh->SetSkeleton(OutSkeleton);
+		OutSkeleton->SetPreviewMesh(OutSkeletalMesh);
+		OutSkeleton->RecreateBoneTree(OutSkeletalMesh);
+	}
+
+	bool CanReplaceSkeletonAsset(const USkeleton* SkeletonAsset, const USkeletalMesh* AllowedToReference)
+	{
+		IAssetRegistry& AssetRegistry = FModuleManager::LoadModuleChecked<FAssetRegistryModule>("AssetRegistry").Get();
+
+		TArray<FAssetIdentifier> Referencers;
+		AssetRegistry.GetReferencers(SkeletonAsset->GetPackage()->GetFName(), Referencers);
+
+		FName AllowedPackage = AllowedToReference ? AllowedToReference->GetPackage()->GetFName() : FName();
+
+		for (const FAssetIdentifier& Referencer : Referencers)
+		{
+			TArray<FAssetData> Assets;
+			AssetRegistry.GetAssetsByPackageName(Referencer.PackageName, Assets, false);
+
+			// handle the no-assets-allowed case
+			if (!AllowedToReference && !Assets.IsEmpty())
+			{
+				return false;
+			}
+
+			for (const FAssetData& Asset : Assets)
+			{
+				// if allowed to reference a specific skeletal mesh, return false if it's not that mesh
+				if (!Asset.IsInstanceOf(USkeletalMesh::StaticClass()) || Asset.PackageName != AllowedPackage)
+				{
+					return false;
+				}
+			}
+		}
+		return true;
+	}
+
+	template<typename AssetType>
+	AssetType* MakeAssetFromPath(TFunctionRef<void(const FString&)> LogErr, const FString& InAssetPath, bool& bUsingExisting, bool bMustOverwriteExisting = false, TFunctionRef<bool(UObject* Existing)> CanOverwiteExisting = [](UObject*) {return true;})
+	{
+		if (!FPackageName::IsValidObjectPath(InAssetPath))
+		{
+			LogErr(FString::Format(TEXT("Asset Path input is not a valid path : {0}"), { InAssetPath }));
+			return nullptr;
+		}
+		const FString PackageName = FPackageName::ObjectPathToPackageName(InAssetPath);
+
+		UPackage* Package = Cast<UPackage>(FindPackage(nullptr, *PackageName));
+		if (Package == nullptr)
+		{
+			Package = CreatePackage(*PackageName);
+		}
+		if (Package == nullptr)
+		{
+			LogErr(FString::Format(TEXT("Failed to find or create package {0}"), { PackageName }));
+			return nullptr;
+		}
+
+		const FName AssetName = FName(FPackageName::GetLongPackageAssetName(PackageName));
+		UObject* ExistingObject = StaticFindObjectFastInternal( /*Class=*/ NULL, Package, AssetName, EFindObjectFlags::ExactClass);
+		if (ExistingObject && !ExistingObject->GetClass()->IsChildOf<AssetType>())
+		{
+			LogErr(FString::Format(TEXT("Asset {0} already exists but is not a compatible type"), { InAssetPath }));
+			return nullptr;
+		}
+		// It is an error if we are trying to replace an object of a different class
+		TObjectPtr<AssetType> OutAsset = Cast<AssetType>(ExistingObject);
+		if (ExistingObject && bMustOverwriteExisting)
+		{
+			if (!CanOverwiteExisting(ExistingObject))
+			{
+				LogErr(FString::Format(TEXT("Asset {0} already exists, and could not overwrite it"), { InAssetPath }));
+				return nullptr;
+			}
+			// Overwrite is ok, create a fresh package for it
+			Package = CreatePackage(*PackageName);
+		}
+		if (!OutAsset || bMustOverwriteExisting)
+		{
+			bUsingExisting = false;
+			OutAsset = NewObject<AssetType>(Package, AssetType::StaticClass(), AssetName, RF_Public | RF_Standalone | RF_Transactional);
+			if (!OutAsset)
+			{
+				LogErr(FString::Format(TEXT("Failed to create asset {0}"), { InAssetPath }));
+				return nullptr;
+			}
+		}
+		else
+		{
+			bUsingExisting = true;
+		}
+
+		// make sure the asset is set properly
+		OutAsset->MarkPackageDirty();
+		FAssetRegistryModule::AssetCreated(OutAsset.Get());
+		return OutAsset;
+	}
+#endif
+}
+
+
+void FMeshToSkeletalMeshTerminalNode::Evaluate(UE::Dataflow::FContext& Context) const
+{
+	// Asset is created and stored by the SetAssetValue
+	// Evaluate only return an existing one or nullptr
+	const FString& InAssetPath = GetValue(Context, &SkeletalMeshAssetPath);
+	TObjectPtr<USkeletalMesh> OutSkeletalMesh = ::LoadObject<USkeletalMesh>(nullptr, InAssetPath);
+	SetValue(Context, OutSkeletalMesh, &SkeletalMeshAsset);
+	TObjectPtr<USkeleton> OutSkeleton = nullptr;
+	if (OutSkeletalMesh)
+	{
+		OutSkeleton = OutSkeletalMesh->GetSkeleton();
+	}
+	SetValue(Context, OutSkeleton, &SkeletonAsset);
+}
+void FMeshToSkeletalMeshTerminalNode::SetAssetValue(TObjectPtr<UObject> Asset, UE::Dataflow::FContext& Context) const
+{
+	auto SetOutputsOnFailure = [this, &Context]()
+	{
+		TObjectPtr<USkeletalMesh> NullSkeletalMeshPtr = nullptr;
+		TObjectPtr<USkeleton> NullSkeletonPtr = nullptr;
+		SetValue(Context, NullSkeletalMeshPtr, &SkeletalMeshAsset);
+		SetValue(Context, NullSkeletonPtr, &SkeletonAsset);
+	};
+
+#if WITH_EDITOR
+
+	const FString& InAssetPath = GetValue(Context, &SkeletalMeshAssetPath);
+	auto LogErr = [&Context, this](const FString& Msg) { Context.Error(Msg, this); };
+	bool bUsingExistingSkeletalMesh = false;
+	// We try to re-use existing skeletal mesh assets when possible
+	constexpr bool bMustReplaceExistingSkeletalMesh = false;
+	TObjectPtr<USkeletalMesh> OutSkeletalMesh = MeshToSkeletalMeshAssetLocals::MakeAssetFromPath<USkeletalMesh>(LogErr, InAssetPath, bUsingExistingSkeletalMesh, bMustReplaceExistingSkeletalMesh);
+	if (!OutSkeletalMesh)
+	{
+		SetOutputsOnFailure();
+		return;
+	}
+	// Note we try to create a fresh skeleton at the target asset path, rather than re-use an existing skeleton
+	// because we have no guarantee about the skeleton topology being safe to re-use if it is referenced by other assets
+	// and we don't want to solve a re-targeting problem here.
+	// TODO: We could try to verify if the skeleton is compatible and re-use it in that case?
+	TObjectPtr<USkeleton> OutSkeleton = nullptr;
+	auto CanOverwriteSkeletonAsset = [OutSkeletalMesh](UObject* SkelObject) -> bool
+	{
+		if (USkeleton* TestSkeleton = Cast<USkeleton>(SkelObject))
+		{
+			return MeshToSkeletalMeshAssetLocals::CanReplaceSkeletonAsset(TestSkeleton, OutSkeletalMesh);
+		}
+		return false;
+	};
+	bool bUsingExistingSkeleton = false;
+	if (IsConnected(&SkeletonAssetPath))
+	{
+		const FString& InSkelAssetPath = GetValue(Context, &SkeletonAssetPath);
+		OutSkeleton = MeshToSkeletalMeshAssetLocals::MakeAssetFromPath<USkeleton>(LogErr, InSkelAssetPath, bUsingExistingSkeleton, true, CanOverwriteSkeletonAsset);
+	}
+	else if (!OutSkeleton)
+	{
+		OutSkeleton = MeshToSkeletalMeshAssetLocals::MakeAssetFromPath<USkeleton>(LogErr, FString::Printf(TEXT("%s_Skeleton"), *InAssetPath), bUsingExistingSkeleton, true, CanOverwriteSkeletonAsset);
+	}
+	if (!OutSkeleton)
+	{
+		LogErr(TEXT("Failed to find or create a skeleton for the skeletal mesh"));
+		// TODO: could we set a null skeleton here rather than failing?
+		SetOutputsOnFailure();
+		return;
+	}
+
+	// When re-using a skeletal mesh asset, we manually clear existing data so we don't e.g. have anything referencing a previous skeleton
+	if (bUsingExistingSkeletalMesh)
+	{
+		OutSkeletalMesh->PreEditChange(nullptr);
+
+		OutSkeletalMesh->Clear();
+
+		// Make sure that we create non-transient assets
+		OutSkeletalMesh->ClearFlags(RF_Transient);
+	}
+
+	// update with mesh
+	TObjectPtr<const UDynamicMesh> InMesh = GetValue(Context, &Mesh);
+
+	if (InMesh)
+	{
+		const FDynamicMesh3& DynMesh = InMesh->GetMeshRef();
+		const TArray<TObjectPtr<UMaterialInterface>>& InMaterials = GetValue(Context, &Materials);
+		// convert the collection to a mesh description including skeletal mesh attributes
+		FMeshDescription MeshDescription;
+
+		FSkeletalMeshAttributes Attributes(MeshDescription);
+		Attributes.Register();
+
+		FConversionToMeshDescriptionOptions ConversionOptions;
+		FDynamicMeshToMeshDescription Converter(ConversionOptions);
+		Converter.Convert(&DynMesh, MeshDescription);
+
+		MeshToSkeletalMeshAssetLocals::MeshDescriptionToSkeletalMesh(MeshDescription, Attributes, InMaterials, OutSkeletalMesh, OutSkeleton);
+	}
+
+	if (bUsingExistingSkeletalMesh)
+	{
+		OutSkeletalMesh->PostEditChange();
+	}
+
+	// finally set the skeletal mesh to the output
+	SetValue(Context, OutSkeletalMesh, &SkeletalMeshAsset);
+	SetValue(Context, OutSkeleton, &SkeletonAsset);
+
+#else
+	Context.Error(TEXT("Creation of skeletal mesh asset only supported in Editor"), this);
+	SetOutputsOnFailure();
+#endif
+}
+
+FMeshToSkeletalMeshTerminalNode::FMeshToSkeletalMeshTerminalNode(const UE::Dataflow::FNodeParameters& InParam, FGuid InGuid)
+	: FDataflowTerminalNode(InParam, InGuid)
+{
+	RegisterInputConnection(&Mesh);
+	RegisterInputConnection(&Materials);
+	RegisterInputConnection(&SkeletalMeshAssetPath);
+	RegisterInputConnection(&SkeletonAssetPath);
+	RegisterOutputConnection(&SkeletalMeshAsset);
+	RegisterOutputConnection(&SkeletonAsset);
+}
+
+#if WITH_EDITOR
+void FMeshToSkeletalMeshTerminalNode::DebugDraw(UE::Dataflow::FContext& Context,
+	IDataflowDebugDrawInterface& DataflowRenderingInterface,
+	const FDebugDrawParameters& DebugDrawParameters) const
+{
+	if (const FDataflowOutput* Output = FindOutput(&SkeletalMeshAsset))
+	{
+		if (const TObjectPtr<const USkeletalMesh> InSkeletalMeshAsset = Output->ReadValue(Context, SkeletalMeshAsset))
+		{
+			TRefCountPtr<IDataflowDebugDrawObject> SkeletonObject(MakeDebugDrawObject<FDataflowDebugDrawSkeletonObject>(
+				DataflowRenderingInterface.ModifyDataflowElements(), InSkeletalMeshAsset->GetRefSkeleton()));
+
+			DataflowRenderingInterface.DrawObject(SkeletonObject);
+		}
+	}
+}
+
+bool FMeshToSkeletalMeshTerminalNode::CanDebugDrawViewMode(const FName& ViewModeName) const
+{
+	return ViewModeName == UE::Dataflow::FDataflowConstruction3DViewMode::Name;
+}
+#endif
diff --git a/Engine/Plugins/Experimental/GeometryCollectionPlugin/Source/GeometryCollectionNodes/Private/Dataflow/MeshToSkeletalMeshTerminalNode.h b/Engine/Plugins/Experimental/GeometryCollectionPlugin/Source/GeometryCollectionNodes/Private/Dataflow/MeshToSkeletalMeshTerminalNode.h
new file mode 100644
--- /dev/null
+++ b/Engine/Plugins/Experimental/GeometryCollectionPlugin/Source/GeometryCollectionNodes/Private/Dataflow/MeshToSkeletalMeshTerminalNode.h
@@ -0,0 +1,70 @@
+// Copyright Epic Games, Inc. All Rights Reserved.
+
+#pragma once 
+
+#include "CoreMinimal.h"
+#include "Dataflow/DataflowEngine.h"
+#include "GeometryCollection/ManagedArrayCollection.h"
+#include "UDynamicMesh.h"
+
+
+#include "MeshToSkeletalMeshTerminalNode.generated.h"
+
+
+/*
+* Terminal node to a save a skeletal mesh asset, as converted from a dynamic mesh.
+* Note the mesh must have associated skeletal mesh attributes.
+* The terminal node will create an asset if it does not exist yet
+*/
+USTRUCT(meta = (DataflowTerminal))
+struct FMeshToSkeletalMeshTerminalNode : public FDataflowTerminalNode
+{
+	GENERATED_USTRUCT_BODY()
+	DATAFLOW_NODE_DEFINE_INTERNAL(FMeshToSkeletalMeshTerminalNode, "MeshToSkeletalMeshTerminal", "Terminal", "Path Asset Mesh Skeletal Skeleton")
+
+private:
+	/** Surface mesh to convert */
+	UPROPERTY(meta = (DataflowInput, DataflowIntrinsic))
+	TObjectPtr<const UDynamicMesh> Mesh;
+
+	/** Materials to use on the skeletal mesh */
+	UPROPERTY(meta = (DataflowInput))
+	TArray<TObjectPtr<UMaterialInterface>> Materials;
+
+
+	/** Path to the skeletal mesh asset to create or update */
+	UPROPERTY(EditAnywhere, Category = Asset, meta = (DataflowInput))
+	FString SkeletalMeshAssetPath;
+
+	/** Path to the skeleton asset to create or update */
+	UPROPERTY(EditAnywhere, Category = Asset, meta = (DataflowInput))
+	FString SkeletonAssetPath;
+
+	/** Created Skeletal Mesh */
+	UPROPERTY(meta = (DataflowOutput))
+	TObjectPtr<USkeletalMesh> SkeletalMeshAsset = nullptr;
+
+	/** Created Skeleton */
+	UPROPERTY(meta = (DataflowOutput))
+	TObjectPtr<USkeleton> SkeletonAsset = nullptr;
+
+	virtual void Evaluate(UE::Dataflow::FContext& Context) const override;
+	virtual void SetAssetValue(TObjectPtr<UObject> Asset, UE::Dataflow::FContext& Context) const override;
+
+#if WITH_EDITOR
+	virtual void DebugDraw(UE::Dataflow::FContext& Context, IDataflowDebugDrawInterface& DataflowRenderingInterface, const FDebugDrawParameters& DebugDrawParameters) const override;
+	virtual bool CanDebugDrawViewMode(const FName& ViewModeName) const override;
+#endif 
+
+public:
+	FMeshToSkeletalMeshTerminalNode(const UE::Dataflow::FNodeParameters& InParam, FGuid InGuid = FGuid::NewGuid());
+};
+
+
+namespace UE::Dataflow
+{
+	void RegisterMeshToSkeletalMeshTerminalNode();
+}
+
+
+
