Remove MinErrorDelta parameter from Mesh Medial Axis Sampling, as it is confusing to have two similar error terms and the MinErrorDelta is easily confused by a not-useful sphere being added (especially w/ the sphere projection step after optimization, which can make error increase in early/unlucky cases)

#Synced-CL 51328354

diff --git a/Engine/Plugins/Experimental/GeometryCollectionPlugin/Source/GeometryCollectionNodes/Private/Dataflow/GeometryCollectionSamplingNodes.cpp b/Engine/Plugins/Experimental/GeometryCollectionPlugin/Source/GeometryCollectionNodes/Private/Dataflow/GeometryCollectionSamplingNodes.cpp
--- a/Engine/Plugins/Experimental/GeometryCollectionPlugin/Source/GeometryCollectionNodes/Private/Dataflow/GeometryCollectionSamplingNodes.cpp
+++ b/Engine/Plugins/Experimental/GeometryCollectionPlugin/Source/GeometryCollectionNodes/Private/Dataflow/GeometryCollectionSamplingNodes.cpp
@@ -384,7 +384,6 @@ FMeshMedialSkeletonSamplingDataflowNode::FMeshMedialSkeletonSamplingDataflowNode
 {
 	RegisterInputConnection(&TargetMesh);
 	RegisterInputConnection(&MaxSpheres).SetCanHidePin(true);
-	RegisterInputConnection(&MinErrorDelta).SetCanHidePin(true).SetPinIsHidden(true);
 	RegisterInputConnection(&MinClusterErrorToSplit).SetCanHidePin(true).SetPinIsHidden(true);
 	RegisterInputConnection(&MinClusterSizeToKeep).SetCanHidePin(true).SetPinIsHidden(true);
 	RegisterInputConnection(&bSplitClustersIfEdgesIntersectSurface).SetCanHidePin(true).SetPinIsHidden(true);
@@ -413,7 +412,6 @@ void FMeshMedialSkeletonSamplingDataflowNode::Evaluate(UE::Dataflow::FContext& C
 				UE::Geometry::TFastWindingTree<FDynamicMesh3> FWTree(&Spatial);
 				UE::Geometry::MedialAxis::FSkeletonViaSampling SkelSample;
 				SkelSample.MaxSpheres = GetValue(Context, &MaxSpheres);
-				SkelSample.MinErrorDelta = GetValue(Context, &MinErrorDelta);
 				SkelSample.MinClusterErrorToSplit = GetValue(Context, &MinClusterErrorToSplit);
 				SkelSample.MinClusterSizeToKeep = GetValue(Context, &MinClusterSizeToKeep);
 				SkelSample.bSplitClustersIfEdgesIntersectSurface = GetValue(Context, &bSplitClustersIfEdgesIntersectSurface);
diff --git a/Engine/Plugins/Experimental/GeometryCollectionPlugin/Source/GeometryCollectionNodes/Public/Dataflow/GeometryCollectionSamplingNodes.h b/Engine/Plugins/Experimental/GeometryCollectionPlugin/Source/GeometryCollectionNodes/Public/Dataflow/GeometryCollectionSamplingNodes.h
--- a/Engine/Plugins/Experimental/GeometryCollectionPlugin/Source/GeometryCollectionNodes/Public/Dataflow/GeometryCollectionSamplingNodes.h
+++ b/Engine/Plugins/Experimental/GeometryCollectionPlugin/Source/GeometryCollectionNodes/Public/Dataflow/GeometryCollectionSamplingNodes.h
@@ -373,10 +373,6 @@ private:
 	UPROPERTY(EditAnywhere, Category = Options, meta = (DataflowInput, UIMin = .1, ClampMin = 0.))
 	float MinClusterErrorToSplit = 10.f;
 
-	// Stop when the average squared error at vertices decreases by less than this w/ an added sphere addition
-	UPROPERTY(EditAnywhere, Category = Options, meta = (DataflowInput, UIMin = .1, ClampMin = 0.))
-	float MinErrorDelta = 1.f;
-
 	// Clusters with fewer vertices than this will be discarded
 	UPROPERTY(EditAnywhere, Category = Options, meta = (DataflowInput, ClampMin = 1, UIMax = 10))
 	int32 MinClusterSizeToKeep = 4;
diff --git a/Engine/Plugins/Runtime/GeometryProcessing/Source/GeometryAlgorithms/Private/Skeletonization/MeshMedialAxisSampling.cpp b/Engine/Plugins/Runtime/GeometryProcessing/Source/GeometryAlgorithms/Private/Skeletonization/MeshMedialAxisSampling.cpp
--- a/Engine/Plugins/Runtime/GeometryProcessing/Source/GeometryAlgorithms/Private/Skeletonization/MeshMedialAxisSampling.cpp
+++ b/Engine/Plugins/Runtime/GeometryProcessing/Source/GeometryAlgorithms/Private/Skeletonization/MeshMedialAxisSampling.cpp
@@ -719,25 +719,6 @@ namespace UE::Geometry::MedialAxis
 				break;
 			}
 
-			// Stop adding spheres if error stops decreasing
-			if (MinErrorDelta > 0)
-			{
-				double AvgError = 0;
-				for (int32 VID = 0; VID < Mesh.MaxVertexID(); ++VID)
-				{
-					if (Result.VIDtoClusterIndex[VID] != INDEX_NONE)
-					{
-						double AreaWt = (double)VtxAreas[VID];
-						AvgError += AreaWt * GetVertexAvgSquaredError(VID, Result.Spheres[Result.VIDtoClusterIndex[VID]]);
-					}
-				}
-				AvgError /= TotalArea;
-				if (!bHasSurfaceIntersections && LastAvgError - AvgError < MinErrorDelta * MinErrorDelta)
-				{
-					break;
-				}
-				LastAvgError = AvgError;
-			}
 
 		}
 
diff --git a/Engine/Plugins/Runtime/GeometryProcessing/Source/GeometryAlgorithms/Public/Skeletonization/MeshMedialAxisSampling.h b/Engine/Plugins/Runtime/GeometryProcessing/Source/GeometryAlgorithms/Public/Skeletonization/MeshMedialAxisSampling.h
--- a/Engine/Plugins/Runtime/GeometryProcessing/Source/GeometryAlgorithms/Public/Skeletonization/MeshMedialAxisSampling.h
+++ b/Engine/Plugins/Runtime/GeometryProcessing/Source/GeometryAlgorithms/Public/Skeletonization/MeshMedialAxisSampling.h
@@ -384,8 +384,6 @@ namespace MedialAxis
 
 		// Stop at this sphere count
 		int32 MaxSpheres = 10;
-		// Stop when the average squared error at vertices decreases by less than this w/ an added sphere addition
-		double MinErrorDelta = .1;
 		// Do not split a cluster if its max error is below this threshold
 		double MinClusterErrorToSplit = .1;
 		// Weight for the position error term. Relative to plane error term, which has weight 1.
