add DataflowSkeleton as simple wrapper of FReferenceSkeleton + add two initial nodes to use it -- medial skeleton->skeleton + bind skeleton to mesh

#Synced-CL 51328354

diff --git a/Engine/Plugins/Experimental/GeometryCollectionPlugin/Source/GeometryCollectionNodes/Private/Dataflow/MeshMedialSkeletonConversionNodes.cpp b/Engine/Plugins/Experimental/GeometryCollectionPlugin/Source/GeometryCollectionNodes/Private/Dataflow/MeshMedialSkeletonConversionNodes.cpp
--- a/Engine/Plugins/Experimental/GeometryCollectionPlugin/Source/GeometryCollectionNodes/Private/Dataflow/MeshMedialSkeletonConversionNodes.cpp
+++ b/Engine/Plugins/Experimental/GeometryCollectionPlugin/Source/GeometryCollectionNodes/Private/Dataflow/MeshMedialSkeletonConversionNodes.cpp
@@ -5,6 +5,7 @@
 #include "DynamicMesh/DynamicBoneAttribute.h"
 #include "DynamicMesh/DynamicMeshAttributeSet.h"
 #include "DynamicMesh/DynamicVertexSkinWeightsAttribute.h"
+#include "Operations/SkinWeightBinding.h"
 #include "UDynamicMesh.h"
 
 #include "SkeletalMeshAttributes.h" // for FSkeletalMeshAttributes::DefaultSkinWeightProfileName
@@ -12,11 +13,15 @@
 
 #include UE_INLINE_GENERATED_CPP_BY_NAME(MeshMedialSkeletonConversionNodes)
 
+#define LOCTEXT_NAMESPACE "ChaosClothAssetWeightMapNode"
+
 namespace UE::Dataflow
 {
 
 	void RegisterMeshMedialSkeletonConversionNodes()
 	{
+		DATAFLOW_NODE_REGISTER_CREATION_FACTORY(FConvertMedialSkeletonToAnimationSkeletonDataflowNode);
+		DATAFLOW_NODE_REGISTER_CREATION_FACTORY(FBindSkeletonToMeshDataflowNode);
 		DATAFLOW_NODE_REGISTER_CREATION_FACTORY(FSkinMeshViaMedialSkeleton);
 	}
 
@@ -53,12 +58,142 @@ namespace MeshMedialSkeletonConversionLocals
 		case EDataflowMedialSkeletonConversionMergeDisconnectedMethod::ConnectClosestBones:
 			OutOptions.MergeDisconnectedMethod = FMedialSkeletonToTreeSkeletonOptions::EMergeDisconnectedMethod::ConnectClosestBones;
 			break;
+		}	
+	}
+
+	void SetSkinBindSettings(UE::Geometry::SkinBinding::FBindSettings& OutSettings, EDataflowBindSkeletonMethod BindMethod, float Stiffness, int32 MaxInfluences, int32 VoxelRes)
+	{
+		using namespace UE::Geometry;
+		switch (BindMethod)
+		{
+		case EDataflowBindSkeletonMethod::DirectDistance:
+			OutSettings.BindType = ESkinBindingType::DirectDistance;
+			break;
+		case EDataflowBindSkeletonMethod::GeodesicVoxel:
+			OutSettings.BindType = ESkinBindingType::GeodesicVoxel;
+			break;
+		}
+		OutSettings.MaxInfluences = MaxInfluences;
+		OutSettings.Stiffness = Stiffness;
+		OutSettings.VoxelResolution = VoxelRes;
+	}
+
+	TArray<UE::Geometry::SkinBinding::FBonePoseInfo> GetBindingBoneInfoFromReferenceSkeleton(const FReferenceSkeleton& InRefSkeleton)
+	{
+		// Only use non-virtual bones, since virtual bones cannot be bound to the skin.
+		const TArray<FMeshBoneInfo>& BoneInfo = InRefSkeleton.GetRawRefBoneInfo();
+		const TArray<FTransform>& BonePose = InRefSkeleton.GetRawRefBonePose();
+
+		TArray<UE::Geometry::SkinBinding::FBonePoseInfo> Bones;
+		Bones.Reserve(BoneInfo.Num());
+
+		for (int32 Index = 0; Index < BoneInfo.Num(); Index++)
+		{
+			UE::Geometry::SkinBinding::FBonePoseInfo Info{ .LocalTransform = BonePose[Index], .Name = BoneInfo[Index].Name, .ParentIndex = BoneInfo[Index].ParentIndex };
+			Bones.Add(Info);
+		}
+		return Bones;
+	}
+
+}
+
+// ---------------------------------------------------------------------------------------------------------------------------------------
+ 
+FConvertMedialSkeletonToAnimationSkeletonDataflowNode::FConvertMedialSkeletonToAnimationSkeletonDataflowNode(const UE::Dataflow::FNodeParameters& InParam, FGuid InGuid)
+	: FDataflowNode(InParam, InGuid)
+{
+	RegisterInputConnection(&MedialSkeleton);
+	RegisterInputConnection(&RootClusterIndex).SetCanHidePin(true).SetPinIsHidden(true);
+
+	RegisterOutputConnection(&Skeleton);
+}
+
+void FConvertMedialSkeletonToAnimationSkeletonDataflowNode::Evaluate(UE::Dataflow::FContext& Context, const FDataflowOutput* Out) const
+{
+	using namespace UE::Geometry;
+
+	if (Out->IsA(&Skeleton))
+	{
+		const FDataflowMedialSkeleton& InMedialSkeleton = GetValue(Context, &MedialSkeleton);
+		FDataflowSkeleton OutSkeleton;
+
+		TArray<int32> Parents;
+		TArray<FTransform> Poses;
+
+		UE::Geometry::MedialAxis::FMedialSkeletonToTreeSkeletonOptions Options;
+		MeshMedialSkeletonConversionLocals::SetSkeletonConversionOptions(Options, EdgeWeightMethod, MergeDisconnectedMethod);
+		const int32 UseRootIndex = GetValue(Context, &RootClusterIndex);
+		Options.ToHierarchy(InMedialSkeleton.Skeleton, Parents, Poses, Options, UseRootIndex);
+
+		const int32 NumBones = Parents.Num();
+		if (NumBones > 0)
+		{
+			FReferenceSkeletonModifier Modifier = OutSkeleton.ModifySkeleton();
+			for (int32 BoneIdx = 0; BoneIdx < NumBones; ++BoneIdx)
+			{
+				FName BoneName("Bone", BoneIdx);
+				Modifier.Add(FMeshBoneInfo(BoneName, BoneName.ToString(), Parents[BoneIdx]), Poses[BoneIdx], true);
+			}
 		}
-		
+
+		SetValue(Context, MoveTemp(OutSkeleton), &Skeleton);
 	}
 }
 
 
+
+// ---------------------------------------------------------------------------------------------------------------------------------------
+ 
+FBindSkeletonToMeshDataflowNode::FBindSkeletonToMeshDataflowNode(const UE::Dataflow::FNodeParameters& InParam, FGuid InGuid)
+	: FDataflowNode(InParam, InGuid)
+{
+	RegisterInputConnection(&Mesh);
+	RegisterInputConnection(&Skeleton);
+	RegisterInputConnection(&Stiffness).SetCanHidePin(true).SetPinIsHidden(true);
+	RegisterInputConnection(&MaxInfluences).SetCanHidePin(true).SetPinIsHidden(true);
+	RegisterInputConnection(&VoxelResolution).SetCanHidePin(true).SetPinIsHidden(true);
+
+	RegisterOutputConnection(&Mesh, &Mesh);
+}
+
+void FBindSkeletonToMeshDataflowNode::Evaluate(UE::Dataflow::FContext& Context, const FDataflowOutput* Out) const
+{
+	using namespace UE::Geometry;
+	using namespace MeshMedialSkeletonConversionLocals;
+
+	if (Out->IsA(&Mesh))
+	{
+		TObjectPtr<UDynamicMesh> InMesh = GetValue(Context, &Mesh);
+		const FDataflowSkeleton& InSkeleton = GetValue(Context, &Skeleton);
+		TObjectPtr<UDynamicMesh> OutMesh = NewObject<UDynamicMesh>();
+
+		if (InMesh)
+		{
+			// Create a new mesh copy and bind skeleton
+			OutMesh->SetMesh(InMesh->GetMeshRef());
+			FDynamicMesh3& DynMesh = OutMesh->GetMeshRef();
+
+			SkinBinding::FBindSettings Settings;
+			SetSkinBindSettings(Settings, BindMethod,
+				GetValue(Context, &Stiffness),
+				GetValue(Context, &MaxInfluences),
+				GetValue(Context, &VoxelResolution));
+
+			TArray<SkinBinding::FBonePoseInfo> Bones = GetBindingBoneInfoFromReferenceSkeleton(InSkeleton.GetRefSkeleton());
+			SkinBinding::CreateSkinWeights(DynMesh, Bones, FSkeletalMeshAttributes::DefaultSkinWeightProfileName, Settings);
+		}
+		else
+		{
+			Context.Error(LOCTEXT("BindSkeletonNullMeshError", "BindSkeletonToMesh: Cannot bind skeleton to null input mesh"), this);
+		}
+		SetValue(Context, OutMesh, &Mesh);
+	}
+}
+
+
+
+// ---------------------------------------------------------------------------------------------------------------------------------------
+
 FSkinMeshViaMedialSkeleton::FSkinMeshViaMedialSkeleton(const UE::Dataflow::FNodeParameters& InParam, FGuid InGuid)
 	: FDataflowNode(InParam, InGuid)
 {
@@ -91,7 +226,12 @@ void FSkinMeshViaMedialSkeleton::Evaluate(UE::Dataflow::FContext& Context, const
 			Options.ToDynamicMeshSkeleton(InMedialSkeleton.Skeleton, DynMesh, FSkeletalMeshAttributes::DefaultSkinWeightProfileName, Options, UseRootIndex);
 	
 		}
+		else
+		{
+			Context.Error(LOCTEXT("SkinMeshViaMedialSkeletonNullMeshError", "SkinMeshViaMedialSkeleton: Cannot bind skeleton to null input mesh"), this);
+		}
 		SetValue(Context, OutMesh, &Mesh);
 	}
 }
 
+#undef LOCTEXT_NAMESPACE
diff --git a/Engine/Plugins/Experimental/GeometryCollectionPlugin/Source/GeometryCollectionNodes/Private/Dataflow/MeshMedialSkeletonConversionNodes.h b/Engine/Plugins/Experimental/GeometryCollectionPlugin/Source/GeometryCollectionNodes/Private/Dataflow/MeshMedialSkeletonConversionNodes.h
--- a/Engine/Plugins/Experimental/GeometryCollectionPlugin/Source/GeometryCollectionNodes/Private/Dataflow/MeshMedialSkeletonConversionNodes.h
+++ b/Engine/Plugins/Experimental/GeometryCollectionPlugin/Source/GeometryCollectionNodes/Private/Dataflow/MeshMedialSkeletonConversionNodes.h
@@ -5,6 +5,7 @@
 #include "CoreMinimal.h"
 #include "Dataflow/DataflowDebugDraw.h"
 #include "Dataflow/DataflowEngine.h"
+#include "Dataflow/DataflowSkeleton.h"
 #include "Dataflow/GeometryCollectionSamplingNodes.h"
 #include "GeometryCollection/ManagedArrayCollection.h"
 
@@ -31,12 +32,95 @@ enum class EDataflowMedialSkeletonConversionMergeDisconnectedMethod : uint8
 	AddTopLevelRoot
 };
 
+UENUM()
+enum class EDataflowBindSkeletonMethod : uint8
+{
+	// Computes the binding strength by computing the Euclidean distance to the closest set of bones,
+	// where the strength of binding is proportional to the inverse distance. May cause bones to affects
+	// parts of geometry that, although close in space, may be topologically distant.
+	DirectDistance = 0,
+
+	// Computes the binding by computing the geodesic distance from each set of bones. This is slower than the
+	// direct distance.
+	GeodesicVoxel = 1,
+};
+
+USTRUCT(meta = (DataflowGeometryCollection))
+struct FConvertMedialSkeletonToAnimationSkeletonDataflowNode : public FDataflowNode
+{
+	GENERATED_USTRUCT_BODY()
+	DATAFLOW_NODE_DEFINE_INTERNAL(FConvertMedialSkeletonToAnimationSkeletonDataflowNode, "MedialToAnimSkeleton", "SkeletonUtil", "Skeletal Conversion")
+
+private:
+
+	/** Medial skeleton of input mesh */
+	UPROPERTY(meta = (DataflowInput, DataflowIntrinsic))
+	FDataflowMedialSkeleton MedialSkeleton;
+
+	UPROPERTY(meta = (DataflowOutput))
+	FDataflowSkeleton Skeleton;
+
+	UPROPERTY(EditAnywhere, Category = "Options")
+	EDataflowMedialSkeletonConversionEdgeWeightMethod EdgeWeightMethod = EDataflowMedialSkeletonConversionEdgeWeightMethod::ArrayOrder;
+
+	UPROPERTY(EditAnywhere, Category = "Options")
+	EDataflowMedialSkeletonConversionMergeDisconnectedMethod MergeDisconnectedMethod = EDataflowMedialSkeletonConversionMergeDisconnectedMethod::ConnectClosestBones;
+
+	/** If non-negative, the index of the medial axis cluster (sphere) to use as the root of the Skeleton */
+	UPROPERTY(EditAnywhere, Category = "Options", meta = (DataflowInput))
+	int32 RootClusterIndex = INDEX_NONE;
+
+	virtual void Evaluate(UE::Dataflow::FContext& Context, const FDataflowOutput* Out) const override;
+
+
+public:
+	FConvertMedialSkeletonToAnimationSkeletonDataflowNode(const UE::Dataflow::FNodeParameters& InParam, FGuid InGuid = FGuid::NewGuid());
+};
+
+USTRUCT(meta = (DataflowGeometryCollection))
+struct FBindSkeletonToMeshDataflowNode : public FDataflowNode
+{
+	GENERATED_USTRUCT_BODY()
+	DATAFLOW_NODE_DEFINE_INTERNAL(FBindSkeletonToMeshDataflowNode, "BindSkeletonToMesh", "SkeletonUtil", "Skeletal Conversion")
+
+private:
+
+	/** Mesh to add an animation skeleton to */
+	UPROPERTY(meta = (DataflowInput, DataflowOutput, DataflowIntrinsic, DataflowPassthrough = "Mesh"))
+	TObjectPtr<UDynamicMesh> Mesh;
+
+	UPROPERTY(meta = (DataflowInput, DataflowIntrinsic))
+	FDataflowSkeleton Skeleton;
+
+	/** Binding method to use */
+	UPROPERTY(EditAnywhere, Category = "Options")
+	EDataflowBindSkeletonMethod BindMethod = EDataflowBindSkeletonMethod::DirectDistance;
+
+	/** Stiffness of binding. Lower values allow more distant bones to contribute more */
+	UPROPERTY(EditAnywhere, Category = "Options", meta = (DataflowInput, ClampMin=0, ClampMax=1))
+	float Stiffness = 0.2f;
+
+	/** Maximum bones that will influence each vertex */
+	UPROPERTY(EditAnywhere, Category = "Options", meta = (DataflowInput, ClampMin = 1, UIMax = 10))
+	int32 MaxInfluences = 5;
+
+	/** The resolution of the voxel grid if doing geodesic voxel binding */
+	UPROPERTY(EditAnywhere, Category = "Options", meta = (DataflowInput, ClampMin =1, UIMin = 8, ClampMax = 1024))
+	int32 VoxelResolution = 256;
+
+	virtual void Evaluate(UE::Dataflow::FContext& Context, const FDataflowOutput* Out) const override;
+
+
+public:
+	FBindSkeletonToMeshDataflowNode(const UE::Dataflow::FNodeParameters& InParam, FGuid InGuid = FGuid::NewGuid());
+};
+
 
 USTRUCT(meta = (DataflowGeometryCollection))
 struct FSkinMeshViaMedialSkeleton : public FDataflowNode
 {
 	GENERATED_USTRUCT_BODY()
-	DATAFLOW_NODE_DEFINE_INTERNAL(FSkinMeshViaMedialSkeleton, "SkinMeshViaMedialSkeleton", "MedialSkeletonSampling", "Skeletal Mesh Conversion")
+	DATAFLOW_NODE_DEFINE_INTERNAL(FSkinMeshViaMedialSkeleton, "SkinMeshViaMedialSkeleton", "SkeletonUtil", "Skeletal Mesh Conversion")
 
 private:
 	/** Mesh to add an animation skeleton to */
diff --git a/Engine/Source/Runtime/Experimental/Dataflow/Engine/Private/Dataflow/DataflowSkeleton.cpp b/Engine/Source/Runtime/Experimental/Dataflow/Engine/Private/Dataflow/DataflowSkeleton.cpp
new file mode 100644
--- /dev/null
+++ b/Engine/Source/Runtime/Experimental/Dataflow/Engine/Private/Dataflow/DataflowSkeleton.cpp
@@ -0,0 +1,31 @@
+// Copyright Epic Games, Inc. All Rights Reserved.
+
+#include "Dataflow/DataflowSkeleton.h"
+
+#include UE_INLINE_GENERATED_CPP_BY_NAME(DataflowSkeleton)
+
+const FReferenceSkeleton& FDataflowSkeleton::GetRefSkeleton() const
+{
+	return RefSkeleton;
+}
+
+FReferenceSkeletonModifier FDataflowSkeleton::ModifySkeleton()
+{
+	return FReferenceSkeletonModifier(RefSkeleton, nullptr);
+}
+
+void FDataflowSkeleton::Empty()
+{
+	RefSkeleton.Empty();
+}
+
+void FDataflowSkeleton::Copy(const FReferenceSkeleton& ToCopy)
+{
+	RefSkeleton = ToCopy;
+}
+
+bool FDataflowSkeleton::Serialize(FArchive& Ar)
+{
+	Ar << RefSkeleton;
+	return true;
+}
diff --git a/Engine/Source/Runtime/Experimental/Dataflow/Engine/Public/Dataflow/DataflowSkeleton.h b/Engine/Source/Runtime/Experimental/Dataflow/Engine/Public/Dataflow/DataflowSkeleton.h
new file mode 100644
--- /dev/null
+++ b/Engine/Source/Runtime/Experimental/Dataflow/Engine/Public/Dataflow/DataflowSkeleton.h
@@ -0,0 +1,43 @@
+// Copyright Epic Games, Inc. All Rights Reserved.
+
+#pragma once
+
+#include "ReferenceSkeleton.h"
+
+#include "DataflowSkeleton.generated.h"
+
+namespace UE::Geometry
+{
+	class FDynamicMesh3;
+}
+
+/** 
+* Represents skeletons for dataflow
+*/
+USTRUCT()
+struct FDataflowSkeleton
+{
+	GENERATED_USTRUCT_BODY()
+
+	// read-only access to the underlying skeleton object
+	DATAFLOWENGINE_API const FReferenceSkeleton& GetRefSkeleton() const;
+
+	DATAFLOWENGINE_API FReferenceSkeletonModifier ModifySkeleton();
+
+	DATAFLOWENGINE_API void Empty();
+	
+	DATAFLOWENGINE_API void Copy(const FReferenceSkeleton& ToCopy);
+
+	DATAFLOWENGINE_API bool Serialize(FArchive& Ar);
+
+private:
+	FReferenceSkeleton RefSkeleton;
+};
+
+template<> struct TStructOpsTypeTraits<FDataflowSkeleton> : public TStructOpsTypeTraitsBase2<FDataflowSkeleton>
+{
+	enum
+	{
+		WithSerializer = true,
+	};
+};
\ No newline at end of file
