guard against creating/converting a boneless/empty mesh to a skeletal mesh, since that will crash the skeletal mesh rendering code

#Synced-CL 51501523

diff --git a/Engine/Plugins/Experimental/GeometryCollectionPlugin/Source/GeometryCollectionNodes/Private/Dataflow/MeshMedialSkeletonConversionNodes.cpp b/Engine/Plugins/Experimental/GeometryCollectionPlugin/Source/GeometryCollectionNodes/Private/Dataflow/MeshMedialSkeletonConversionNodes.cpp
--- a/Engine/Plugins/Experimental/GeometryCollectionPlugin/Source/GeometryCollectionNodes/Private/Dataflow/MeshMedialSkeletonConversionNodes.cpp
+++ b/Engine/Plugins/Experimental/GeometryCollectionPlugin/Source/GeometryCollectionNodes/Private/Dataflow/MeshMedialSkeletonConversionNodes.cpp
@@ -13,7 +13,7 @@
 
 #include UE_INLINE_GENERATED_CPP_BY_NAME(MeshMedialSkeletonConversionNodes)
 
-#define LOCTEXT_NAMESPACE "ChaosClothAssetWeightMapNode"
+#define LOCTEXT_NAMESPACE "MeshMedialSkeletonConversionNodes"
 
 namespace UE::Dataflow
 {
@@ -180,6 +180,11 @@ void FBindSkeletonToMeshDataflowNode::Evaluate(UE::Dataflow::FContext& Context,
 				GetValue(Context, &VoxelResolution));
 
 			TArray<SkinBinding::FBonePoseInfo> Bones = GetBindingBoneInfoFromReferenceSkeleton(InSkeleton.GetRefSkeleton());
+			if (Bones.IsEmpty())
+			{
+				Context.Warning(LOCTEXT("BindSkeletonEmptySkelError", "BindSkeletonToMesh: Attempted to bind empty skeleton to mesh; creating a single-bone skeleton instead."), this);
+				Bones.Add(SkinBinding::FBonePoseInfo{ .LocalTransform = FTransform::Identity, .Name = FName("Root"), .ParentIndex = INDEX_NONE });
+			}
 			SkinBinding::CreateSkinWeights(DynMesh, Bones, FSkeletalMeshAttributes::DefaultSkinWeightProfileName, Settings);
 		}
 		else
diff --git a/Engine/Plugins/Experimental/GeometryCollectionPlugin/Source/GeometryCollectionNodes/Private/Dataflow/MeshToSkeletalMeshTerminalNode.cpp b/Engine/Plugins/Experimental/GeometryCollectionPlugin/Source/GeometryCollectionNodes/Private/Dataflow/MeshToSkeletalMeshTerminalNode.cpp
--- a/Engine/Plugins/Experimental/GeometryCollectionPlugin/Source/GeometryCollectionNodes/Private/Dataflow/MeshToSkeletalMeshTerminalNode.cpp
+++ b/Engine/Plugins/Experimental/GeometryCollectionPlugin/Source/GeometryCollectionNodes/Private/Dataflow/MeshToSkeletalMeshTerminalNode.cpp
@@ -26,6 +26,8 @@
 
 #include UE_INLINE_GENERATED_CPP_BY_NAME(MeshToSkeletalMeshTerminalNode)
 
+#define LOCTEXT_NAMESPACE "MeshToSkeletalMeshTerminalNode"
+
 namespace UE::Dataflow
 {
 
@@ -226,6 +228,23 @@ void FMeshToSkeletalMeshTerminalNode::SetAssetValue(TObjectPtr<UObject> Asset, U
 	const FString& InAssetPath = GetValue(Context, &SkeletalMeshAssetPath);
 	auto LogErr = [&Context, this](const FString& Msg) { Context.Error(Msg, this); };
 	bool bUsingExistingSkeletalMesh = false;
+
+
+	// validate that the mesh actually exists and has bones before trying to create/clear any assets
+	TObjectPtr<const UDynamicMesh> InMesh = GetValue(Context, &Mesh);
+	if (!InMesh)
+	{
+		Context.Error(LOCTEXT("MeshToSkeletalMeshNullMeshError", "MeshToSkeletalMeshTerminalNode: Null input mesh cannot be converted to a skeletal mesh."), this);
+		SetOutputsOnFailure();
+		return;
+	}
+	if (!InMesh->GetMeshRef().HasAttributes() || InMesh->GetMeshRef().Attributes()->GetNumBones() == 0)
+	{
+		Context.Error(LOCTEXT("MeshToSkeletalMeshBonelessMeshError", "MeshToSkeletalMeshTerminalNode: Cannot convert boneless mesh to a skeletal mesh."), this);
+			SetOutputsOnFailure();
+		return;
+	}
+
 	// We try to re-use existing skeletal mesh assets when possible
 	constexpr bool bMustReplaceExistingSkeletalMesh = false;
 	TObjectPtr<USkeletalMesh> OutSkeletalMesh = MeshToSkeletalMeshAssetLocals::MakeAssetFromPath<USkeletalMesh>(LogErr, InAssetPath, bUsingExistingSkeletalMesh, bMustReplaceExistingSkeletalMesh);
@@ -276,25 +295,20 @@ void FMeshToSkeletalMeshTerminalNode::SetAssetValue(TObjectPtr<UObject> Asset, U
 		OutSkeletalMesh->ClearFlags(RF_Transient);
 	}
 
-	// update with mesh
-	TObjectPtr<const UDynamicMesh> InMesh = GetValue(Context, &Mesh);
 
-	if (InMesh)
-	{
-		const FDynamicMesh3& DynMesh = InMesh->GetMeshRef();
-		const TArray<TObjectPtr<UMaterialInterface>>& InMaterials = GetValue(Context, &Materials);
-		// convert the collection to a mesh description including skeletal mesh attributes
-		FMeshDescription MeshDescription;
+	const FDynamicMesh3& DynMesh = InMesh->GetMeshRef();
+	const TArray<TObjectPtr<UMaterialInterface>>& InMaterials = GetValue(Context, &Materials);
+	// convert the collection to a mesh description including skeletal mesh attributes
+	FMeshDescription MeshDescription;
 
-		FSkeletalMeshAttributes Attributes(MeshDescription);
-		Attributes.Register();
+	FSkeletalMeshAttributes Attributes(MeshDescription);
+	Attributes.Register();
 
-		FConversionToMeshDescriptionOptions ConversionOptions;
-		FDynamicMeshToMeshDescription Converter(ConversionOptions);
-		Converter.Convert(&DynMesh, MeshDescription);
+	FConversionToMeshDescriptionOptions ConversionOptions;
+	FDynamicMeshToMeshDescription Converter(ConversionOptions);
+	Converter.Convert(&DynMesh, MeshDescription);
 
-		MeshToSkeletalMeshAssetLocals::MeshDescriptionToSkeletalMesh(MeshDescription, Attributes, InMaterials, OutSkeletalMesh, OutSkeleton);
-	}
+	MeshToSkeletalMeshAssetLocals::MeshDescriptionToSkeletalMesh(MeshDescription, Attributes, InMaterials, OutSkeletalMesh, OutSkeleton);
 
 	if (bUsingExistingSkeletalMesh)
 	{
@@ -344,3 +358,6 @@ bool FMeshToSkeletalMeshTerminalNode::CanDebugDrawViewMode(const FName& ViewMode
 	return ViewModeName == UE::Dataflow::FDataflowConstruction3DViewMode::Name;
 }
 #endif
+
+#undef LOCTEXT_NAMESPACE
+
