fixed crash on passing UObject* to a dataflow geometryscript node

#Synced-CL 51501523

diff --git a/Engine/Plugins/Experimental/GeometryCollectionPlugin/Source/GeometryCollectionNodes/Private/Dataflow/GeometryCollectionMeshNodes.cpp b/Engine/Plugins/Experimental/GeometryCollectionPlugin/Source/GeometryCollectionNodes/Private/Dataflow/GeometryCollectionMeshNodes.cpp
--- a/Engine/Plugins/Experimental/GeometryCollectionPlugin/Source/GeometryCollectionNodes/Private/Dataflow/GeometryCollectionMeshNodes.cpp
+++ b/Engine/Plugins/Experimental/GeometryCollectionPlugin/Source/GeometryCollectionNodes/Private/Dataflow/GeometryCollectionMeshNodes.cpp
@@ -923,6 +923,25 @@ void FMeshProcessorDataflowNodeBase::ApplyParametersToBlueprintInstance(UE::Data
 										BPProperty->CopyCompleteValue(TargetAddress, &DoubleValue);
 									}
 								}
+								else if (Desc.ValueType == EPropertyBagPropertyType::Object)
+								{
+									// Special case: BP uses raw UObject* while dataflow cache uses TObjectPtr<UObject> (same applies for array)
+									if (Desc.ContainerTypes.GetFirstContainerType() == EPropertyBagContainerType::Array)
+									{
+										const TArray<TObjectPtr<UObject>> CacheValue = Context.GetData<TArray<TObjectPtr<UObject>>>(Output->CacheKey(), Output->GetProperty(), {});
+										TArray<UObject*> RawPtrArray;
+										Algo::Copy(CacheValue, RawPtrArray);
+										void* TargetAddress = BPProperty->ContainerPtrToValuePtr<TArray<UObject*>>(MeshProcessorInstance);
+										BPProperty->CopyCompleteValue(TargetAddress, &RawPtrArray);
+									}
+									else
+									{
+										const TObjectPtr<UObject> CachedValue = Context.GetData<TObjectPtr<UObject>>(Output->CacheKey(), Output->GetProperty(), nullptr);
+										const UObject* RawPtrValue = { CachedValue.Get() };
+										void* TargetAddress = BPProperty->ContainerPtrToValuePtr<UObject*>(MeshProcessorInstance);
+										BPProperty->CopyCompleteValue(TargetAddress, &RawPtrValue);
+									}
+								}
 								else if (const void* CachedValueAddress = Context.GetUntypedData(Output->CacheKey(), nullptr))
 								{
 									void* TargetAddress = BPProperty->ContainerPtrToValuePtr<void>(MeshProcessorInstance);
