make dataflow geometryscript nodes capture the mapping from property bag's sanitized names to the original blueprint property name, and disambiguate as needed if multiple original names map to the same sanitized name

#Synced-CL 51501523

diff --git a/Engine/Plugins/Experimental/GeometryCollectionPlugin/Source/GeometryCollectionNodes/Private/Dataflow/GeometryCollectionMeshNodes.cpp b/Engine/Plugins/Experimental/GeometryCollectionPlugin/Source/GeometryCollectionNodes/Private/Dataflow/GeometryCollectionMeshNodes.cpp
--- a/Engine/Plugins/Experimental/GeometryCollectionPlugin/Source/GeometryCollectionNodes/Private/Dataflow/GeometryCollectionMeshNodes.cpp
+++ b/Engine/Plugins/Experimental/GeometryCollectionPlugin/Source/GeometryCollectionNodes/Private/Dataflow/GeometryCollectionMeshNodes.cpp
@@ -861,18 +861,33 @@ void FMeshProcessorDataflowNodeBase::RefreshConnectionsFromBlueprint()
 {
 #if WITH_EDITOR
 	PropertyBag.Reset();
+	CachedPropertyNameMap.Reset();
 	if (MeshProcessor)
 	{
 		if (UBlueprint* Blueprint = Cast<UBlueprint>(MeshProcessor->ClassGeneratedBy))
 		{
 			if (Blueprint->GeneratedClass != nullptr)
 			{
+				TSet<FName> AddedSanitizedNames;
 				for (TFieldIterator<FProperty> PropertyIt(Blueprint->GeneratedClass, EFieldIteratorFlags::IncludeSuper); PropertyIt; ++PropertyIt)
 				{
 					FProperty* Property = *PropertyIt;
 					if (!Property->HasAnyPropertyFlags(CPF_Parm) && Property->HasAllPropertyFlags(CPF_BlueprintVisible))
 					{
-						PropertyBag.AddProperty(Property->GetFName(), Property, false);
+						FName SanitizedName = FInstancedPropertyBag::SanitizePropertyName(Property->GetFName());
+						bool bRenumbered = false;
+						while (AddedSanitizedNames.Contains(SanitizedName))
+						{
+							SanitizedName.SetNumber(SanitizedName.GetNumber() + 1);
+							bRenumbered = true;
+						}
+						if (bRenumbered)
+						{
+							UE_LOG(LogGeometryCollectionMeshNodes, Warning, TEXT("Geometry Script node has multiple inputs that map to the same sanitized name; renaming input %s to %s"), *Property->GetFName().ToString(), *SanitizedName.ToString());
+						}
+						AddedSanitizedNames.Add(SanitizedName);
+						CachedPropertyNameMap.Add(SanitizedName, Property->GetFName());
+						PropertyBag.AddProperty(SanitizedName, Property, false);
 					}
 				}
 			}
@@ -893,7 +908,12 @@ void FMeshProcessorDataflowNodeBase::ApplyParametersToBlueprintInstance(UE::Data
 			TConstArrayView<FPropertyBagPropertyDesc> Descs = PropertyBag.GetPropertyBagStruct()->GetPropertyDescs();
 			for (const FPropertyBagPropertyDesc& Desc : Descs)
 			{
-				if (FProperty* BPProperty = Blueprint->GeneratedClass->FindPropertyByName(Desc.Name))
+				const FName* BPMappedName = CachedPropertyNameMap.Find(Desc.Name);
+				if (!ensureMsgf(BPMappedName, TEXT("Property %s not found in cached name map, so could not be applied to BP"), *Desc.Name.ToString()))
+				{
+					continue;
+				}
+				if (FProperty* BPProperty = Blueprint->GeneratedClass->FindPropertyByName(*BPMappedName))
 				{
 					if (BPProperty->GetClass() == Desc.CachedProperty->GetClass())
 					{
diff --git a/Engine/Plugins/Experimental/GeometryCollectionPlugin/Source/GeometryCollectionNodes/Public/Dataflow/GeometryCollectionMeshNodes.h b/Engine/Plugins/Experimental/GeometryCollectionPlugin/Source/GeometryCollectionNodes/Public/Dataflow/GeometryCollectionMeshNodes.h
--- a/Engine/Plugins/Experimental/GeometryCollectionPlugin/Source/GeometryCollectionNodes/Public/Dataflow/GeometryCollectionMeshNodes.h
+++ b/Engine/Plugins/Experimental/GeometryCollectionPlugin/Source/GeometryCollectionNodes/Public/Dataflow/GeometryCollectionMeshNodes.h
@@ -135,6 +135,9 @@ private:
 
 	UPROPERTY()
 	FInstancedPropertyBag PropertyBag;
+
+	// Cached mapping from property bag names to their source blueprint property names, which can be different due to sanitization
+	TMap<FName, FName> CachedPropertyNameMap;
 };
 
 /**
