Add SmoothBoneWeights Dataflow node

#Synced-CL 51735880

diff --git a/Engine/Plugins/Experimental/GeometryCollectionPlugin/Source/GeometryCollectionNodes/Private/Dataflow/GeometryCollectionNodesPlugin.cpp b/Engine/Plugins/Experimental/GeometryCollectionPlugin/Source/GeometryCollectionNodes/Private/Dataflow/GeometryCollectionNodesPlugin.cpp
--- a/Engine/Plugins/Experimental/GeometryCollectionPlugin/Source/GeometryCollectionNodes/Private/Dataflow/GeometryCollectionNodesPlugin.cpp
+++ b/Engine/Plugins/Experimental/GeometryCollectionPlugin/Source/GeometryCollectionNodes/Private/Dataflow/GeometryCollectionNodesPlugin.cpp
@@ -41,6 +41,7 @@
 #include "Dataflow/SetVertexColorFromVertexSelectionNode.h"
 #include "Dataflow/MeshMedialSkeletonConversionNodes.h"
 #include "Dataflow/MeshToSkeletalMeshTerminalNode.h"
+#include "Dataflow/MeshWeightEditingNodes.h"
 #include "Dataflow/DataflowNodeAndPinTypeColors.h"
 #include "Dataflow/DataflowCategoryRegistry.h"
 #include "GeometryCollection/GeometryCollectionObject.h"
@@ -76,6 +77,7 @@ void IGeometryCollectionNodesPlugin::StartupModule()
 	UE::Dataflow::RegisterGeometryStaticMeshToCollectionNodes();
 	UE::Dataflow::RegisterMeshMedialSkeletonConversionNodes();
 	UE::Dataflow::RegisterMeshToSkeletalMeshTerminalNode();
+	UE::Dataflow::RegisterMeshWeightEditingNodes();
 	DATAFLOW_NODE_REGISTER_CREATION_FACTORY(FCreateColorArrayFromFloatArrayDataflowNode);
 	DATAFLOW_NODE_REGISTER_CREATION_FACTORY(FGeometryCollectionVertexScalarToVertexIndicesNode);
 	DATAFLOW_NODE_REGISTER_CREATION_FACTORY(FSetVertexColorFromFloatArrayDataflowNode);
diff --git a/Engine/Plugins/Experimental/GeometryCollectionPlugin/Source/GeometryCollectionNodes/Private/Dataflow/MeshWeightEditingNodes.cpp b/Engine/Plugins/Experimental/GeometryCollectionPlugin/Source/GeometryCollectionNodes/Private/Dataflow/MeshWeightEditingNodes.cpp
new file mode 100644
--- /dev/null
+++ b/Engine/Plugins/Experimental/GeometryCollectionPlugin/Source/GeometryCollectionNodes/Private/Dataflow/MeshWeightEditingNodes.cpp
@@ -0,0 +1,81 @@
+// Copyright Epic Games, Inc. All Rights Reserved.
+
+#include "Dataflow/MeshWeightEditingNodes.h"
+
+#include "Operations/SmoothBoneWeights.h"
+
+#include "UDynamicMesh.h"
+
+#include UE_INLINE_GENERATED_CPP_BY_NAME(MeshWeightEditingNodes)
+
+#define LOCTEXT_NAMESPACE "MeshWeightEditingNodes"
+
+namespace UE::Dataflow
+{
+	void RegisterMeshWeightEditingNodes()
+	{
+		DATAFLOW_NODE_REGISTER_CREATION_FACTORY(FSmoothMeshBoneWeightsDataflowNode);
+	}
+}
+
+FSmoothMeshBoneWeightsDataflowNode::FSmoothMeshBoneWeightsDataflowNode(
+	const UE::Dataflow::FNodeParameters& InParam, FGuid InGuid)
+	: FDataflowNode(InParam, InGuid)
+{
+	RegisterInputConnection(&Mesh);
+	RegisterInputConnection(&Strength).SetCanHidePin(true).SetPinIsHidden(true);
+	RegisterInputConnection(&NumIterations).SetCanHidePin(true).SetPinIsHidden(true);
+	RegisterInputConnection(&WeightProfile).SetCanHidePin(true).SetPinIsHidden(true);
+
+	RegisterOutputConnection(&Mesh, &Mesh);
+}
+
+void FSmoothMeshBoneWeightsDataflowNode::Evaluate(UE::Dataflow::FContext& Context, const FDataflowOutput* Out) const
+{
+	using namespace UE::Geometry;
+
+	if (Out->IsA(&Mesh))
+	{
+		TObjectPtr<const UDynamicMesh> InMesh = GetValue(Context, &Mesh);
+		TObjectPtr<UDynamicMesh> OutMesh = NewObject<UDynamicMesh>();
+
+		if (InMesh)
+		{
+			OutMesh->SetMesh(InMesh->GetMeshRef());
+			FDynamicMesh3& DynMesh = OutMesh->GetMeshRef();
+
+			const FName UseProfile = GetValue(Context, &WeightProfile);
+			FSmoothDynamicMeshVertexSkinWeights Smoother(&DynMesh, UseProfile);
+			if (Smoother.Validate() == EOperationValidationResult::Ok)
+			{
+				// Select all vertices to smooth
+				TArray<int32> VerticesToSmooth;
+				VerticesToSmooth.Reserve(DynMesh.VertexCount());
+				for (int32 VID : DynMesh.VertexIndicesItr())
+				{
+					VerticesToSmooth.Add(VID);
+				}
+
+				const float UseStrength = FMath::Clamp(GetValue(Context, &Strength), 0.f, 1.f);
+				const int32 UseIterations = FMath::Max(1, GetValue(Context, &NumIterations));
+				constexpr double FloodFillUpToDistance = 0.0; // TODO: can expose this parameter if/when adding selection support, not needed when smoothing all vertices
+
+				Smoother.SmoothWeightsAtVerticesWithinDistance(VerticesToSmooth, UseStrength, FloodFillUpToDistance, UseIterations);
+			}
+			else
+			{
+				Context.Warning(LOCTEXT("SmoothBoneWeightsNoWeightsError",
+					"SmoothBoneWeights: Mesh has no skin weights on the specified profile; nothing to smooth."), this);
+			}
+		}
+		else
+		{
+			Context.Error(LOCTEXT("SmoothBoneWeightsNullMeshError",
+				"SmoothBoneWeights: Cannot smooth weights on a null input mesh"), this);
+		}
+
+		SetValue(Context, OutMesh, &Mesh);
+	}
+}
+
+#undef LOCTEXT_NAMESPACE
diff --git a/Engine/Plugins/Experimental/GeometryCollectionPlugin/Source/GeometryCollectionNodes/Private/Dataflow/MeshWeightEditingNodes.h b/Engine/Plugins/Experimental/GeometryCollectionPlugin/Source/GeometryCollectionNodes/Private/Dataflow/MeshWeightEditingNodes.h
new file mode 100644
--- /dev/null
+++ b/Engine/Plugins/Experimental/GeometryCollectionPlugin/Source/GeometryCollectionNodes/Private/Dataflow/MeshWeightEditingNodes.h
@@ -0,0 +1,47 @@
+// Copyright Epic Games, Inc. All Rights Reserved.
+
+#pragma once
+
+#include "CoreMinimal.h"
+#include "Dataflow/DataflowEngine.h"
+#include "Dataflow/DataflowSelection.h"
+#include "SkeletalMeshAttributes.h"
+
+#include "MeshWeightEditingNodes.generated.h"
+
+class UDynamicMesh;
+
+/** Smooth bone weights on a dynamic mesh using Laplacian weight smoothing. */
+USTRUCT(meta = (DataflowGeometryCollection))
+struct FSmoothMeshBoneWeightsDataflowNode : public FDataflowNode
+{
+	GENERATED_USTRUCT_BODY()
+	DATAFLOW_NODE_DEFINE_INTERNAL(FSmoothMeshBoneWeightsDataflowNode, "SmoothBoneWeights", "SkeletonUtil", "Relax Blend")
+
+private:
+	/** Mesh with skin weights to smooth */
+	UPROPERTY(meta = (DataflowInput, DataflowOutput, DataflowIntrinsic, DataflowPassthrough = "Mesh"))
+	TObjectPtr<UDynamicMesh> Mesh;
+
+	/** Blend factor between original and smoothed weights (0=no change, 1=fully smoothed) */
+	UPROPERTY(EditAnywhere, Category = "Options", meta = (DataflowInput, ClampMin = 0.0, ClampMax = 1.0))
+	float Strength = 0.5f;
+
+	/** Number of smoothing iterations to apply */
+	UPROPERTY(EditAnywhere, Category = "Options", meta = (DataflowInput, ClampMin = 1, UIMax = 20))
+	int32 NumIterations = 1;
+
+	/** Name of the skin weight profile to smooth. */
+	UPROPERTY(EditAnywhere, Category = "Options", meta = (DataflowInput))
+	FName WeightProfile = FSkeletalMeshAttributes::DefaultSkinWeightProfileName;
+
+	virtual void Evaluate(UE::Dataflow::FContext& Context, const FDataflowOutput* Out) const override;
+
+public:
+	FSmoothMeshBoneWeightsDataflowNode(const UE::Dataflow::FNodeParameters& InParam, FGuid InGuid = FGuid::NewGuid());
+};
+
+namespace UE::Dataflow
+{
+	void RegisterMeshWeightEditingNodes();
+}
