Dataflow:
- Added CollectionSelection nodes, they all work on any type of selection
- Deprecated old selection nodes

#Synced-CL 51839604

diff --git a/Engine/Plugins/Experimental/GeometryCollectionPlugin/Source/GeometryCollectionNodes/Private/Dataflow/GeometryCollectionSelectionNodes.cpp b/Engine/Plugins/Experimental/GeometryCollectionPlugin/Source/GeometryCollectionNodes/Private/Dataflow/GeometryCollectionSelectionNodes.cpp
--- a/Engine/Plugins/Experimental/GeometryCollectionPlugin/Source/GeometryCollectionNodes/Private/Dataflow/GeometryCollectionSelectionNodes.cpp
+++ b/Engine/Plugins/Experimental/GeometryCollectionPlugin/Source/GeometryCollectionNodes/Private/Dataflow/GeometryCollectionSelectionNodes.cpp
@@ -37,9 +37,15 @@
 #include "GeometryCollection/Facades/CollectionTransformSelectionFacade.h"
 #include "GeometryCollection/Facades/CollectionHierarchyFacade.h"
 #include "Dataflow/GeometryCollectionUtils.h"
+#include "GeometryCollection/Facades/PointsFacade.h"
+#include "Dataflow/DataflowUtils.h"
+#include "GeometryCollection/Facades/CollectionMeshFacade.h"
+#include "UDynamicMesh.h"
 
 #include UE_INLINE_GENERATED_CPP_BY_NAME(GeometryCollectionSelectionNodes)
 
+#define LOCTEXT_NAMESPACE "DataflowGeometryCollectionSelectionNodes"
+
 namespace UE::Dataflow
 {
 	namespace SelectionNodes::Private
@@ -83,16 +89,10 @@ namespace UE::Dataflow
 
 	void GeometryCollectionSelectionNodes()
 	{
-		DATAFLOW_NODE_REGISTER_CREATION_FACTORY(FCollectionTransformSelectionAllDataflowNode);
 		DATAFLOW_NODE_REGISTER_CREATION_FACTORY(FCollectionTransformSelectionInfoDataflowNode);
-		DATAFLOW_NODE_REGISTER_CREATION_FACTORY(FCollectionTransformSelectionNoneDataflowNode);
-		DATAFLOW_NODE_REGISTER_CREATION_FACTORY(FCollectionTransformSelectionRandomDataflowNode);
 		DATAFLOW_NODE_REGISTER_CREATION_FACTORY(FCollectionTransformSelectionRootDataflowNode);
-		DATAFLOW_NODE_REGISTER_CREATION_FACTORY(FCollectionTransformSelectionCustomDataflowNode);
-		DATAFLOW_NODE_REGISTER_CREATION_FACTORY(FCollectionTransformSelectionCustomDataflowNode_v2);
 		DATAFLOW_NODE_REGISTER_CREATION_FACTORY(FCollectionTransformSelectionFromIndexArrayDataflowNode);
 		DATAFLOW_NODE_REGISTER_CREATION_FACTORY(FCollectionTransformSelectionParentDataflowNode);
-		DATAFLOW_NODE_REGISTER_CREATION_FACTORY(FCollectionTransformSelectionByPercentageDataflowNode);
 		DATAFLOW_NODE_REGISTER_CREATION_FACTORY(FCollectionTransformSelectionChildrenDataflowNode);
 		DATAFLOW_NODE_REGISTER_CREATION_FACTORY(FCollectionTransformSelectionSiblingsDataflowNode);
 		DATAFLOW_NODE_REGISTER_CREATION_FACTORY(FCollectionTransformSelectionLevelDataflowNode);
@@ -103,35 +103,53 @@ namespace UE::Dataflow
 		DATAFLOW_NODE_REGISTER_CREATION_FACTORY(FCollectionTransformSelectionClusterDataflowNode_v2);
 		DATAFLOW_NODE_REGISTER_CREATION_FACTORY(FCollectionTransformSelectionBySizeDataflowNode);
 		DATAFLOW_NODE_REGISTER_CREATION_FACTORY(FCollectionTransformSelectionByVolumeDataflowNode);
-		DATAFLOW_NODE_REGISTER_CREATION_FACTORY(FCollectionTransformSelectionInBoxDataflowNode);
-		DATAFLOW_NODE_REGISTER_CREATION_FACTORY(FCollectionTransformSelectionInSphereDataflowNode);
-		DATAFLOW_NODE_REGISTER_CREATION_FACTORY(FCollectionTransformSelectionByFloatAttrDataflowNode);
-		DATAFLOW_NODE_REGISTER_CREATION_FACTORY(FCollectionVertexSelectionByBoxDataflowNode);
-		DATAFLOW_NODE_REGISTER_CREATION_FACTORY(FCollectionVertexSelectionBySphereDataflowNode);
-		DATAFLOW_NODE_REGISTER_CREATION_FACTORY(FCollectionVertexSelectionByPlaneDataflowNode);
 		DATAFLOW_NODE_REGISTER_CREATION_FACTORY(FSelectFloatArrayIndicesInRangeDataflowNode);
-		DATAFLOW_NODE_REGISTER_CREATION_FACTORY(FCollectionTransformSelectionByIntAttrDataflowNode);
-		DATAFLOW_NODE_REGISTER_CREATION_FACTORY(FCollectionVertexSelectionCustomDataflowNode);
-		DATAFLOW_NODE_REGISTER_CREATION_FACTORY(FCollectionFaceSelectionCustomDataflowNode);
-		DATAFLOW_NODE_REGISTER_CREATION_FACTORY(FCollectionSelectionConvertDataflowNode);
-		DATAFLOW_NODE_REGISTER_CREATION_FACTORY(FCollectionVertexSelectionByPercentageDataflowNode);
-		DATAFLOW_NODE_REGISTER_CREATION_FACTORY(FCollectionVertexSelectionSetOperationDataflowNode);
-		DATAFLOW_NODE_REGISTER_CREATION_FACTORY(FCollectionSelectionByAttrDataflowNode);
-		DATAFLOW_NODE_REGISTER_CREATION_FACTORY(FGeometrySelectionToVertexSelectionDataflowNode);
 		DATAFLOW_NODE_REGISTER_CREATION_FACTORY(FCollectionSelectInternalFacesDataflowNode);
 		DATAFLOW_NODE_REGISTER_CREATION_FACTORY(FCollectionSelectTransformStringDataflowNode);
 		DATAFLOW_NODE_REGISTER_CREATION_FACTORY(FCollectionSetTransformStringValueDataflowNode);
 		DATAFLOW_NODE_REGISTER_CREATION_FACTORY(FGetNumElementsInSelectionDataflowNode);
-		
-
-		// generic input nodes
-		DATAFLOW_NODE_REGISTER_CREATION_FACTORY(FCollectionSelectionSetOperationDataflowNode);
-		DATAFLOW_NODE_REGISTER_CREATION_FACTORY(FCollectionSelectionInvertDataflowNode);
 
 		// deprecated
 		DATAFLOW_NODE_REGISTER_CREATION_FACTORY(FCollectionTransformSelectionInvertDataflowNode);
 		DATAFLOW_NODE_REGISTER_CREATION_FACTORY(FCollectionFaceSelectionInvertDataflowNode);
 		DATAFLOW_NODE_REGISTER_CREATION_FACTORY(FCollectionTransformSelectionSetOperationDataflowNode);
+
+		DATAFLOW_NODE_REGISTER_CREATION_FACTORY(FCollectionTransformSelectionAllDataflowNode); //
+		DATAFLOW_NODE_REGISTER_CREATION_FACTORY(FCollectionTransformSelectionNoneDataflowNode); //
+		DATAFLOW_NODE_REGISTER_CREATION_FACTORY(FCollectionTransformSelectionRandomDataflowNode); //
+		DATAFLOW_NODE_REGISTER_CREATION_FACTORY(FCollectionTransformSelectionCustomDataflowNode); //
+		DATAFLOW_NODE_REGISTER_CREATION_FACTORY(FCollectionTransformSelectionCustomDataflowNode_v2); //
+		DATAFLOW_NODE_REGISTER_CREATION_FACTORY(FCollectionTransformSelectionByPercentageDataflowNode); //
+		DATAFLOW_NODE_REGISTER_CREATION_FACTORY(FCollectionTransformSelectionInBoxDataflowNode); //
+		DATAFLOW_NODE_REGISTER_CREATION_FACTORY(FCollectionTransformSelectionInSphereDataflowNode); //
+		DATAFLOW_NODE_REGISTER_CREATION_FACTORY(FCollectionTransformSelectionByFloatAttrDataflowNode); //
+		DATAFLOW_NODE_REGISTER_CREATION_FACTORY(FCollectionVertexSelectionByBoxDataflowNode); //
+		DATAFLOW_NODE_REGISTER_CREATION_FACTORY(FCollectionVertexSelectionBySphereDataflowNode); //
+		DATAFLOW_NODE_REGISTER_CREATION_FACTORY(FCollectionVertexSelectionByPlaneDataflowNode); //
+		DATAFLOW_NODE_REGISTER_CREATION_FACTORY(FCollectionTransformSelectionByIntAttrDataflowNode); //
+		DATAFLOW_NODE_REGISTER_CREATION_FACTORY(FCollectionVertexSelectionCustomDataflowNode); //
+		DATAFLOW_NODE_REGISTER_CREATION_FACTORY(FCollectionFaceSelectionCustomDataflowNode); //
+		DATAFLOW_NODE_REGISTER_CREATION_FACTORY(FCollectionSelectionConvertDataflowNode); //
+		DATAFLOW_NODE_REGISTER_CREATION_FACTORY(FCollectionVertexSelectionByPercentageDataflowNode); //
+		DATAFLOW_NODE_REGISTER_CREATION_FACTORY(FCollectionVertexSelectionSetOperationDataflowNode); //
+		DATAFLOW_NODE_REGISTER_CREATION_FACTORY(FCollectionSelectionByAttrDataflowNode); //
+		DATAFLOW_NODE_REGISTER_CREATION_FACTORY(FGeometrySelectionToVertexSelectionDataflowNode); //
+
+		// generic input nodes
+		DATAFLOW_NODE_REGISTER_CREATION_FACTORY(FCollectionSelectionSetOperationDataflowNode);
+		DATAFLOW_NODE_REGISTER_CREATION_FACTORY(FCollectionSelectionInvertDataflowNode);
+
+		// New nodes
+		DATAFLOW_NODE_REGISTER_CREATION_FACTORY(FCollectionSelectionCustomDataflowNode);
+		DATAFLOW_NODE_REGISTER_CREATION_FACTORY(FCollectionSelectionConvertDataflowNode_v2);
+		DATAFLOW_NODE_REGISTER_CREATION_FACTORY(FCollectionSelectionAllDataflowNode);
+		DATAFLOW_NODE_REGISTER_CREATION_FACTORY(FCollectionSelectionNoneDataflowNode);
+		DATAFLOW_NODE_REGISTER_CREATION_FACTORY(FCollectionSelectionRandomDataflowNode);
+		DATAFLOW_NODE_REGISTER_CREATION_FACTORY(FCollectionSelectionByPercentageDataflowNode);
+		DATAFLOW_NODE_REGISTER_CREATION_FACTORY(FCollectionSelectionByAttributeDataflowNode);
+		DATAFLOW_NODE_REGISTER_CREATION_FACTORY(FCollectionSelectionByPrimitiveDataflowNode);
+		DATAFLOW_NODE_REGISTER_CREATION_FACTORY(FCollectionSelectionByMeshDataflowNode);
+		DATAFLOW_NODE_REGISTER_CREATION_FACTORY(FCollectionSelectionToAttributeDataflowNode);
 	}
 }
 
@@ -490,18 +508,22 @@ void FCollectionTransformSelectionCustomDataflowNode_v2::Evaluate(UE::Dataflow::
 			const FString InBoneIndices = GetValue(Context, &BoneIndices);
 
 			TArray<int32> Indices;
-			uint32 ErrorCode = 0;
-			if (UE::Dataflow::Utils::ParseIndicesStr(InBoneIndices, Indices, ErrorCode))
+
+			using namespace UE::Dataflow::Utils;
+
+			EErrorCode ErrorCode = ParseIndicesStr(InBoneIndices, Indices);
+
+			if (ErrorCode == EErrorCode::None)
 			{
 				NewTransformSelection.SetSelected(Indices);
 			}
 			else
 			{
-				if (ErrorCode == UE::Dataflow::Utils::Error_InvalidChars)
+				if (ErrorCode == EErrorCode::InvalidChars)
 				{
 					SetError(Context, &TransformSelection, TEXT("Invalid character(s) specified in list"));
 				}
-				else if (ErrorCode == UE::Dataflow::Utils::Error_InvalidFormatInSegment)
+				else if (ErrorCode == EErrorCode::InvalidFormatInSegment)
 				{
 					SetError(Context, &TransformSelection, TEXT("Invalid format in segment"));
 				}
@@ -594,7 +616,7 @@ void FCollectionTransformSelectionByPercentageDataflowNode::Evaluate(UE::Dataflo
 
 		TArray<int32> SelectionArr = InTransformSelection.AsArray();
 
-		GeometryCollection::Facades::FCollectionTransformSelectionFacade::SelectByPercentage(SelectionArr, InPercentage, bDeterministic, InRandomSeed);
+		GeometryCollection::Facades::FCollectionTransformSelectionFacade::SelectByPercentage(SelectionArr, InPercentage, bDeterministic, (int32)InRandomSeed);
 
 		InTransformSelection.SetFromArray(SelectionArr);
 		SetValue(Context, MoveTemp(InTransformSelection), &TransformSelection);
@@ -1449,7 +1471,7 @@ void FCollectionVertexSelectionByPercentageDataflowNode::Evaluate(UE::Dataflow::
 
 		TArray<int32> SelectionArr = InVertexSelection.AsArray();
 
-		GeometryCollection::Facades::FCollectionTransformSelectionFacade::SelectByPercentage(SelectionArr, InPercentage, bDeterministic, InRandomSeed);
+		GeometryCollection::Facades::FCollectionTransformSelectionFacade::SelectByPercentage(SelectionArr, InPercentage, bDeterministic, (int32)InRandomSeed);
 
 		InVertexSelection.SetFromArray(SelectionArr);
 		SetValue(Context, MoveTemp(InVertexSelection), &VertexSelection);
@@ -1882,7 +1904,7 @@ void FCollectionSelectionSetOperationDataflowNode::Evaluate(UE::Dataflow::FConte
 		}
 		else
 		{
-			UE_LOG(LogTemp, Error, TEXT("Dataflow: CollectionSelectionSetOperationDataflowNode : Input Selections have different number of elements."));
+			SetError(Context, &Selection, LOCTEXT("CollectionSelectionSetOperationInputSelectionsMismatch", "Input selections have different number of elements.").ToString());
 		}
 
 		SetValue(Context, MoveTemp(OutSelection), &Selection);
@@ -2176,3 +2198,1020 @@ void FGetNumElementsInSelectionDataflowNode::Evaluate(UE::Dataflow::FContext& Co
 }
 
 /* ----------------------------------------------------------------------------------------------------------------------- */
+
+namespace UE::Dataflow::Selection::Private
+{
+	void ConvertIndicesStrToIntArray(const FString& InIndices, TArray<int32>& OutIndicesArray)
+	{
+		TArray<FString> IndicesArray;
+		InIndices.ParseIntoArray(IndicesArray, TEXT(" "), true);
+
+		for (FString IndexStr : IndicesArray)
+		{
+			if (IndexStr.IsNumeric())
+			{
+				OutIndicesArray.Add(FCString::Atoi(*IndexStr));
+			}
+		}
+	}
+}
+
+/* ----------------------------------------------------------------------------------------------------------------------- */
+
+FCollectionSelectionCustomDataflowNode::FCollectionSelectionCustomDataflowNode(const UE::Dataflow::FNodeParameters& InParam, FGuid InGuid)
+	: FDataflowNode(InParam, InGuid)
+{
+	RegisterInputConnection(&Collection);
+	RegisterInputConnection(&Indices);
+	RegisterOutputConnection(&Collection, &Collection);
+	RegisterOutputConnection(&Selection);
+}
+
+void FCollectionSelectionCustomDataflowNode::Evaluate(UE::Dataflow::FContext& Context, const FDataflowOutput* Out) const
+{
+	if (Out->IsA(&Selection))
+	{
+		if (IsConnected(&Collection))
+		{
+			const FManagedArrayCollection& InCollection = GetValue(Context, &Collection);
+
+			const FDataflowOutput* SelectionOutput = FindOutput(&Selection);
+			if (SelectionOutput)
+			{
+				const FName GroupName = FDataflowSelection::GetSelectionTypeFromCollection(*SelectionOutput);
+				if (GroupName != NAME_None && InCollection.HasGroup(GroupName))
+				{
+					FDataflowSelection NewSelection(GroupName);
+					NewSelection.Initialize(InCollection.NumElements(GroupName), false);
+
+					const FString InIndices = GetValue(Context, &Indices);
+
+					TArray<int32> IndicesArray;
+
+					using namespace UE::Dataflow::Utils;
+
+					EErrorCode ErrorCode = ParseIndicesStr(InIndices, IndicesArray);
+
+					if (ErrorCode == EErrorCode::None)
+					{
+						if (!NewSelection.SetSelectedWithCheck(IndicesArray))
+						{
+							SetError(Context, &Selection, LOCTEXT("CollectionSelectionCustomInvalidIndex", "Invalid index specified.").ToString());
+						}
+					}
+					else
+					{
+						if (ErrorCode == EErrorCode::InvalidChars)
+						{
+							SetError(Context, &Selection, LOCTEXT("CollectionSelectionCustomInvalidChars", "Invalid character(s) specified in list.").ToString());
+						}
+						else if (ErrorCode == EErrorCode::InvalidFormatInSegment)
+						{
+							SetError(Context, &Selection, LOCTEXT("CollectionSelectionCustomInvalidFormat", "Invalid format in segment.").ToString());
+						}
+					}
+
+					SetValue(Context, MoveTemp(NewSelection), &Selection);
+					return;
+				}
+			}
+		}
+
+		SetValue(Context, FDataflowSelection(), &Selection);
+	}
+	else if (Out->IsA(&Collection))
+	{
+		SafeForwardInput(Context, &Collection, &Collection);
+	}
+}
+
+/* ----------------------------------------------------------------------------------------------------------------------- */
+
+FCollectionSelectionConvertDataflowNode_v2::FCollectionSelectionConvertDataflowNode_v2(const UE::Dataflow::FNodeParameters& InParam, FGuid InGuid)
+	: FDataflowNode(InParam, InGuid)
+{
+	RegisterInputConnection(&Collection);
+	RegisterInputConnection(&Selection);
+	RegisterOutputConnection(&Collection, &Collection);
+	RegisterOutputConnection(&Selection, &Selection);
+}
+
+void FCollectionSelectionConvertDataflowNode_v2::Evaluate(UE::Dataflow::FContext& Context, const FDataflowOutput* Out) const
+{
+	if (Out->IsA(&Selection))
+	{
+		if (IsConnected(&Collection) && IsConnected(&Selection))
+		{
+			const FManagedArrayCollection& InCollection = GetValue(Context, &Collection);
+			GeometryCollection::Facades::FCollectionTransformSelectionFacade TransformSelectionFacade(InCollection);
+
+			const FDataflowInput* SelectionInput = FindInput(&Selection);
+			const FDataflowOutput* SelectionOutput = FindOutput(&Selection);
+			if (SelectionInput && SelectionOutput)
+			{
+				const FName InGroupName = FDataflowSelection::GetSelectionTypeFromCollection(*SelectionInput);
+				const FName OutGroupName = FDataflowSelection::GetSelectionTypeFromCollection(*SelectionOutput);
+
+				if (InGroupName != NAME_None && InCollection.HasGroup(InGroupName) &&
+					OutGroupName != NAME_None && InCollection.HasGroup(OutGroupName))
+				{
+					const FDataflowSelection& InSelection = GetValue(Context, &Selection);
+					FDataflowSelection OutSelection(OutGroupName);
+					TArray<int32> SelectionArray;
+
+					if (OutGroupName == FDataflowTransformSelection::TransformGroupName)
+					{
+						if (InGroupName == FDataflowVertexSelection::VerticesGroupName)
+						{
+							SelectionArray = TransformSelectionFacade.ConvertVertexSelectionToTransformSelection(InSelection.AsArray(), bAllElementsMustBeSelected);
+						}
+						else if (InGroupName == FDataflowFaceSelection::FacesGroupName)
+						{
+							SelectionArray = TransformSelectionFacade.ConvertFaceSelectionToTransformSelection(InSelection.AsArray(), bAllElementsMustBeSelected);
+						}
+						else if (InGroupName == FDataflowGeometrySelection::GeometryGroupName)
+						{
+							SelectionArray = TransformSelectionFacade.ConvertGeometrySelectionToTransformSelection(InSelection.AsArray());
+						}
+						else if (InGroupName == FDataflowCurveSelection::CurveGroupName)
+						{
+							SelectionArray = TransformSelectionFacade.ConvertCurveSelectionToTransformSelection(InSelection.AsArray(), bAllElementsMustBeSelected);
+						}
+						else
+						{
+							SetError(Context, &Selection, LOCTEXT("CollectionSelectionConvertInvalidConversion", "Invalid conversion specified.").ToString());
+
+						}
+					}
+					else if (OutGroupName == FDataflowFaceSelection::FacesGroupName)
+					{
+						if (InGroupName == FDataflowVertexSelection::VerticesGroupName)
+						{
+							SelectionArray = TransformSelectionFacade.ConvertVertexSelectionToFaceSelection(InSelection.AsArray(), bAllElementsMustBeSelected);
+						}
+						else if (InGroupName == FDataflowTransformSelection::TransformGroupName)
+						{
+							SelectionArray = TransformSelectionFacade.ConvertTransformSelectionToFaceSelection(InSelection.AsArray());
+						}
+						else if (InGroupName == FDataflowGeometrySelection::GeometryGroupName)
+						{
+							SelectionArray = TransformSelectionFacade.ConvertGeometrySelectionToFaceSelection(InSelection.AsArray());
+						}
+						else if (InGroupName == FDataflowCurveSelection::CurveGroupName)
+						{
+							SelectionArray = TransformSelectionFacade.ConvertCurveSelectionToFaceSelection(InSelection.AsArray());
+						}
+						else
+						{
+							SetError(Context, &Selection, LOCTEXT("CollectionSelectionConvertInvalidConversion", "Invalid conversion specified.").ToString());
+						}
+					}
+					else if (OutGroupName == FDataflowVertexSelection::VerticesGroupName)
+					{
+						if (InGroupName == FDataflowFaceSelection::FacesGroupName)
+						{
+							SelectionArray = TransformSelectionFacade.ConvertFaceSelectionToVertexSelection(InSelection.AsArray());
+						}
+						else if (InGroupName == FDataflowTransformSelection::TransformGroupName)
+						{
+							SelectionArray = TransformSelectionFacade.ConvertTransformSelectionToVertexSelection(InSelection.AsArray());
+						}
+						else if (InGroupName == FDataflowGeometrySelection::GeometryGroupName)
+						{
+							SelectionArray = TransformSelectionFacade.ConvertGeometrySelectionToVertexSelection(InSelection.AsArray());
+						}
+						else if (InGroupName == FDataflowCurveSelection::CurveGroupName)
+						{
+							SelectionArray = TransformSelectionFacade.ConvertCurveSelectionToVertexSelection(InSelection.AsArray());
+						}
+						else
+						{
+							SetError(Context, &Selection, LOCTEXT("CollectionSelectionConvertInvalidConversion", "Invalid conversion specified.").ToString());
+						}
+					}
+					else if (OutGroupName == FDataflowGeometrySelection::GeometryGroupName)
+					{
+						if (InGroupName == FDataflowFaceSelection::FacesGroupName)
+						{
+							SelectionArray = TransformSelectionFacade.ConvertFaceSelectionToGeometrySelection(InSelection.AsArray(), bAllElementsMustBeSelected);
+						}
+						else if (InGroupName == FDataflowTransformSelection::TransformGroupName)
+						{
+							SelectionArray = TransformSelectionFacade.ConvertTransformSelectionToGeometrySelection(InSelection.AsArray());
+						}
+						else if (InGroupName == FDataflowVertexSelection::VerticesGroupName)
+						{
+							SelectionArray = TransformSelectionFacade.ConvertVertexSelectionToGeometrySelection(InSelection.AsArray(), bAllElementsMustBeSelected);
+						}
+						else if (InGroupName == FDataflowCurveSelection::CurveGroupName)
+						{
+							SelectionArray = TransformSelectionFacade.ConvertCurveSelectionToGeometrySelection(InSelection.AsArray(), bAllElementsMustBeSelected);
+						}
+						else
+						{
+							SetError(Context, &Selection, LOCTEXT("CollectionSelectionConvertInvalidConversion", "Invalid conversion specified.").ToString());
+						}
+					}
+					else if (OutGroupName == FDataflowCurveSelection::CurveGroupName)
+					{
+						if (InGroupName == FDataflowFaceSelection::FacesGroupName)
+						{
+							SelectionArray = TransformSelectionFacade.ConvertFaceSelectionToCurveSelection(InSelection.AsArray(), bAllElementsMustBeSelected);
+						}
+						else if (InGroupName == FDataflowTransformSelection::TransformGroupName)
+						{
+							SelectionArray = TransformSelectionFacade.ConvertTransformSelectionToCurveSelection(InSelection.AsArray());
+						}
+						else if (InGroupName == FDataflowVertexSelection::VerticesGroupName)
+						{
+							SelectionArray = TransformSelectionFacade.ConvertVertexSelectionToCurveSelection(InSelection.AsArray(), bAllElementsMustBeSelected);
+						}
+						else if (InGroupName == FDataflowGeometrySelection::GeometryGroupName)
+						{
+							SelectionArray = TransformSelectionFacade.ConvertGeometrySelectionToCurveSelection(InSelection.AsArray());
+						}
+						else
+						{
+							SetError(Context, &Selection, LOCTEXT("CollectionSelectionConvertInvalidConversion", "Invalid conversion specified.").ToString());
+						}
+					}
+
+					OutSelection.Initialize(InCollection.NumElements(OutGroupName), false);
+
+					OutSelection.SetFromArray(SelectionArray);
+
+					SetValue(Context, MoveTemp(OutSelection), &Selection);
+					return;
+				}
+			}
+		}
+
+		SetValue(Context, FDataflowSelection(), &Selection);
+	}
+	else if (Out->IsA(&Collection))
+	{
+		SafeForwardInput(Context, &Collection, &Collection);
+	}
+}
+
+/* ----------------------------------------------------------------------------------------------------------------------- */
+
+FCollectionSelectionAllDataflowNode::FCollectionSelectionAllDataflowNode(const UE::Dataflow::FNodeParameters& InParam, FGuid InGuid)
+	: FDataflowNode(InParam, InGuid)
+{
+	RegisterInputConnection(&Collection);
+	RegisterOutputConnection(&Collection, &Collection);
+	RegisterOutputConnection(&Selection);
+}
+
+void FCollectionSelectionAllDataflowNode::Evaluate(UE::Dataflow::FContext& Context, const FDataflowOutput* Out) const
+{
+	if (Out->IsA(&Selection))
+	{
+		if (IsConnected(&Collection))
+		{
+			const FManagedArrayCollection& InCollection = GetValue(Context, &Collection);
+
+			const FDataflowOutput* SelectionOutput = FindOutput(&Selection);
+			if (SelectionOutput)
+			{
+				const FName OutGroupName = FDataflowSelection::GetSelectionTypeFromCollection(*SelectionOutput);
+				if (OutGroupName != NAME_None && InCollection.HasGroup(OutGroupName))
+				{
+					FDataflowSelection OutSelection(OutGroupName);
+					OutSelection.Initialize(InCollection.NumElements(OutGroupName), true);
+
+					SetValue(Context, MoveTemp(OutSelection), &Selection);
+					return;
+				}
+				else
+				{
+					SetError(Context, &Selection, LOCTEXT("CollectionSelectionAllGroupDoesntExist", "Selected group type doesn't exist in the collection.").ToString());
+				}
+			}
+		}
+
+		SetValue(Context, FDataflowSelection(), &Selection);
+	}
+	else if (Out->IsA(&Collection))
+	{
+		SafeForwardInput(Context, &Collection, &Collection);
+	}
+}
+
+/* ----------------------------------------------------------------------------------------------------------------------- */
+
+FCollectionSelectionNoneDataflowNode::FCollectionSelectionNoneDataflowNode(const UE::Dataflow::FNodeParameters& InParam, FGuid InGuid)
+	: FDataflowNode(InParam, InGuid)
+{
+	RegisterInputConnection(&Collection);
+	RegisterOutputConnection(&Collection, &Collection);
+	RegisterOutputConnection(&Selection);
+}
+
+void FCollectionSelectionNoneDataflowNode::Evaluate(UE::Dataflow::FContext& Context, const FDataflowOutput* Out) const
+{
+	if (Out->IsA(&Selection))
+	{
+		if (IsConnected(&Collection))
+		{
+			const FManagedArrayCollection& InCollection = GetValue(Context, &Collection);
+
+			const FDataflowOutput* SelectionOutput = FindOutput(&Selection);
+			if (SelectionOutput)
+			{
+				const FName OutGroupName = FDataflowSelection::GetSelectionTypeFromCollection(*SelectionOutput);
+				if (OutGroupName != NAME_None && InCollection.HasGroup(OutGroupName))
+				{
+					FDataflowSelection OutSelection(OutGroupName);
+					OutSelection.Initialize(InCollection.NumElements(OutGroupName), false);
+
+					SetValue(Context, MoveTemp(OutSelection), &Selection);
+					return;
+				}
+				else
+				{
+					SetError(Context, &Selection, LOCTEXT("CollectionSelectionNoneGroupDoesntExist", "Selected group type doesn't exist in the collection.").ToString());
+				}
+			}
+		}
+
+		SetValue(Context, FDataflowSelection(), &Selection);
+	}
+	else if (Out->IsA(&Collection))
+	{
+		SafeForwardInput(Context, &Collection, &Collection);
+	}
+}
+
+/* ----------------------------------------------------------------------------------------------------------------------- */
+
+FCollectionSelectionRandomDataflowNode::FCollectionSelectionRandomDataflowNode(const UE::Dataflow::FNodeParameters& InParam, FGuid InGuid)
+	: FDataflowNode(InParam, InGuid)
+{
+	RegisterInputConnection(&Collection);
+	RegisterInputConnection(&RandomSeed);
+	RegisterInputConnection(&RandomThreshold);
+	RegisterOutputConnection(&Collection, &Collection);
+	RegisterOutputConnection(&Selection);
+}
+
+void FCollectionSelectionRandomDataflowNode::Evaluate(UE::Dataflow::FContext& Context, const FDataflowOutput* Out) const
+{
+	if (Out->IsA(&Selection))
+	{
+		if (IsConnected(&Collection))
+		{
+			const FManagedArrayCollection& InCollection = GetValue(Context, &Collection);
+
+			const FDataflowOutput* SelectionOutput = FindOutput(&Selection);
+			if (SelectionOutput)
+			{
+				const FName OutGroupName = FDataflowSelection::GetSelectionTypeFromCollection(*SelectionOutput);
+				if (OutGroupName != NAME_None && InCollection.HasGroup(OutGroupName))
+				{
+					const int32 NumElems = InCollection.NumElements(OutGroupName);
+
+					FDataflowSelection OutSelection(OutGroupName);
+					OutSelection.Initialize(NumElems, false);
+
+					const int32 InRandomSeed = GetValue(Context, &RandomSeed);
+					const float InRandomThreshold = GetValue(Context, &RandomThreshold);
+
+					GeometryCollection::Facades::FCollectionTransformSelectionFacade TransformSelectionFacade(InCollection);
+					const TArray<int32>& SelectionArray = TransformSelectionFacade.SelectRandom(NumElems, bDeterministic, InRandomSeed, InRandomThreshold);
+
+					OutSelection.SetFromArray(SelectionArray);
+
+					SetValue(Context, MoveTemp(OutSelection), &Selection);
+					return;
+				}
+				else
+				{
+					SetError(Context, &Selection, LOCTEXT("CollectionSelectionRandomGroupDoesntExist", "Selected group type doesn't exist in the collection.").ToString());
+				}
+			}
+		}
+
+		SetValue(Context, FDataflowSelection(), &Selection);
+	}
+	else if (Out->IsA(&Collection))
+	{
+		SafeForwardInput(Context, &Collection, &Collection);
+	}
+}
+
+/* ----------------------------------------------------------------------------------------------------------------------- */
+
+FCollectionSelectionByPercentageDataflowNode::FCollectionSelectionByPercentageDataflowNode(const UE::Dataflow::FNodeParameters& InParam, FGuid InGuid)
+	: FDataflowNode(InParam, InGuid)
+{
+	static const FName TypeDependencyGroup("Main");
+
+	RegisterInputConnection(&Selection)
+		.SetTypeDependencyGroup(TypeDependencyGroup);
+	RegisterInputConnection(&Percentage);
+	RegisterInputConnection(&RandomSeed);
+	RegisterOutputConnection(&Selection, &Selection)
+		.SetTypeDependencyGroup(TypeDependencyGroup);
+}
+
+void FCollectionSelectionByPercentageDataflowNode::Evaluate(UE::Dataflow::FContext& Context, const FDataflowOutput* Out) const
+{
+	if (Out->IsA(&Selection))
+	{
+		const FDataflowSelection& InSelection = GetValue(Context, &Selection);
+
+		const FDataflowOutput* SelectionOutput = FindOutput(&Selection);
+		if (SelectionOutput)
+		{
+			const FName OutGroupName = FDataflowSelection::GetSelectionTypeFromCollection(*SelectionOutput);
+			if (OutGroupName != NAME_None)
+			{
+				FDataflowSelection OutSelection(OutGroupName);
+				OutSelection.Initialize(InSelection);
+
+				const int32 InPercentage = GetValue(Context, &Percentage);
+				const int32 InRandomSeed = GetValue(Context, &RandomSeed);
+
+				TArray<int32> SelectionArray = OutSelection.AsArray();
+
+				GeometryCollection::Facades::FCollectionTransformSelectionFacade::SelectByPercentage(SelectionArray, InPercentage, bDeterministic, InRandomSeed);
+
+				OutSelection.SetFromArray(SelectionArray);
+
+				SetValue(Context, MoveTemp(OutSelection), &Selection);
+				return;
+			}
+		}
+
+		SetValue(Context, FDataflowSelection(), &Selection);
+	}
+}
+
+/* ----------------------------------------------------------------------------------------------------------------------- */
+
+FCollectionSelectionByAttributeDataflowNode::FCollectionSelectionByAttributeDataflowNode(const UE::Dataflow::FNodeParameters& InParam, FGuid InGuid)
+	: FDataflowNode(InParam, InGuid)
+{
+	RegisterInputConnection(&Collection);
+	RegisterOutputConnection(&Collection, &Collection);
+	RegisterOutputConnection(&Selection);
+}
+
+void FCollectionSelectionByAttributeDataflowNode::Evaluate(UE::Dataflow::FContext& Context, const FDataflowOutput* Out) const
+{
+	if (Out->IsA(&Selection))
+	{
+		if (IsConnected(&Collection))
+		{
+			const FManagedArrayCollection& InCollection = GetValue(Context, &Collection);
+			const FName AttributeName = FName(Attribute);
+
+			const FDataflowOutput* SelectionOutput = FindOutput(&Selection);
+			if (SelectionOutput)
+			{
+				const FName OutGroupName = FDataflowSelection::GetSelectionTypeFromCollection(*SelectionOutput);
+				if (OutGroupName != NAME_None && InCollection.HasGroup(OutGroupName))
+				{
+					if (InCollection.HasAttribute(AttributeName, OutGroupName))
+					{
+						const int32 NumElems = InCollection.NumElements(OutGroupName);
+
+						FDataflowSelection OutSelection(OutGroupName);
+						OutSelection.Initialize(NumElems, false);
+
+						CreateSelectionFromAttr(InCollection,
+							OutGroupName,
+							AttributeName,
+							Value,
+							Operation,
+							OutSelection);
+
+						SetValue(Context, MoveTemp(OutSelection), &Selection);
+						return;
+					}
+				}
+			}
+		}
+
+		SetValue(Context, FDataflowSelection(), &Selection);
+	}
+	else if (Out->IsA(&Collection))
+	{
+		SafeForwardInput(Context, &Collection, &Collection);
+	}
+}
+
+/* ----------------------------------------------------------------------------------------------------------------------- */
+
+FCollectionSelectionByPrimitiveDataflowNode::FCollectionSelectionByPrimitiveDataflowNode(const UE::Dataflow::FNodeParameters& InParam, FGuid InGuid)
+	: FDataflowNode(InParam, InGuid)
+{
+	RegisterInputConnection(&Collection);
+	RegisterInputConnection(&Primitive);
+	RegisterInputConnection(&Transform);
+
+	RegisterOutputConnection(&Collection, &Collection);
+	RegisterOutputConnection(&Selection);
+}
+
+void FCollectionSelectionByPrimitiveDataflowNode::Evaluate(UE::Dataflow::FContext& Context, const FDataflowOutput* Out) const
+{
+	if (Out->IsA(&Selection))
+	{
+		if (IsConnected(&Collection))
+		{
+			const FManagedArrayCollection& InCollection = GetValue(Context, &Collection);
+			const FTransform& InTransform = GetValue(Context, &Transform);
+
+			const FDataflowOutput* SelectionOutput = FindOutput(&Selection);
+			if (SelectionOutput)
+			{
+				const FName OutGroupName = FDataflowSelection::GetSelectionTypeFromCollection(*SelectionOutput);
+
+				if (OutGroupName != NAME_None && 
+					InCollection.HasGroup(FGeometryCollection::VerticesGroup) && 
+					InCollection.HasGroup(OutGroupName))
+				{
+					GeometryCollection::Facades::FCollectionTransformSelectionFacade TransformSelectionFacade(InCollection);
+
+					const FDataflowPrimitiveTypesStorage& InPrimitive = GetValue(Context, &Primitive);
+
+					TBitArray<> VertexSelection; 
+					VertexSelection.Init(false, InCollection.NumElements(FGeometryCollection::VerticesGroup));
+					
+					if (const FBox* InBox = InPrimitive.TryGet<FBox>())
+					{
+						TransformSelectionFacade.SelectVerticesInTransformedBox(*InBox, InTransform, VertexSelection);
+					}
+					else if (const FSphere* InSphere = InPrimitive.TryGet<FSphere>())
+					{
+						FSphere TransformedSphere = InSphere->TransformBy(InTransform);
+						TransformSelectionFacade.SelectVerticesInSphere(TransformedSphere, VertexSelection);
+					}
+					else if (const FDataflowPlane* InPlane = InPrimitive.TryGet<FDataflowPlane>())
+					{
+						FDataflowPlane TransformedPlane = InPlane->GetTransformed(InTransform);
+						TransformSelectionFacade.SelectVerticesOnPlaneSide(TransformedPlane.AsPlane(), bPositiveSide, VertexSelection);
+					}
+
+					FDataflowSelection OutSelection(OutGroupName);
+					OutSelection.Initialize(InCollection.NumElements(OutGroupName), false);
+
+					TArray<int32> IntArray; IntArray.SetNumUninitialized(VertexSelection.Num());
+					for (int32 Idx = 0; Idx < VertexSelection.Num(); ++Idx)
+					{
+						IntArray[Idx] = VertexSelection[Idx] ? 1 : 0;
+					}
+
+					if (OutGroupName == FDataflowVertexSelection::VerticesGroupName)
+					{
+						OutSelection.Initialize(VertexSelection);
+					}
+					else if (OutGroupName == FDataflowFaceSelection::FacesGroupName)
+					{
+						TArray<int32> SelectionArray = TransformSelectionFacade.ConvertVertexSelectionToFaceSelection(IntArray, bAllElementsMustBeSelected);
+						OutSelection.SetFromArray(SelectionArray);
+					}
+					else if (OutGroupName == FDataflowTransformSelection::TransformGroupName)
+					{
+						TArray<int32> SelectionArray = TransformSelectionFacade.ConvertVertexSelectionToTransformSelection(IntArray, bAllElementsMustBeSelected);
+						OutSelection.SetFromArray(SelectionArray);
+					}
+					else if (OutGroupName == FDataflowGeometrySelection::GeometryGroupName)
+					{
+						TArray<int32> SelectionArray = TransformSelectionFacade.ConvertVertexSelectionToGeometrySelection(IntArray, bAllElementsMustBeSelected);
+						OutSelection.SetFromArray(SelectionArray);
+					}
+					else if (OutGroupName == FDataflowCurveSelection::CurveGroupName)
+					{
+						TBitArray<> OutCurveSelection;
+						TransformSelectionFacade.ConvertVertexSelectionToCurveSelection(VertexSelection, bAllElementsMustBeSelected, OutCurveSelection);
+						OutSelection.Initialize(OutCurveSelection);
+					}
+					else
+					{
+						SetError(Context, &Selection, LOCTEXT("CollectionSelectionByPrimitiveSelectedTypeNotSupported", "Selected type is not supported.").ToString());
+					}
+
+					SetValue(Context, MoveTemp(OutSelection), &Selection);					
+					return;
+				}
+				else
+				{
+					SetError(Context, &Selection, LOCTEXT("CollectionSelectionByPrimitiveGroupDoesntExist", "Selected group type doesn't exist in the collection.").ToString());
+				}
+			}
+		}
+
+		SetValue(Context, FDataflowSelection(), &Selection);
+	}
+	else if (Out->IsA(&Collection))
+	{
+		SafeForwardInput(Context, &Collection, &Collection);
+	}
+}
+
+#if WITH_EDITOR
+bool FCollectionSelectionByPrimitiveDataflowNode::CanDebugDrawViewMode(const FName& ViewModeName) const
+{
+	return ViewModeName == UE::Dataflow::FDataflowConstruction3DViewMode::Name;
+}
+
+void FCollectionSelectionByPrimitiveDataflowNode::DebugDraw(UE::Dataflow::FContext& Context, IDataflowDebugDrawInterface& DataflowRenderingInterface, const FDebugDrawParameters& DebugDrawParameters) const
+{
+	if ((DebugDrawParameters.bNodeIsSelected || DebugDrawParameters.bNodeIsPinned))
+	{
+		const FTransform& InTransform = GetValue(Context, &Transform);
+
+		const FDataflowPrimitiveTypesStorage& InPrimitive = GetValue(Context, &Primitive);
+		if (const FBox* InBox = InPrimitive.TryGet<FBox>())
+		{
+			UE::Dataflow::SelectionNodes::Private::DebugDrawBox(*InBox, InTransform, DataflowRenderingInterface);
+		}
+		else if (const FSphere* InSphere = InPrimitive.TryGet<FSphere>())
+		{
+			FSphere TransformedSphere = InSphere->TransformBy(InTransform);
+			UE::Dataflow::SelectionNodes::Private::DebugDrawSphere(TransformedSphere, DataflowRenderingInterface);
+		}
+		else if (const FDataflowPlane* InPlane = InPrimitive.TryGet<FDataflowPlane>())
+		{
+			FDataflowPlane TransformedPlane = InPlane->GetTransformed(InTransform);
+			UE::Dataflow::SelectionNodes::Private::DebugDrawPlane(TransformedPlane.AsTransform(), PlaneSize, VertexPerEdge, DataflowRenderingInterface);
+		}
+	}
+}
+#endif
+
+/* ----------------------------------------------------------------------------------------------------------------------- */
+
+FCollectionSelectionByMeshDataflowNode::FCollectionSelectionByMeshDataflowNode(const UE::Dataflow::FNodeParameters& InParam, FGuid InGuid)
+	: FDataflowNode(InParam, InGuid)
+{
+	RegisterInputConnection(&Collection);
+	RegisterInputConnection(&Mesh);
+	RegisterInputConnection(&bKeepInside).SetCanHidePin(true).SetPinIsHidden(true);
+	RegisterInputConnection(&WindingThreshold).SetCanHidePin(true).SetPinIsHidden(true);
+	RegisterInputConnection(&MinDistance).SetCanHidePin(true).SetPinIsHidden(true);
+	RegisterInputConnection(&MaxDistance).SetCanHidePin(true).SetPinIsHidden(true);
+	RegisterInputConnection(&bUseSignedDistance).SetCanHidePin(true).SetPinIsHidden(true);
+
+	RegisterOutputConnection(&Collection, &Collection);
+	RegisterOutputConnection(&Selection);
+}
+
+void FCollectionSelectionByMeshDataflowNode::Evaluate(UE::Dataflow::FContext& Context, const FDataflowOutput* Out) const
+{
+	if (Out->IsA(&Selection))
+	{
+		if (IsConnected(&Collection))
+		{
+			const FManagedArrayCollection& InCollection = GetValue(Context, &Collection);
+			const bool bInKeepInside = GetValue(Context, &bKeepInside);
+			const double InWindingThreshold = (double)GetValue(Context, &WindingThreshold);
+			const double InMinDistance = (double)GetValue(Context, &MinDistance);
+			const double InMaxDistance = (double)GetValue(Context, &MaxDistance);
+			const bool bInUseSignedDistance = GetValue(Context, &bUseSignedDistance);
+
+			const FDataflowOutput* SelectionOutput = FindOutput(&Selection);
+			if (SelectionOutput)
+			{
+				const FName OutGroupName = FDataflowSelection::GetSelectionTypeFromCollection(*SelectionOutput);
+
+				if (OutGroupName != NAME_None && InCollection.HasGroup(OutGroupName))
+				{
+					GeometryCollection::Facades::FCollectionTransformSelectionFacade TransformSelectionFacade(InCollection);
+
+					if (TObjectPtr<const UDynamicMesh> InMesh = GetValue(Context, &Mesh))
+					{
+						TArray<FVector3f> VerticesInCollectionSpace;
+
+						GeometryCollection::Facades::FCollectionMeshFacade MeshFacade(InCollection);
+						MeshFacade.GetVerticesInCollectionSpace(VerticesInCollectionSpace);
+
+						TArray<FVector> InSamplePoints = TArray<FVector>(VerticesInCollectionSpace);
+
+						TArray<bool> KeepPoints;
+						KeepPoints.SetNumZeroed(InSamplePoints.Num());
+
+						using namespace UE::Geometry;
+						InMesh->ProcessMesh(
+							[&InSamplePoints, &KeepPoints, InWindingThreshold, bInKeepInside, InMinDistance, InMaxDistance, bInUseSignedDistance, this]
+							(const FDynamicMesh3& DynMesh)
+							{
+								const bool bFilterWinding = bool(FilterMethod == ESelectionByMeshMethodFlags::Winding);
+								const bool bFilterMinDist = bool(FilterMethod == ESelectionByMeshMethodFlags::MinDistance);
+								const bool bFilterMaxDist = bool(FilterMethod == ESelectionByMeshMethodFlags::MaxDistance);
+
+								if (!bFilterWinding && !bFilterMinDist && !bFilterMaxDist)
+								{
+									// No filtering, early-out
+									return;
+								}
+
+								// Clamp unsigned distances to 0
+								double UseMinDistance = InMinDistance;
+								double UseMaxDistance = InMaxDistance;
+								if (!bInUseSignedDistance)
+								{
+									UseMinDistance = FMath::Max(0, UseMinDistance);
+									UseMaxDistance = FMath::Max(0, UseMaxDistance);
+								}
+
+								if (bFilterMinDist && bFilterMaxDist)
+								{
+									// If Min > Max, impossible for points to pass the filter, so early-out
+									if (UseMinDistance > UseMaxDistance)
+									{
+										InSamplePoints.Empty();
+										KeepPoints.Empty();
+										return;
+									}
+								}
+
+								const bool bNeedWinding = bFilterWinding || bInUseSignedDistance;
+								const bool bNeedDistance = bFilterMinDist || bFilterMaxDist;
+								const double MaxRelevantDistance = UE_DOUBLE_KINDA_SMALL_NUMBER + FMath::Max(
+									bFilterMinDist ? FMath::Abs(UseMinDistance) : 0,
+									bFilterMaxDist ? FMath::Abs(UseMaxDistance) : 0);
+
+								//  set up AABBTree and FWNTree lists
+								TMeshAABBTree3<FDynamicMesh3> Spatial(&DynMesh);
+								TFastWindingTree<FDynamicMesh3> FastWinding(&Spatial, bNeedWinding);
+
+								// square threshold distances, but keep their signs, so we can do threshold testing w/out sqrt
+								const double UseSignSqMinDist = FMath::CopySign(UseMinDistance * UseMinDistance, UseMinDistance);
+								const double UseSignSqMaxDist = FMath::CopySign(UseMaxDistance * UseMaxDistance, UseMaxDistance);
+
+								// Filter points 
+								ParallelFor(KeepPoints.Num(),
+									[&KeepPoints, &Spatial, &FastWinding, &InSamplePoints,
+									bFilterWinding, bFilterMinDist, bFilterMaxDist, bInUseSignedDistance, bInKeepInside,
+									bNeedWinding, InWindingThreshold, bNeedDistance, UseSignSqMinDist, UseSignSqMaxDist, MaxRelevantDistance]
+									(int32 PointIdx)
+									{
+										FVector Pt = InSamplePoints[PointIdx];
+
+										bool bWindingInside = false;
+										if (bNeedWinding)
+										{
+											bWindingInside = FastWinding.IsInside(Pt, InWindingThreshold);
+											// test if we fail the winding filter
+											if (bFilterWinding && bInKeepInside != bWindingInside)
+											{
+												KeepPoints[PointIdx] = false;
+												return;
+											}
+										}
+										double FoundDistSq = 0;
+										if (bNeedDistance)
+										{
+											const int32 FoundTID = Spatial.FindNearestTriangle(Pt, FoundDistSq, UE::Geometry::IMeshSpatial::FQueryOptions(MaxRelevantDistance));
+											if (FoundTID == INDEX_NONE)
+											{
+												// we have a max dist, lack of closest point -> we fail the max dist filter
+												// or it's inside the shape w/ signed distances -> it's too far inside, fail the min filter
+												if (bFilterMaxDist || (bInUseSignedDistance && bWindingInside))
+												{
+													KeepPoints[PointIdx] = false;
+													return;
+												}
+												else
+												{
+													// point at least passes the min distance threshold
+													FoundDistSq = UseSignSqMinDist;
+												}
+											}
+											else
+											{
+												if (bInUseSignedDistance && bWindingInside)
+												{
+													// sign the squared distance
+													FoundDistSq = -FoundDistSq;
+												}
+											}
+											if (bFilterMinDist && FoundDistSq < UseSignSqMinDist)
+											{
+												KeepPoints[PointIdx] = false;
+												return;
+											}
+											else if (bFilterMaxDist && FoundDistSq > UseSignSqMaxDist)
+											{
+												KeepPoints[PointIdx] = false;
+												return;
+											}
+										}
+
+										// passed all filters, keep the point
+										KeepPoints[PointIdx] = true;
+									}
+								);
+							}
+						);
+
+						if (KeepPoints.Num())
+						{
+							FDataflowSelection OutSelection(OutGroupName);
+							OutSelection.Initialize(InCollection.NumElements(OutGroupName), false);
+
+							if (OutGroupName == FDataflowVertexSelection::VerticesGroupName)
+							{
+								OutSelection.SetFromArray(KeepPoints);
+							}
+							else if (OutGroupName == FDataflowFaceSelection::FacesGroupName)
+							{
+								TArray<int32> SelectionArray = TransformSelectionFacade.ConvertVertexSelectionToFaceSelection(TArray<int32>(KeepPoints), bAllElementsMustBeSelected);
+								OutSelection.SetFromArray(SelectionArray);
+							}
+							else if (OutGroupName == FDataflowTransformSelection::TransformGroupName)
+							{
+								TArray<int32> SelectionArray = TransformSelectionFacade.ConvertVertexSelectionToTransformSelection(TArray<int32>(KeepPoints), bAllElementsMustBeSelected);
+								OutSelection.SetFromArray(SelectionArray);
+							}
+							else if (OutGroupName == FDataflowGeometrySelection::GeometryGroupName)
+							{
+								TArray<int32> SelectionArray = TransformSelectionFacade.ConvertVertexSelectionToGeometrySelection(TArray<int32>(KeepPoints), bAllElementsMustBeSelected);
+								OutSelection.SetFromArray(SelectionArray);
+							}
+							else if (OutGroupName == FDataflowCurveSelection::CurveGroupName)
+							{
+								TArray<int32> SelectionArray = TransformSelectionFacade.ConvertVertexSelectionToCurveSelection(TArray<int32>(KeepPoints), bAllElementsMustBeSelected);
+								OutSelection.SetFromArray(SelectionArray);
+							}
+							else
+							{
+								SetError(Context, &Selection, LOCTEXT("CollectionSelectionByMeshSelectedTypeNotSupported", "Selected type is not supported.").ToString());
+							}
+
+							SetValue(Context, MoveTemp(OutSelection), &Selection);
+							return;
+						}
+					}
+				}
+				else
+				{
+					SetError(Context, &Selection, LOCTEXT("CollectionSelectionByMeshGroupDoesntExist", "Selected group type doesn't exist in the collection.").ToString());
+				}
+			}
+		}
+
+		SetValue(Context, FDataflowSelection(), &Selection);
+	}
+	else if (Out->IsA(&Collection))
+	{
+		SafeForwardInput(Context, &Collection, &Collection);
+	}
+}
+
+/* ----------------------------------------------------------------------------------------------------------------------- */
+
+FCollectionSelectionToAttributeDataflowNode::FCollectionSelectionToAttributeDataflowNode(const UE::Dataflow::FNodeParameters& InParam, FGuid InGuid)
+	: FDataflowNode(InParam, InGuid)
+{
+	RegisterInputConnection(&Collection);
+	RegisterInputConnection(&Selection);
+
+	RegisterInputConnection(&Attribute).SetCanHidePin(true).SetPinIsHidden(true);
+	RegisterInputConnection(&DefaultValue).SetCanHidePin(true).SetPinIsHidden(true);
+	RegisterInputConnection(&Value).SetCanHidePin(true).SetPinIsHidden(true);
+
+
+	RegisterOutputConnection(&Collection, &Collection);
+	RegisterOutputConnection(&Attribute, &Attribute);
+}
+
+void FCollectionSelectionToAttributeDataflowNode::Evaluate(UE::Dataflow::FContext& Context, const FDataflowOutput* Out) const
+{
+	if (Out->IsA(&Collection))
+	{
+		if (IsConnected(&Collection) && IsConnected(&Selection))
+		{
+			FManagedArrayCollection InCollection = GetValue(Context, &Collection);
+			FDataflowSelection InSelection = GetValue(Context, &Selection);
+
+			const FDataflowInput* SelectionInput = FindInput(&Selection);
+			if (SelectionInput)
+			{
+				const FName InGroupName = FDataflowSelection::GetSelectionTypeFromCollection(*SelectionInput);
+
+				if (InGroupName != NAME_None && InCollection.HasGroup(InGroupName))
+				{
+					if (InCollection.NumElements(InGroupName) == InSelection.Num())
+					{
+						const FName AttributeName = FName(*Attribute);
+
+						if (InCollection.HasAttribute(AttributeName, InGroupName))
+						{
+							const FManagedArrayCollection::EArrayType ArrayType = InCollection.GetAttributeType(AttributeName, InGroupName);
+
+							if (ArrayType == FManagedArrayCollection::EArrayType::FInt32Type)
+							{
+								TManagedArray<int32>* Attr = InCollection.FindAttribute<int32>(AttributeName, InGroupName);
+
+								for (int32 Idx = 0; Idx < InSelection.Num(); ++Idx)
+								{
+									if (InSelection.IsSelected(Idx))
+									{
+										(*Attr)[Idx] = FCString::Atoi(*Value);
+									}
+								}
+							}
+							else if (ArrayType == FManagedArrayCollection::EArrayType::FFloatType)
+							{
+								TManagedArray<float>* Attr = InCollection.FindAttribute<float>(AttributeName, InGroupName);
+
+								for (int32 Idx = 0; Idx < InSelection.Num(); ++Idx)
+								{
+									if (InSelection.IsSelected(Idx))
+									{
+										(*Attr)[Idx] = FCString::Atof(*Value);
+									}
+								}
+							}
+							else if (ArrayType == FManagedArrayCollection::EArrayType::FBoolType)
+							{
+								TManagedArray<bool>* Attr = InCollection.FindAttribute<bool>(AttributeName, InGroupName);
+
+								for (int32 Idx = 0; Idx < InSelection.Num(); ++Idx)
+								{
+									if (InSelection.IsSelected(Idx))
+									{
+										(*Attr)[Idx] = Value.ToBool();
+									}
+								}
+							}
+						}
+						else
+						{
+							if (AttributeType == ESelectionToAttributeTypeFlags::Int)
+							{
+								TManagedArray<int32>& Attr = InCollection.AddAttribute<int32>(AttributeName, InGroupName);
+
+								for (int32 Idx = 0; Idx < InSelection.Num(); ++Idx)
+								{
+									if (InSelection.IsSelected(Idx))
+									{
+										Attr[Idx] = FCString::Atoi(*Value);
+									}
+									else
+									{
+										Attr[Idx] = FCString::Atoi(*DefaultValue);
+									}
+								}
+							}
+							else if (AttributeType == ESelectionToAttributeTypeFlags::Float)
+							{
+								TManagedArray<float>& Attr = InCollection.AddAttribute<float>(AttributeName, InGroupName);
+
+								for (int32 Idx = 0; Idx < InSelection.Num(); ++Idx)
+								{
+									if (InSelection.IsSelected(Idx))
+									{
+										Attr[Idx] = FCString::Atof(*Value);
+									}
+									else
+									{
+										Attr[Idx] = FCString::Atof(*DefaultValue);
+									}
+								}
+							}
+							else if (AttributeType == ESelectionToAttributeTypeFlags::Bool)
+							{
+								TManagedArray<bool>& Attr = InCollection.AddAttribute<bool>(AttributeName, InGroupName);
+
+								for (int32 Idx = 0; Idx < InSelection.Num(); ++Idx)
+								{
+									if (InSelection.IsSelected(Idx))
+									{
+										Attr[Idx] = Value.ToBool();
+									}
+									else
+									{
+										Attr[Idx] = DefaultValue.ToBool();
+									}
+								}
+							}
+						}
+			
+						SetValue(Context, MoveTemp(InCollection), &Collection);
+						return;
+					}
+				}
+			}
+
+			SafeForwardInput(Context, &Collection, &Collection);
+			return;
+		}
+
+		SetValue(Context, FManagedArrayCollection(), &Collection);
+	}
+	else if (Out->IsA(&Attribute))
+	{
+		SafeForwardInput(Context, &Attribute, &Attribute);
+	}
+}
+
+/* ----------------------------------------------------------------------------------------------------------------------- */
+
+#undef LOCTEXT_NAMESPACE
diff --git a/Engine/Plugins/Experimental/GeometryCollectionPlugin/Source/GeometryCollectionNodes/Public/Dataflow/GeometryCollectionSelectionNodes.h b/Engine/Plugins/Experimental/GeometryCollectionPlugin/Source/GeometryCollectionNodes/Public/Dataflow/GeometryCollectionSelectionNodes.h
--- a/Engine/Plugins/Experimental/GeometryCollectionPlugin/Source/GeometryCollectionNodes/Public/Dataflow/GeometryCollectionSelectionNodes.h
+++ b/Engine/Plugins/Experimental/GeometryCollectionPlugin/Source/GeometryCollectionNodes/Public/Dataflow/GeometryCollectionSelectionNodes.h
@@ -10,6 +10,8 @@
 #include "Math/MathFwd.h"
 #include "Math/Sphere.h"
 #include "Dataflow/DataflowConnectionTypes.h"
+#include "Dataflow/DataflowPoints.h"
+#include "Dataflow/DataflowPrimitiveTypes.h"
 
 #include "GeometryCollectionSelectionNodes.generated.h"
 
@@ -21,9 +23,9 @@ class FGeometryCollection;
 /**
  *
  * Selects all the bones for the Collection
- *
+ * DEPRECATED 5.8 - use FCollectionSelectionAllDataflowNode instead
  */
-USTRUCT(meta = (DataflowGeometryCollection))
+USTRUCT(meta = (DataflowGeometryCollection, Deprecated = "5.8"))
 struct FCollectionTransformSelectionAllDataflowNode : public FDataflowNode
 {
 	GENERATED_USTRUCT_BODY()
@@ -148,9 +150,9 @@ public:
 /**
  *
  * Generates an empty bone selection for the Collection
- *
+ * DEPRECATED 5.8 - use FCollectionSelectionNoneDataflowNode instead
  */
-USTRUCT(meta = (DataflowGeometryCollection))
+USTRUCT(meta = (DataflowGeometryCollection, Deprecated = "5.8"))
 struct FCollectionTransformSelectionNoneDataflowNode : public FDataflowNode
 {
 	GENERATED_USTRUCT_BODY()
@@ -209,9 +211,9 @@ public:
 /**
  *
  * Selects bones randomly in the Collection
- *
+ * DEPRECATED 5.8 - use FCollectionSelectionRandomDataflowNode instead
  */
-USTRUCT(meta = (DataflowGeometryCollection))
+USTRUCT(meta = (DataflowGeometryCollection, Deprecated = "5.8"))
 struct FCollectionTransformSelectionRandomDataflowNode : public FDataflowNode
 {
 	GENERATED_USTRUCT_BODY()
@@ -226,7 +228,7 @@ public:
 	UPROPERTY(EditAnywhere, Category = "Random", meta = (DataflowInput, EditCondition = "bDeterministic"))
 	float RandomSeed = 0.f;
 
-	/** Bones get selected if RandomValue > RandomThreshold */
+	/** Bones get selected if RandomValue < RandomThreshold */
 	UPROPERTY(EditAnywhere, Category = "Random", meta = (DataflowInput, UIMin = 0.f, UIMax = 1.f))
 	float RandomThreshold = 0.5f;
 
@@ -291,7 +293,7 @@ public:
  * 
  * Selects specified bones in the GeometryCollection by using a 
  * space separated list, e.g. "0 1 2 12 23"
- *
+ * DEPRECATED 5.6 - use FCollectionTransformSelectionCustomDataflowNode_v2 instead
  */
 USTRUCT(meta = (DataflowGeometryCollection, Deprecated = "5.6"))
 struct FCollectionTransformSelectionCustomDataflowNode : public FDataflowNode
@@ -329,9 +331,9 @@ public:
  *
  * Selects specified bones in the GeometryCollection by using a
  * comma separated list, e.g. "0, 2, 5-10, 12-15"
- *
+ * DEPRECATED 5.8 - use FCollectionSelectionCustomDataflowNode instead
  */
-USTRUCT(meta = (DataflowGeometryCollection))
+USTRUCT(meta = (DataflowGeometryCollection, Deprecated = "5.8"))
 struct FCollectionTransformSelectionCustomDataflowNode_v2 : public FDataflowNode
 {
 	GENERATED_USTRUCT_BODY()
@@ -431,9 +433,9 @@ public:
 /**
  *
  * Outputs the specified percentage of the selected bones
- *
+ * DEPRECATED 5.8 - use FCollectionSelectionByPercentageDataflowNode instead
  */
-USTRUCT(meta = (DataflowGeometryCollection))
+USTRUCT(meta = (DataflowGeometryCollection, Deprecated = "5.8"))
 struct FCollectionTransformSelectionByPercentageDataflowNode : public FDataflowNode
 {
 	GENERATED_USTRUCT_BODY()
@@ -941,9 +943,9 @@ enum class ESelectSubjectTypeEnum : uint8
 /**
  *
  * Selects bones if their Vertices/BoundingBox/Centroid in a box
- *
+ * DEPRECATED 5.8 - use FCollectionSelectionByPrimitiveDataflowNode instead
  */
-USTRUCT(meta = (DataflowGeometryCollection))
+USTRUCT(meta = (DataflowGeometryCollection, Deprecated = "5.8"))
 struct FCollectionTransformSelectionInBoxDataflowNode : public FDataflowNode
 {
 	GENERATED_USTRUCT_BODY()
@@ -995,9 +997,9 @@ public:
 /**
  *
  * Selects bones if their Vertices/BoundingBox/Centroid in a sphere
- *
+ * DEPRECATED 5.8 - use FCollectionSelectionByPrimitiveDataflowNode instead
  */
-USTRUCT(meta = (DataflowGeometryCollection))
+USTRUCT(meta = (DataflowGeometryCollection, Deprecated = "5.8"))
 struct FCollectionTransformSelectionInSphereDataflowNode : public FDataflowNode
 {
 	GENERATED_USTRUCT_BODY()
@@ -1050,9 +1052,9 @@ public:
 /**
  *
  * Selects vertices they are contained in a transformed box
- *
+ * DEPRECATED 5.8 - use FCollectionSelectionByPrimitiveDataflowNode instead
  */
-USTRUCT(meta = (DataflowGeometryCollection))
+USTRUCT(meta = (DataflowGeometryCollection, Deprecated = "5.8"))
 struct FCollectionVertexSelectionByBoxDataflowNode : public FDataflowNode
 {
 	GENERATED_USTRUCT_BODY()
@@ -1090,9 +1092,9 @@ private:
 /**
  *
  * Selects vertices they are contained in a sphere
- *
+ * DEPRECATED 5.8 - use FCollectionSelectionByPrimitiveDataflowNode instead
  */
-USTRUCT(meta = (DataflowGeometryCollection))
+USTRUCT(meta = (DataflowGeometryCollection, Deprecated = "5.8"))
 struct FCollectionVertexSelectionBySphereDataflowNode : public FDataflowNode
 {
 	GENERATED_USTRUCT_BODY()
@@ -1125,9 +1127,9 @@ private:
 /**
  *
  * Selects vertices they are contained on a place side 
- *
+ * DEPRECATED 5.8 - use FCollectionSelectionByPrimitiveDataflowNode instead
  */
-USTRUCT(meta = (DataflowGeometryCollection))
+USTRUCT(meta = (DataflowGeometryCollection, Deprecated = "5.8"))
 struct FCollectionVertexSelectionByPlaneDataflowNode : public FDataflowNode
 {
 	GENERATED_USTRUCT_BODY()
@@ -1173,9 +1175,9 @@ private:
 /**
  *
  * Selects bones by a float attribute
- *
+ * DEPRECATED 5.8 - use FCollectionSelectionByAttributeDataflowNode instead
  */
-USTRUCT(meta = (DataflowGeometryCollection))
+USTRUCT(meta = (DataflowGeometryCollection, Deprecated = "5.8"))
 struct FCollectionTransformSelectionByFloatAttrDataflowNode : public FDataflowNode
 {
 	GENERATED_USTRUCT_BODY()
@@ -1232,9 +1234,9 @@ public:
 /**
  *
  * Selects bones by an int attribute
- *
+ * DEPRECATED 5.8 - use FCollectionSelectionByPrimitiveDataflowNode instead
  */
-USTRUCT(meta = (DataflowGeometryCollection))
+USTRUCT(meta = (DataflowGeometryCollection, Deprecated = "5.8"))
 struct FCollectionTransformSelectionByIntAttrDataflowNode : public FDataflowNode
 {
 	GENERATED_USTRUCT_BODY()
@@ -1292,9 +1294,9 @@ public:
  *
  * Selects specified vertices in the GeometryCollection by using a
  * space separated list
- *
+ * DEPRECATED 5.8 - use FCollectionSelectionCustomDataflowNode instead
  */
-USTRUCT(meta = (DataflowGeometryCollection))
+USTRUCT(meta = (DataflowGeometryCollection, Deprecated = "5.8"))
 struct FCollectionVertexSelectionCustomDataflowNode : public FDataflowNode
 {
 	GENERATED_USTRUCT_BODY()
@@ -1331,9 +1333,9 @@ public:
  *
  * Selects specified faces in the GeometryCollection by using a
  * space separated list
- *
+ * DEPRECATED 5.8 - use FCollectionSelectionCustomDataflowNode instead
  */
-USTRUCT(meta = (DataflowGeometryCollection))
+USTRUCT(meta = (DataflowGeometryCollection, Deprecated = "5.8"))
 struct FCollectionFaceSelectionCustomDataflowNode : public FDataflowNode
 {
 	GENERATED_USTRUCT_BODY()
@@ -1369,9 +1371,9 @@ public:
 /**
  *
  * Converts Vertex/Face/Transform selection into Vertex/Face/Transform selection
- *
+ * DEPRECATED 5.8 - use FCollectionSelectionConvertDataflowNode_v2 instead
  */
-USTRUCT(meta = (DataflowGeometryCollection))
+USTRUCT(meta = (DataflowGeometryCollection, Deprecated = "5.8"))
 struct FCollectionSelectionConvertDataflowNode : public FDataflowNode
 {
 	GENERATED_USTRUCT_BODY()
@@ -1447,9 +1449,9 @@ public:
 /**
  *
  * Outputs the specified percentage of the selected vertices
- *
+ * DEPRECATED 5.8 - use FCollectionSelectionByPercentageDataflowNode instead
  */
-USTRUCT(meta = (DataflowGeometryCollection))
+USTRUCT(meta = (DataflowGeometryCollection, Deprecated = "5.8"))
 struct FCollectionVertexSelectionByPercentageDataflowNode : public FDataflowNode
 {
 	GENERATED_USTRUCT_BODY()
@@ -1572,9 +1574,9 @@ enum class ESelectionByAttrOperation : uint8
  *
  * Selects specified Vertices/Faces/Transforms in the GeometryCollection by using an attribute value
  * Currently supported attribute types: float, int32, String, bool
- *
+ * DEPRECATED 5.8 - use FCollectionSelectionByAttributeDataflowNode instead
  */
-USTRUCT(meta = (DataflowGeometryCollection))
+USTRUCT(meta = (DataflowGeometryCollection, Deprecated = "5.8"))
 struct FCollectionSelectionByAttrDataflowNode : public FDataflowNode
 {
 	GENERATED_USTRUCT_BODY()
@@ -1648,8 +1650,10 @@ public:
 
 /**
  * Converts GeometrySelection to VertexSelection
+ * DEPRECATED 5.8 - use FCollectionSelectionConvertDataflowNode_v2 instead
+ * 
  */
-USTRUCT(meta = (DataflowGeometryCollection))
+USTRUCT(meta = (DataflowGeometryCollection, Deprecated = "5.8"))
 struct FGeometrySelectionToVertexSelectionDataflowNode final : public FDataflowNode
 {
 	GENERATED_USTRUCT_BODY()
@@ -1875,11 +1879,11 @@ private:
 	FDataflowSelectionTypes Selection;
 
 	/** Number of all (selected and non-selected) elements in the Selection */
-	UPROPERTY(meta = (DataflowOutput));
+	UPROPERTY(meta = (DataflowOutput))
 	int32 NumElements = 0;
 
 	/** Number of selected elements in the Selection */
-	UPROPERTY(meta = (DataflowOutput));
+	UPROPERTY(meta = (DataflowOutput))
 	int32 NumSelectedElements = 0;
 
 	virtual void Evaluate(UE::Dataflow::FContext& Context, const FDataflowOutput* Out) const override;
@@ -1888,8 +1892,438 @@ public:
 	FGetNumElementsInSelectionDataflowNode(const UE::Dataflow::FNodeParameters& InParam, FGuid InGuid = FGuid::NewGuid());
 };
 
+///////////////////////////////////////////////////////////////////////////////////////////////////////////////////////
+
+/**
+ *
+ * Creates any type of selection from a collection by using a comma separated list
+ *
+ */
+USTRUCT(meta = (DataflowGeometryCollection))
+struct FCollectionSelectionCustomDataflowNode : public FDataflowNode
+{
+	GENERATED_USTRUCT_BODY()
+	DATAFLOW_NODE_DEFINE_INTERNAL(FCollectionSelectionCustomDataflowNode, "CollectionSelectionCustom", "GeometryCollection|Selection", "")
+
+public:
+	/** GeometryCollection for the selection */
+	UPROPERTY(meta = (DataflowInput, DataflowOutput, DataflowPassthrough = "Collection", DataflowIntrinsic))
+	FManagedArrayCollection Collection;
+
+	/** Comma separated list of indices (example: "0, 2, 5-10, 12-15") to specify the selection */
+	UPROPERTY(EditAnywhere, Category = "Selection", meta = (DataflowInput))
+	FString Indices = FString();
+
+	/** Selection from the indices. To set the selection type rt mouse click on the Selection output */
+	UPROPERTY(meta = (DataflowOutput))
+	FDataflowSelectionTypes Selection;
+
+	FCollectionSelectionCustomDataflowNode(const UE::Dataflow::FNodeParameters& InParam, FGuid InGuid = FGuid::NewGuid());
+
+private:
+	virtual void Evaluate(UE::Dataflow::FContext& Context, const FDataflowOutput* Out) const override;
+};
+
+/* ----------------------------------------------------------------------------------------------------------------------- */
+
+/**
+ *
+ * Converts Vertex/Face/Transform selection into Vertex/Face/Transform selection
+ * Supported groups: Transform/Vertices/Faces/Geometry/Curves
+ *
+ */
+USTRUCT(meta = (DataflowGeometryCollection))
+struct FCollectionSelectionConvertDataflowNode_v2 : public FDataflowNode
+{
+	GENERATED_USTRUCT_BODY()
+	DATAFLOW_NODE_DEFINE_INTERNAL(FCollectionSelectionConvertDataflowNode_v2, "CollectionSelectionConvert", "GeometryCollection|Selection", "")
+
+public:
+	/** GeometryCollection for the selection */
+	UPROPERTY(meta = (DataflowInput, DataflowOutput, DataflowPassthrough = "Collection", DataflowIntrinsic))
+	FManagedArrayCollection Collection;
+
+	/** Vertex/Face/Transform selection to convert. To set the selection type rt mouse click on the Selection output */
+	UPROPERTY(meta = (DataflowInput, DataflowOutput, DataflowPassthrough = "Selection", DataflowIntrinsic))
+	FDataflowSelectionTypes Selection;
+
+	/** If true then for converting vertex/face selection to transform selection all vertex/face must be selected for selecting the associated transform */
+	UPROPERTY(EditAnywhere, Category = "Selection")
+	bool bAllElementsMustBeSelected = false;
+
+	FCollectionSelectionConvertDataflowNode_v2(const UE::Dataflow::FNodeParameters& InParam, FGuid InGuid = FGuid::NewGuid());
+
+private:
+	virtual void Evaluate(UE::Dataflow::FContext& Context, const FDataflowOutput* Out) const override;
+};
+
+/* ----------------------------------------------------------------------------------------------------------------------- */
+
+/**
+ *
+ * Selects all the elements in the selected group from the collection
+ *
+ */
+USTRUCT(meta = (DataflowGeometryCollection))
+struct FCollectionSelectionAllDataflowNode : public FDataflowNode
+{
+	GENERATED_USTRUCT_BODY()
+	DATAFLOW_NODE_DEFINE_INTERNAL(FCollectionSelectionAllDataflowNode, "CollectionSelectionAll", "GeometryCollection|Selection", "")
+
+public:
+	/** GeometryCollection for the selection */
+	UPROPERTY(meta = (DataflowInput, DataflowOutput, DataflowPassthrough = "Collection", DataflowIntrinsic))
+	FManagedArrayCollection Collection;
+
+	/** Output selection. To set the selection type rt mouse click on the Selection output */
+	UPROPERTY(meta = (DataflowOutput))
+	FDataflowSelectionTypes Selection;
+
+	FCollectionSelectionAllDataflowNode(const UE::Dataflow::FNodeParameters& InParam, FGuid InGuid = FGuid::NewGuid());
+
+private:
+	virtual void Evaluate(UE::Dataflow::FContext& Context, const FDataflowOutput* Out) const override;
+
+};
+
+/* ----------------------------------------------------------------------------------------------------------------------- */
+
+/**
+ *
+ * Selects none of the elements in the selected group from the collection
+ *
+ */
+USTRUCT(meta = (DataflowGeometryCollection))
+struct FCollectionSelectionNoneDataflowNode : public FDataflowNode
+{
+	GENERATED_USTRUCT_BODY()
+	DATAFLOW_NODE_DEFINE_INTERNAL(FCollectionSelectionNoneDataflowNode, "CollectionSelectionNone", "GeometryCollection|Selection", "")
+
+public:
+	/** GeometryCollection for the selection */
+	UPROPERTY(meta = (DataflowInput, DataflowOutput, DataflowPassthrough = "Collection", DataflowIntrinsic))
+	FManagedArrayCollection Collection;
+
+	/** Output selection. To set the selection type rt mouse click on the Selection output */
+	UPROPERTY(meta = (DataflowOutput))
+	FDataflowSelectionTypes Selection;
+
+	FCollectionSelectionNoneDataflowNode(const UE::Dataflow::FNodeParameters& InParam, FGuid InGuid = FGuid::NewGuid());
+
+private:
+	virtual void Evaluate(UE::Dataflow::FContext& Context, const FDataflowOutput* Out) const override;
+
+};
+
+/* ----------------------------------------------------------------------------------------------------------------------- */
+
+/**
+ *
+ * Selects randomly in the collection from the specified group
+ *
+ */
+USTRUCT(meta = (DataflowGeometryCollection))
+struct FCollectionSelectionRandomDataflowNode : public FDataflowNode
+{
+	GENERATED_USTRUCT_BODY()
+	DATAFLOW_NODE_DEFINE_INTERNAL(FCollectionSelectionRandomDataflowNode, "CollectionSelectionRandom", "GeometryCollection|Selection", "")
+
+public:
+	/** GeometryCollection for the selection */
+	UPROPERTY(meta = (DataflowInput, DataflowOutput, DataflowPassthrough = "Collection", DataflowIntrinsic))
+	FManagedArrayCollection Collection;
+
+	/** If true, it always generates the same result for the same RandomSeed */
+	UPROPERTY(EditAnywhere, Category = "Random")
+	bool bDeterministic = false;
+
+	/** Seed for the random generation, only used if Deterministic is on */
+	UPROPERTY(EditAnywhere, Category = "Random", meta = (DataflowInput, EditCondition = "bDeterministic"))
+	int32 RandomSeed = 0;
+
+	/** Get selected if RandomValue < RandomThreshold */
+	UPROPERTY(EditAnywhere, Category = "Random", meta = (DataflowInput, UIMin = 0.f, UIMax = 1.f))
+	float RandomThreshold = 0.5f;
+
+	/** Output selection. To set the selection type rt mouse click on the Selection output */
+	UPROPERTY(meta = (DataflowOutput))
+	FDataflowSelectionTypes Selection;
+
+	FCollectionSelectionRandomDataflowNode(const UE::Dataflow::FNodeParameters& InParam, FGuid InGuid = FGuid::NewGuid());
+
+private:
+	virtual void Evaluate(UE::Dataflow::FContext& Context, const FDataflowOutput* Out) const override;
+
+};
+
+/* ----------------------------------------------------------------------------------------------------------------------- */
+
+/**
+ *
+ * Selects specified percentage of elements in the collection from the specified group
+ *
+ */
+USTRUCT(meta = (DataflowGeometryCollection))
+struct FCollectionSelectionByPercentageDataflowNode : public FDataflowNode
+{
+	GENERATED_USTRUCT_BODY()
+	DATAFLOW_NODE_DEFINE_INTERNAL(FCollectionSelectionByPercentageDataflowNode, "CollectionSelectionByPercentage", "GeometryCollection|Selection", "")
+
+public:
+	/** Selection to selectt from by percentage */
+	UPROPERTY(meta = (DataflowInput, DataflowOutput, DataflowPassthrough = "Selection", DataflowIntrinsic))
+	FDataflowSelectionTypes Selection;
+
+	/** Percentage to keep from the elements */
+	UPROPERTY(EditAnywhere, Category = "Selection", meta = (DataflowInput, UIMin = 0, ClampMin = 0, UIMax = 100, ClampMax = 100))
+	int32 Percentage = 100;
+
+	/** If true, it always generates the same result for the same RandomSeed */
+	UPROPERTY(EditAnywhere, Category = "Random")
+	bool bDeterministic = false;
+
+	/** Seed for the random generation, only used if Deterministic is on */
+	UPROPERTY(EditAnywhere, Category = "Random", meta = (DataflowInput, EditCondition = "bDeterministic"))
+	int32 RandomSeed = 0;
+
+	FCollectionSelectionByPercentageDataflowNode(const UE::Dataflow::FNodeParameters& InParam, FGuid InGuid = FGuid::NewGuid());
+
+private:
+	virtual void Evaluate(UE::Dataflow::FContext& Context, const FDataflowOutput* Out) const override;
+
+};
+
+/* ----------------------------------------------------------------------------------------------------------------------- */
+
+/**
+ *
+ * Selects specified Vertices/Faces/Transforms/.... in the collection by using an attribute value
+ * Currently supported attribute types: float, int32, String, bool
+ *
+ */
+USTRUCT(meta = (DataflowGeometryCollection))
+struct FCollectionSelectionByAttributeDataflowNode : public FDataflowNode
+{
+	GENERATED_USTRUCT_BODY()
+	DATAFLOW_NODE_DEFINE_INTERNAL(FCollectionSelectionByAttributeDataflowNode, "CollectionSelectionByAttribute", "GeometryCollection|Selection", "")
+
+public:
+	/** GeometryCollection for the selection */
+	UPROPERTY(meta = (DataflowInput, DataflowOutput, DataflowPassthrough = "Collection", DataflowIntrinsic))
+	FManagedArrayCollection Collection;
+
+	/** Attribute for the selection */
+	UPROPERTY(EditAnywhere, Category = "Selection")
+	FString Attribute = FString("Internal");
+
+	/** Operation */
+	UPROPERTY(EditAnywhere, Category = "Selection")
+	ESelectionByAttrOperation Operation = ESelectionByAttrOperation::Equal;
+
+	/** Attribute value for the operation */
+	UPROPERTY(EditAnywhere, Category = "Selection")
+	FString Value = FString("true");
+
+	/** Selection output. To set the selection type rt mouse click on the Selection output */
+	UPROPERTY(meta = (DataflowOutput))
+	FDataflowSelectionTypes Selection;
+
+	FCollectionSelectionByAttributeDataflowNode(const UE::Dataflow::FNodeParameters& InParam, FGuid InGuid = FGuid::NewGuid());
+
+private:
+	virtual void Evaluate(UE::Dataflow::FContext& Context, const FDataflowOutput* Out) const override;
+};
+
+/* ----------------------------------------------------------------------------------------------------------------------- */
+
+/**
+ *
+ * Selects elements which are contained in a transformed box
+ * Supported groups: Transform/Vertices/Faces/Geometry/Curves
+ *
+ */
+USTRUCT(meta = (DataflowGeometryCollection))
+struct FCollectionSelectionByPrimitiveDataflowNode : public FDataflowNode
+{
+	GENERATED_USTRUCT_BODY()
+	DATAFLOW_NODE_DEFINE_INTERNAL(FCollectionSelectionByPrimitiveDataflowNode, "CollectionSelectionByPrimitive", "GeometryCollection|Selection", "In BoundingBox Points Transformed")
+
+private:
+	/** GeometryCollection for the selection */
+	UPROPERTY(meta = (DataflowInput, DataflowOutput, DataflowPassthrough = "Collection", DataflowIntrinsic))
+	FManagedArrayCollection Collection;
+
+	/** Primitive used for the selection (Box, Sphere or Plane) */
+	UPROPERTY(EditAnywhere, Category = "Selection", meta = (DataflowInput))
+	FDataflowPrimitiveTypes Primitive;
+
+	/** Transform for the primitive */
+	UPROPERTY(EditAnywhere, Category = "Selection", meta = (DataflowInput))
+	FTransform Transform;
+
+	/** If true then all elements must contained by the primitive (all vertices of a face, all vertices of a curve, etc.) */
+	UPROPERTY(EditAnywhere, Category = "Selection")
+	bool bAllElementsMustBeSelected = false;
+
+	/** Select elements on positive side of the plane. Only used for selection by plane. */
+	UPROPERTY(EditAnywhere, Category = "Select")
+	bool bPositiveSide = true;
+
+	/** DebugDraw / Plane size */
+	UPROPERTY(EditAnywhere, Category = "Debug Draw|Plane", meta = (UIMin = 1, ClampMin = 1))
+	float PlaneSize = 250.f;
+
+	/** DebugDraw / Vertex per edge */
+	UPROPERTY(EditAnywhere, Category = "Debug Draw|Plane", meta = (UIMin = 2, ClampMin = 2))
+	int32 VertexPerEdge = 5;
+
+	/** Output selection */
+	UPROPERTY(meta = (DataflowOutput))
+	FDataflowSelectionTypes Selection;
+
+	virtual void Evaluate(UE::Dataflow::FContext& Context, const FDataflowOutput* Out) const override;
+
+#if WITH_EDITOR
+	virtual bool CanDebugDrawViewMode(const FName& ViewModeName) const override;
+	virtual void DebugDraw(UE::Dataflow::FContext& Context, IDataflowDebugDrawInterface& DataflowRenderingInterface, const FDebugDrawParameters& DebugDrawParameters) const override;
+#endif //WITH_EDITOR
+
+public:
+	FCollectionSelectionByPrimitiveDataflowNode(const UE::Dataflow::FNodeParameters& InParam, FGuid InGuid = FGuid::NewGuid());
+};
+
+/* ----------------------------------------------------------------------------------------------------------------------- */
+
+class UDynamicMesh;
+
+/** Flags to control which mesh point filtering method(s) are applied */
+UENUM(BlueprintType)
+enum class ESelectionByMeshMethodFlags : uint8
+{
+	None UMETA(DisplayName = "None"),
+	Winding UMETA(DisplayName = "Winding"),
+	MinDistance UMETA(DisplayName = "MinDistance"),
+	MaxDistance UMETA(DisplayName = "MaxDistance")
+};
+/**
+ *
+ * Collection selection bassed on if the element is inside or outside of a given mesh
+ *
+ */
+USTRUCT(meta = (DataflowGeometryCollection))
+struct FCollectionSelectionByMeshDataflowNode : public FDataflowNode
+{
+	GENERATED_USTRUCT_BODY()
+	DATAFLOW_NODE_DEFINE_INTERNAL(FCollectionSelectionByMeshDataflowNode, "CollectionSelectionByMesh", "PointSampling", "")
+
+public:
+	FCollectionSelectionByMeshDataflowNode(const UE::Dataflow::FNodeParameters& InParam, FGuid InGuid = FGuid::NewGuid());
+
+private:
+	/** GeometryCollection for the selection */
+	UPROPERTY(meta = (DataflowInput, DataflowOutput, DataflowPassthrough = "Collection", DataflowIntrinsic))
+	FManagedArrayCollection Collection;
+
+	/** Mesh to use for the selection */
+	UPROPERTY(meta = (DataflowInput, DataflowIntrinsic))
+	TObjectPtr<UDynamicMesh> Mesh;
+
+	UPROPERTY(EditAnywhere, Category = Options)
+	ESelectionByMeshMethodFlags FilterMethod = ESelectionByMeshMethodFlags::Winding;
+
+	/** Whether to keep the points inside or (if false) outside the mesh, when filtering by Winding Number. */
+	UPROPERTY(EditAnywhere, Category = WindingNumber, meta = (DataflowInput, EditCondition = "FilterMethod == ESelectionByMeshMethodFlags::Winding"))
+	bool bKeepInside = true;
+
+	/** The winding number threshold to use for determining whether a point is inside or outside of the mesh, if corresponding Filter Method is set  */
+	UPROPERTY(EditAnywhere, Category = WindingNumber, meta = (DataflowInput))
+	float WindingThreshold = .5f;
+
+	/** The min distance to surface to keep, if corresponding Filter Method is set  */
+	UPROPERTY(EditAnywhere, Category = Distance, meta = (DataflowInput, EditCondition = "FilterMethod == ESelectionByMeshMethodFlags::MinDistance"))
+	float MinDistance = 0.f;
+
+	/** The max distance to surface to keep, if corresponding Filter Method is set */
+	UPROPERTY(EditAnywhere, Category = Distance, meta = (DataflowInput, EditCondition = "FilterMethod == ESelectionByMeshMethodFlags::MaxDistance"))
+	float MaxDistance = 1000.f;
+
+	/**
+	 * Whether to use signed distances for the Min and Max Distance thresholds. Otherwise, unsigned distance is used.
+	 * Note: Signs are computed via the Winding Number. The sign is negative if the point's Winding Number is below the Winding Threshold.
+	 */
+	UPROPERTY(EditAnywhere, Category = Distance, meta = (DataflowInput))
+	bool bUseSignedDistance = false;
+
+	/** If true then for converting vertex/face selection to transform selection all vertex/face must be selected for selecting the associated transform */
+	UPROPERTY(EditAnywhere, Category = "Selection")
+	bool bAllElementsMustBeSelected = false;
+
+	/** Output selection */
+	UPROPERTY(meta = (DataflowOutput))
+	FDataflowSelectionTypes Selection;
+
+	virtual void Evaluate(UE::Dataflow::FContext& Context, const FDataflowOutput* Out) const override;
+};
+
+/* ----------------------------------------------------------------------------------------------------------------------- */
+
+/** Attribute type for SelectionToAttribute */
+UENUM(BlueprintType)
+enum class ESelectionToAttributeTypeFlags : uint8
+{
+	Int UMETA(DisplayName = "Int"),
+	Float UMETA(DisplayName = "Float"),
+	Bool UMETA(DisplayName = "Bool")
+};
+
+/**
+ *
+ * Creates specified attribute on the incoming selection's group and sets specified value to selected elements,
+ * the nonselected elements will be set to default valu. 
+ * If the attribute already exists then sets specified value to selected elements.
+ * Currently supported attribute types: float, int32, bool
+ *
+ */
+USTRUCT(meta = (DataflowGeometryCollection))
+struct FCollectionSelectionToAttributeDataflowNode : public FDataflowNode
+{
+	GENERATED_USTRUCT_BODY()
+	DATAFLOW_NODE_DEFINE_INTERNAL(FCollectionSelectionToAttributeDataflowNode, "CollectionSelectionToAttribute", "GeometryCollection|Selection", "")
+
+public:
+	/** GeometryCollection input */
+	UPROPERTY(meta = (DataflowInput, DataflowOutput, DataflowPassthrough = "Collection", DataflowIntrinsic))
+	FManagedArrayCollection Collection;
+
+	/** Selection input. The type of this selection defines the group where the data will be stored */
+	UPROPERTY(meta = (DataflowInput, DataflowIntrinsic))
+	FDataflowSelectionTypes Selection;
+
+	/** Attribute type for storing the selection */
+	UPROPERTY(EditAnywhere, Category = "Selection")
+	ESelectionToAttributeTypeFlags AttributeType = ESelectionToAttributeTypeFlags::Int;
+
+	/** Attribute for storing the selection */
+	UPROPERTY(EditAnywhere, Category = "Selection", meta = (DataflowInput, DataflowOutput, DataflowPassthrough = "Attribute"))
+	FString Attribute = FString("Dummy");
+
+	/** Attribute default value */
+	UPROPERTY(EditAnywhere, Category = "Selection", meta = (DataflowInput))
+	FString DefaultValue = FString("0");
+
+	/** Attribute value */
+	UPROPERTY(EditAnywhere, Category = "Selection", meta = (DataflowInput))
+	FString Value = FString("1");
+
+	FCollectionSelectionToAttributeDataflowNode(const UE::Dataflow::FNodeParameters& InParam, FGuid InGuid = FGuid::NewGuid());
+
+private:
+	virtual void Evaluate(UE::Dataflow::FContext& Context, const FDataflowOutput* Out) const override;
+};
+
+/* ----------------------------------------------------------------------------------------------------------------------- */
+
 namespace UE::Dataflow
 {
 	void GeometryCollectionSelectionNodes();
 }
-
diff --git a/Engine/Source/Runtime/Experimental/Chaos/Private/GeometryCollection/Facades/CollectionMeshFacade.cpp b/Engine/Source/Runtime/Experimental/Chaos/Private/GeometryCollection/Facades/CollectionMeshFacade.cpp
--- a/Engine/Source/Runtime/Experimental/Chaos/Private/GeometryCollection/Facades/CollectionMeshFacade.cpp
+++ b/Engine/Source/Runtime/Experimental/Chaos/Private/GeometryCollection/Facades/CollectionMeshFacade.cpp
@@ -7,6 +7,7 @@
 #include "GeometryCollection/Facades/CollectionMeshFacade.h"
 #include "GeometryCollection/GeometryCollection.h"
 #include "GeometryCollection/TransformCollection.h"
+#include "GeometryCollection/GeometryCollectionAlgo.h"
 
 namespace GeometryCollection::Facades
 {
@@ -29,6 +30,8 @@ namespace GeometryCollection::Facades
 		, InternalAttribute(InCollection, "Internal", FGeometryCollection::FacesGroup)
 		, FaceStartAttribute(InCollection, "FaceStart", FGeometryCollection::GeometryGroup)
 		, FaceCountAttribute(InCollection, "FaceCount", FGeometryCollection::GeometryGroup)
+		, ParentAttribute(InCollection, FTransformCollection::ParentAttribute, FTransformCollection::TransformGroup)
+		, TransformAttribute(InCollection, FTransformCollection::TransformAttribute, FTransformCollection::TransformGroup)
 	{
 	}
 
@@ -51,6 +54,8 @@ namespace GeometryCollection::Facades
 		, InternalAttribute(InCollection, "Internal", FGeometryCollection::FacesGroup)
 		, FaceStartAttribute(InCollection, "FaceStart", FGeometryCollection::GeometryGroup)
 		, FaceCountAttribute(InCollection, "FaceCount", FGeometryCollection::GeometryGroup)
+		, ParentAttribute(InCollection, FTransformCollection::ParentAttribute, FTransformCollection::TransformGroup)
+		, TransformAttribute(InCollection, FTransformCollection::TransformAttribute, FTransformCollection::TransformGroup)
 	{
 	}
 
@@ -74,6 +79,8 @@ namespace GeometryCollection::Facades
 			&& InternalAttribute.IsValid()
 			&& FaceStartAttribute.IsValid()
 			&& FaceCountAttribute.IsValid()
+			&& ParentAttribute.IsValid()
+			&& TransformAttribute.IsValid()
 			;
 	}
 
@@ -97,6 +104,8 @@ namespace GeometryCollection::Facades
 		InternalAttribute.Add();
 		FaceStartAttribute.Add();
 		FaceCountAttribute.Add();
+		ParentAttribute.Add();
+		TransformAttribute.Add();
 	}
 
 	const TArray<int32> FCollectionMeshFacade::GetVertexIndices(int32 BoneIdx) const
@@ -141,6 +150,55 @@ namespace GeometryCollection::Facades
 		return {};
 	}
 
+	void FCollectionMeshFacade::GetVerticesInCollectionSpace(TArray<FVector3f>& OutPositions) const
+	{
+		if (IsValid())
+		{
+			const TManagedArray<FVector3f>& Vertices = VertexAttribute.Get();
+			const TManagedArray<int32>& Parents = ParentAttribute.Get();
+			const TManagedArray<FTransform3f>& Transforms = TransformAttribute.Get();
+			const TManagedArray<int32>& VertexStarts = VertexStartAttribute.Get();
+			const TManagedArray<int32>& VertexCounts = VertexCountAttribute.Get();
+			const TManagedArray<int32>& TransformIndex = TransformIndexAttribute.Get();
+
+			OutPositions.Reset(Vertices.Num());
+			OutPositions.SetNumUninitialized(Vertices.Num());
+
+			TArray<FTransform3f> RootSpaceTransforms;
+			int32 Idx = 0;
+
+			GeometryCollectionAlgo::GlobalMatrices(Transforms, Parents, RootSpaceTransforms);
+
+			for (int32 GeometryIdx = 0; GeometryIdx < TransformIndex.Num(); ++GeometryIdx)
+			{
+				if (TransformIndex.IsValidIndex(GeometryIdx) &&
+					VertexStarts.IsValidIndex(GeometryIdx) &&
+					VertexCounts.IsValidIndex(GeometryIdx))
+				{
+					const int32 TransformIdx = TransformIndex[GeometryIdx];
+
+					if (RootSpaceTransforms.IsValidIndex(TransformIdx))
+					{
+						const int32 GlobalVertexOffset = VertexStarts[GeometryIdx];
+						for (int32 LocalVertexIdx = 0; LocalVertexIdx < VertexCounts[GeometryIdx]; ++LocalVertexIdx)
+						{
+							if (Idx < OutPositions.Num())
+							{
+								OutPositions[Idx++] = (RootSpaceTransforms[TransformIdx].TransformPosition((Vertices[GlobalVertexOffset + LocalVertexIdx])));
+							}
+
+						}
+					}
+				}
+			}
+
+			if (Idx < OutPositions.Num())
+			{
+				OutPositions.SetNum(Idx);
+			}
+		}
+	}
+
 	const TArrayView<const FIntVector> FCollectionMeshFacade::GetTriangles(int32 BoneIdx) const
 	{
 		const TManagedArray<int32>& TransformToGeometryIndicies = TransformToGeometryIndexAttribute.Get();
diff --git a/Engine/Source/Runtime/Experimental/Chaos/Private/GeometryCollection/Facades/CollectionTransformSelectionFacade.cpp b/Engine/Source/Runtime/Experimental/Chaos/Private/GeometryCollection/Facades/CollectionTransformSelectionFacade.cpp
--- a/Engine/Source/Runtime/Experimental/Chaos/Private/GeometryCollection/Facades/CollectionTransformSelectionFacade.cpp
+++ b/Engine/Source/Runtime/Experimental/Chaos/Private/GeometryCollection/Facades/CollectionTransformSelectionFacade.cpp
@@ -378,7 +378,7 @@ namespace GeometryCollection::Facades
 		InOutSelection = NewSelection;
 	}
 
-	TArray<int32> FCollectionTransformSelectionFacade::SelectRandom(bool bDeterministic, float RandomSeed, float RandomThresholdVal) const
+	TArray<int32> FCollectionTransformSelectionFacade::SelectRandom(bool bDeterministic, int32 RandomSeed, float RandomThresholdVal) const
 	{
 		TArray<int32> OutSelection;
 
@@ -387,7 +387,7 @@ namespace GeometryCollection::Facades
 			const TManagedArray<int32>& Parents = ParentAttribute.Get();
 			const int32 NumTransforms = ParentAttribute.Num();
 
-			FRandomStream Stream((int32)RandomSeed);
+			FRandomStream Stream(RandomSeed);
 
 			for (int32 Idx = 0; Idx < NumTransforms; ++Idx)
 			{
@@ -395,7 +395,7 @@ namespace GeometryCollection::Facades
 				{
 					float RandomVal = (float)Stream.FRandRange(0.f, 1.f);
 
-					if (RandomVal > RandomThresholdVal)
+					if (RandomVal < RandomThresholdVal)
 					{
 						OutSelection.Add(Idx);
 					}
@@ -404,7 +404,7 @@ namespace GeometryCollection::Facades
 				{
 					float RandomVal = FMath::FRandRange(0.f, 1.f);
 
-					if (RandomVal > RandomThresholdVal)
+					if (RandomVal < RandomThresholdVal)
 					{
 						OutSelection.Add(Idx);
 					}
@@ -415,6 +415,37 @@ namespace GeometryCollection::Facades
 		return OutSelection;
 	}
 
+	TArray<int32> FCollectionTransformSelectionFacade::SelectRandom(const int32 NumElems, const bool bDeterministic, const int32 RandomSeed, const float RandomThresholdVal) const
+	{
+		TArray<int32> OutSelection;
+
+		FRandomStream Stream(RandomSeed);
+
+		for (int32 Idx = 0; Idx < NumElems; ++Idx)
+		{
+			if (bDeterministic)
+			{
+				float RandomVal = (float)Stream.FRandRange(0.f, 1.f);
+
+				if (RandomVal < RandomThresholdVal)
+				{
+					OutSelection.Add(Idx);
+				}
+			}
+			else
+			{
+				float RandomVal = FMath::FRandRange(0.f, 1.f);
+
+				if (RandomVal < RandomThresholdVal)
+				{
+					OutSelection.Add(Idx);
+				}
+			}
+		}
+
+		return OutSelection;
+	}
+
 	TArray<int32> FCollectionTransformSelectionFacade::SelectLeaf() const
 	{
 		TArray<int32> OutSelection;
@@ -585,9 +616,9 @@ namespace GeometryCollection::Facades
 		InOutSelection = NewSelection.Array();
 	}
 
-	static void RandomShuffleArray(TArray<int32>& InArray, bool Deterministic, float RandomSeed)
+	static void RandomShuffleArray(TArray<int32>& InArray, bool Deterministic, int32 RandomSeed)
 	{
-		FRandomStream Stream((int32)RandomSeed);
+		FRandomStream Stream(RandomSeed);
 
 		int32 LastIndex = InArray.Num() - 1;
 
@@ -618,10 +649,15 @@ namespace GeometryCollection::Facades
 		}
 	}
 
-	void FCollectionTransformSelectionFacade::SelectByPercentage(TArray<int32>& InOutSelection, int32 InPercentage, bool Deterministic, float RandomSeed)
+	void FCollectionTransformSelectionFacade::SelectByPercentage(TArray<int32>& InOutSelection, int32 InPercentage, bool Deterministic, int32 RandomSeed)
 	{
 		RandomShuffleArray(InOutSelection, Deterministic, RandomSeed);
 
+		if (!FMath::IsWithinInclusive<int32>(InPercentage, 0, 100))
+		{
+			InPercentage = 100;
+		}
+
 		float Percentage = (float)InPercentage * 0.01f;
 		int32 NewNumElements = FMath::RoundToInt32((float)InOutSelection.Num() * Percentage);
 		
@@ -771,6 +807,8 @@ namespace GeometryCollection::Facades
 		return SiblingGroups;
 	}
 
+	/* --------------------------------------------------------------------------------------------------------------------------- */
+
 	void FCollectionTransformSelectionFacade::SelectVerticesInTransformedBox(const FBox& InBox, const FTransform& InBoxTransform, TArray<bool>& OutVertexSelection) const
 	{
 		const FCollectionMeshFacade MeshFacade(ConstCollection);
@@ -807,6 +845,47 @@ namespace GeometryCollection::Facades
 		}
 	}
 	
+	/* --------------------------------------------------------------------------------------------------------------------------- */
+
+	void FCollectionTransformSelectionFacade::SelectVerticesInTransformedBox(const FBox& InBox, const FTransform& InBoxTransform, TBitArray<>& OutVertexSelection) const
+	{
+		const FCollectionMeshFacade MeshFacade(ConstCollection);
+		const FCollectionTransformFacade TransformFacade(ConstCollection);
+
+		OutVertexSelection.Reset();
+		if (!MeshFacade.IsValid() || !TransformFacade.IsValid())
+		{
+			return;
+		}
+		const int32 NumVertices = MeshFacade.VertexAttribute.Num();
+		OutVertexSelection.Reserve(NumVertices);
+		OutVertexSelection.SetNum(NumVertices, false);
+
+		const TArray<FTransform> CollectionSpaceTransforms = TransformFacade.ComputeCollectionSpaceTransforms();
+
+		const FVector BoxCenter = InBox.GetCenter();
+		const FVector BoxExtentsScaled = InBox.GetExtent() * InBoxTransform.GetScale3D();
+		const FBox ScaledBoxCenterAtZero(-BoxExtentsScaled, BoxExtentsScaled);
+
+		for (int32 VertexIdx = 0; VertexIdx < NumVertices; ++VertexIdx)
+		{
+			const int32 TransformIdx = MeshFacade.BoneMapAttribute.IsValidIndex(VertexIdx) ? MeshFacade.BoneMapAttribute[VertexIdx] : INDEX_NONE;
+
+			bool bIsSelected = false;
+			if (CollectionSpaceTransforms.IsValidIndex(TransformIdx))
+			{
+				const FVector VertexInBoneSpace = FVector(MeshFacade.VertexAttribute[VertexIdx]);
+				const FVector VertexInCollectionSpace = CollectionSpaceTransforms[TransformIdx].TransformPosition(VertexInBoneSpace);
+				const FVector VertexInBoxSpace = InBoxTransform.InverseTransformPositionNoScale(VertexInCollectionSpace - BoxCenter);
+
+				bIsSelected = ScaledBoxCenterAtZero.IsInside(VertexInBoxSpace);
+			}
+			OutVertexSelection[VertexIdx] = bIsSelected;
+		}
+	}
+
+	/* --------------------------------------------------------------------------------------------------------------------------- */
+
 	void FCollectionTransformSelectionFacade::SelectVerticesInSphere(const FSphere& InSphere, TArray<bool>& OutVertexSelection) const
 	{
 		const FCollectionMeshFacade MeshFacade(ConstCollection);
@@ -838,6 +917,42 @@ namespace GeometryCollection::Facades
 		}
 	}
 
+	/* --------------------------------------------------------------------------------------------------------------------------- */
+
+	void FCollectionTransformSelectionFacade::SelectVerticesInSphere(const FSphere& InSphere, TBitArray<>& OutVertexSelection) const
+	{
+		const FCollectionMeshFacade MeshFacade(ConstCollection);
+		const FCollectionTransformFacade TransformFacade(ConstCollection);
+
+		OutVertexSelection.Reset();
+		if (!MeshFacade.IsValid() || !TransformFacade.IsValid())
+		{
+			return;
+		}
+		const int32 NumVertices = MeshFacade.VertexAttribute.Num();
+		OutVertexSelection.Reserve(NumVertices);
+		OutVertexSelection.SetNum(NumVertices, false);
+
+		const TArray<FTransform> CollectionSpaceTransforms = TransformFacade.ComputeCollectionSpaceTransforms();
+
+		for (int32 VertexIdx = 0; VertexIdx < NumVertices; ++VertexIdx)
+		{
+			const int32 TransformIdx = MeshFacade.BoneMapAttribute.IsValidIndex(VertexIdx) ? MeshFacade.BoneMapAttribute[VertexIdx] : INDEX_NONE;
+
+			bool bIsSelected = false;
+			if (CollectionSpaceTransforms.IsValidIndex(TransformIdx))
+			{
+				const FVector VertexInBoneSpace = FVector(MeshFacade.VertexAttribute[VertexIdx]);
+				const FVector VertexInCollectionSpace = CollectionSpaceTransforms[TransformIdx].TransformPosition(VertexInBoneSpace);
+
+				bIsSelected = InSphere.IsInside(VertexInCollectionSpace);
+			}
+			OutVertexSelection[VertexIdx] = bIsSelected;
+		}
+	}
+
+	/* --------------------------------------------------------------------------------------------------------------------------- */
+
 	void FCollectionTransformSelectionFacade::SelectVerticesOnPlaneSide(const FPlane& InPlane, bool bUsePositiveSide, TArray<bool>& OutVertexSelection) const
 	{
 		const FCollectionMeshFacade MeshFacade(ConstCollection);
@@ -872,6 +987,45 @@ namespace GeometryCollection::Facades
 		}
 	}
 
+	/* --------------------------------------------------------------------------------------------------------------------------- */
+
+	void FCollectionTransformSelectionFacade::SelectVerticesOnPlaneSide(const FPlane& InPlane, bool bUsePositiveSide, TBitArray<>& OutVertexSelection) const
+	{
+		const FCollectionMeshFacade MeshFacade(ConstCollection);
+		const FCollectionTransformFacade TransformFacade(ConstCollection);
+
+		OutVertexSelection.Reset();
+		if (!MeshFacade.IsValid() || !TransformFacade.IsValid())
+		{
+			return;
+		}
+		const int32 NumVertices = MeshFacade.VertexAttribute.Num();
+		OutVertexSelection.Reserve(NumVertices);
+		OutVertexSelection.SetNum(NumVertices, false);
+
+		const TArray<FTransform> CollectionSpaceTransforms = TransformFacade.ComputeCollectionSpaceTransforms();
+
+		for (int32 VertexIdx = 0; VertexIdx < NumVertices; ++VertexIdx)
+		{
+			const int32 TransformIdx = MeshFacade.BoneMapAttribute.IsValidIndex(VertexIdx) ? MeshFacade.BoneMapAttribute[VertexIdx] : INDEX_NONE;
+
+			bool bIsSelected = false;
+			if (CollectionSpaceTransforms.IsValidIndex(TransformIdx))
+			{
+				const FVector VertexInBoneSpace = FVector(MeshFacade.VertexAttribute[VertexIdx]);
+				const FVector VertexInCollectionSpace = CollectionSpaceTransforms[TransformIdx].TransformPosition(VertexInBoneSpace);
+
+				const double DistanceFromPlane = InPlane.PlaneDot(VertexInCollectionSpace);
+				bIsSelected = bUsePositiveSide
+					? (DistanceFromPlane >= 0)
+					: (DistanceFromPlane <= 0);
+			}
+			OutVertexSelection[VertexIdx] = bIsSelected;
+		}
+	}
+
+	/* --------------------------------------------------------------------------------------------------------------------------- */
+
 	TArray<int32> FCollectionTransformSelectionFacade::SelectVerticesInBox(const FBox& InBox, const FTransform& InBoxTransform, bool bAllVerticesInBox) const
 	{
 		TArray<int32> OutSelection;
@@ -1413,6 +1567,31 @@ namespace GeometryCollection::Facades
 		return OutSelection;
 	}
 
+	void FCollectionTransformSelectionFacade::GetVertexToFaceIndicesArray(TArray<TArray<int32>>& VertexToFaceIndices) const
+	{
+		if (IndicesAttribute.IsValid())
+		{
+			// Build VertexToFaceIndices inverse lookup
+			const int32 NumFaces = ConstCollection.NumElements(FGeometryCollection::FacesGroup);
+			const int32 NumVertices = ConstCollection.NumElements(FGeometryCollection::VerticesGroup);
+
+			VertexToFaceIndices.Empty();
+			VertexToFaceIndices.SetNum(NumVertices);
+
+			const TManagedArray<FIntVector>& Indices = IndicesAttribute.Get();
+
+			for (int32 FaceIdx = 0; FaceIdx < NumFaces; ++FaceIdx)
+			{
+				if (Indices.IsValidIndex(FaceIdx))
+				{
+					VertexToFaceIndices[Indices[FaceIdx].X].AddUnique(FaceIdx);
+					VertexToFaceIndices[Indices[FaceIdx].Y].AddUnique(FaceIdx);
+					VertexToFaceIndices[Indices[FaceIdx].Z].AddUnique(FaceIdx);
+				}
+			}
+		}
+	}
+
 	TArray<int32> FCollectionTransformSelectionFacade::ConvertTransformSelectionToFaceSelection(const TArray<int32>& InTransformSelection) const
 	{
 		TArray<int32> OutSelection;
diff --git a/Engine/Source/Runtime/Experimental/Chaos/Public/GeometryCollection/Facades/CollectionMeshFacade.h b/Engine/Source/Runtime/Experimental/Chaos/Public/GeometryCollection/Facades/CollectionMeshFacade.h
--- a/Engine/Source/Runtime/Experimental/Chaos/Public/GeometryCollection/Facades/CollectionMeshFacade.h
+++ b/Engine/Source/Runtime/Experimental/Chaos/Public/GeometryCollection/Facades/CollectionMeshFacade.h
@@ -37,21 +37,25 @@ namespace GeometryCollection::Facades
 		 */
 		CHAOS_API const TArrayView<const FVector3f> GetVertexPositions(int32 BoneIdx) const;
 
+		/**
+		 * Returns the vertex positions in collection space
+		 */
+		CHAOS_API void GetVerticesInCollectionSpace(TArray<FVector3f>& OutPositions) const;
+
 		/**
 		* Returns the triangles for the bone 
 		 */
 		CHAOS_API const TArrayView<const FIntVector> GetTriangles(int32 BoneIdx) const;
 
-
 		/**
 		 * Returns the face indices for the bone
 		 */
-		CHAOS_API const TArray<int32> GetFaceIndices(int32 BoneIdx) const;
+		CHAOS_API const TArray<int32> GetFaceIndices(int32 BoneIdx) const;	
 
 		/**
 		 * Returns the vertex indices of the face for the bone
 		 */
-//		const TArrayView<FIntVector> GetFaceVertexIndices(int32 BoneIdx) const;
+//		const TArray<FIntVector> GetFaceVertexIndices(int32 BoneIdx) const;
 
 		/**
 		 * Bakes the transforms into the vertex positions and sets the bone transforms to identity
@@ -77,6 +81,7 @@ namespace GeometryCollection::Facades
 		TManagedArrayAccessor<bool> InternalAttribute;
 		TManagedArrayAccessor<int32> FaceStartAttribute;
 		TManagedArrayAccessor<int32> FaceCountAttribute;
+		TManagedArrayAccessor<int32> ParentAttribute;
+		TManagedArrayAccessor<FTransform3f> TransformAttribute;
 	};
-
 }
diff --git a/Engine/Source/Runtime/Experimental/Chaos/Public/GeometryCollection/Facades/CollectionTransformSelectionFacade.h b/Engine/Source/Runtime/Experimental/Chaos/Public/GeometryCollection/Facades/CollectionTransformSelectionFacade.h
--- a/Engine/Source/Runtime/Experimental/Chaos/Public/GeometryCollection/Facades/CollectionTransformSelectionFacade.h
+++ b/Engine/Source/Runtime/Experimental/Chaos/Public/GeometryCollection/Facades/CollectionTransformSelectionFacade.h
@@ -90,7 +90,8 @@ namespace GeometryCollection::Facades
 
 		/**  */
 		bool CanSelectRandom() const { return ParentAttribute.IsValid(); }
-		CHAOS_API TArray<int32> SelectRandom(bool bDeterministic, float RandomSeed, float RandomThresholdVal) const;
+		CHAOS_API TArray<int32> SelectRandom(bool bDeterministic, int32 RandomSeed, float RandomThresholdVal) const;
+		CHAOS_API TArray<int32> SelectRandom(const int32 NumElems, const bool bDeterministic, const int32 RandomSeed, const float RandomThresholdVal) const;
 
 		/**  */
 		bool CanSelectLeaf() const { return SimulationTypeAttribute.IsValid(); }
@@ -121,7 +122,7 @@ namespace GeometryCollection::Facades
 		CHAOS_API void SelectLevel(TArray<int32>& InOutSelection) const;
 
 		/**  */
-		static CHAOS_API void SelectByPercentage(TArray<int32>& InOutSelection, int32 Percentage, bool Deterministic, float RandomSeed);
+		static CHAOS_API void SelectByPercentage(TArray<int32>& InOutSelection, int32 Percentage, bool Deterministic, int32 RandomSeed);
 
 		/** Select transforms by their Size or Relative Size */
 		CHAOS_API TArray<int32> SelectBySize(float SizeMin, float SizeMax, bool bInclusive, bool bInsideRange, bool bUseRelativeSize = true) const;
@@ -137,18 +138,21 @@ namespace GeometryCollection::Facades
 		* The transformed box must be in collection space
 		*/
 		CHAOS_API void SelectVerticesInTransformedBox(const FBox& InBox, const FTransform& InBoxTransform, TArray<bool>& OutVertexSelection) const;
+		CHAOS_API void SelectVerticesInTransformedBox(const FBox& InBox, const FTransform& InBoxTransform, TBitArray<>& OutVertexSelection) const;
 
 		/**
 		* Fill an array of bool that contains whether or not a vertex is contained in a sphere
 		* The sphere must be in collection space
 		*/
 		CHAOS_API void SelectVerticesInSphere(const FSphere& InSphere, TArray<bool>& OutVertexSelection) const;
+		CHAOS_API void SelectVerticesInSphere(const FSphere& InSphere, TBitArray<>& OutVertexSelection) const;
 
 		/**
 		* Fill an array of bool that contains whether or not a vertex is on a specific side of a plane 
 		* The plane must be defined in collection space
 		*/
 		CHAOS_API void SelectVerticesOnPlaneSide(const FPlane& InPlane, bool bUsePositiveSide, TArray<bool>& OutVertexSelection) const;
+		CHAOS_API void SelectVerticesOnPlaneSide(const FPlane& InPlane, bool bUsePositiveSide, TBitArray<>& OutVertexSelection) const;
 
 		/** 
 		* Returns indices of bones where all the vertices or at least one vertex of the bone is inside InBox 
@@ -302,6 +306,10 @@ namespace GeometryCollection::Facades
 		*/
 		CHAOS_API TArray<int32> ConvertCurveSelectionToTransformSelection(const TArray<int32>& InCurveSelection, bool bAllElementsMustBeSelected) const;
 
+		/**
+		* Compute and return VertexToFaceIndices array
+		*/
+		CHAOS_API void GetVertexToFaceIndicesArray(TArray<TArray<int32>>& VertexToFaceIndices) const;
 
 	private:
 		// This a const facade
diff --git a/Engine/Source/Runtime/Experimental/Dataflow/Core/Private/Dataflow/DataflowPlane.cpp b/Engine/Source/Runtime/Experimental/Dataflow/Core/Private/Dataflow/DataflowPlane.cpp
new file mode 100644
--- /dev/null
+++ b/Engine/Source/Runtime/Experimental/Dataflow/Core/Private/Dataflow/DataflowPlane.cpp
@@ -0,0 +1,60 @@
+// Copyright Epic Games, Inc. All Rights Reserved.
+
+#include "Dataflow/DataflowPlane.h"
+#include "Kismet/KismetMathLibrary.h"
+
+#include UE_INLINE_GENERATED_CPP_BY_NAME(DataflowPlane)
+
+FDataflowPlane::FDataflowPlane(const FPlane& InPlane)
+{
+	Normal = InPlane.GetNormal();
+	Point = InPlane.GetOrigin();
+}
+
+FPlane FDataflowPlane::AsPlane() const
+{
+	return FPlane(Point, Normal.GetSafeNormal());
+}
+
+FTransform FDataflowPlane::AsTransform() const
+{
+	if (!Normal.IsNearlyZero() && Normal.IsNormalized())
+	{
+		const FQuat Quat = FQuat::FindBetweenVectors(FVector::UpVector, Normal);
+		return FTransform(Quat, Point);
+	}
+	else
+	{
+		return FTransform::Identity;
+	}
+
+}
+
+double FDataflowPlane::DistanceFromPoint(const FVector& InPoint) const
+{
+	return AsPlane().PlaneDot(InPoint);
+}
+
+FDataflowPlane FDataflowPlane::GetTransformed(const FTransform& InTransform) const
+{
+	const FVector TransformedPoint = InTransform.TransformPosition(Point);
+
+	FMatrix TransformMatrix = InTransform.ToMatrixWithScale();
+	const float MulBy = TransformMatrix.Determinant() < 0.f ? -1.f : 1.f;
+
+	FMatrix NormalsTransform = TransformMatrix.TransposeAdjoint() * (double)MulBy;
+	const FVector TransformedNormal = NormalsTransform.TransformVector(Normal).GetSafeNormal();
+
+	return FDataflowPlane(TransformedPoint, TransformedNormal);
+}
+
+
+
+
+
+
+
+
+
+
+
diff --git a/Engine/Source/Runtime/Experimental/Dataflow/Core/Private/Dataflow/DataflowPrimitiveTypes.cpp b/Engine/Source/Runtime/Experimental/Dataflow/Core/Private/Dataflow/DataflowPrimitiveTypes.cpp
new file mode 100644
--- /dev/null
+++ b/Engine/Source/Runtime/Experimental/Dataflow/Core/Private/Dataflow/DataflowPrimitiveTypes.cpp
@@ -0,0 +1,26 @@
+// Copyright Epic Games, Inc. All Rights Reserved.
+
+#include "Dataflow/DataflowPrimitiveTypes.h"
+#include "Dataflow/DataflowAnyTypeRegistry.h"
+
+#include "Dataflow/DataflowNodeFactory.h"
+#include "Dataflow/DataflowCategoryRegistry.h"
+#include "Dataflow/DataflowNodeColorsRegistry.h"
+
+namespace UE::Dataflow
+{
+	void RegisterPrimitiveTypes()
+	{
+		UE_DATAFLOW_REGISTER_ANYTYPE(FDataflowPrimitiveTypes);
+		UE_DATAFLOW_REGISTER_ANYTYPE(FDataflowPrimitiveArrayTypes);
+
+		static const FLinearColor CDefaultNodeBodyTintColor = FLinearColor(0.f, 0.f, 0.f, 0.5f);
+		static const float CDefaultWireThickness = 2.f;
+
+		DATAFLOW_NODE_REGISTER_CREATION_FACTORY_PIN_SETTINGS_BY_TYPE("FDataflowPrimitiveTypes", FLinearColor(0.6f, 0.f, 0.6f), CDefaultWireThickness);
+	}
+};
+
+
+
+
diff --git a/Engine/Source/Runtime/Experimental/Dataflow/Core/Private/Dataflow/DataflowSelection.cpp b/Engine/Source/Runtime/Experimental/Dataflow/Core/Private/Dataflow/DataflowSelection.cpp
--- a/Engine/Source/Runtime/Experimental/Dataflow/Core/Private/Dataflow/DataflowSelection.cpp
+++ b/Engine/Source/Runtime/Experimental/Dataflow/Core/Private/Dataflow/DataflowSelection.cpp
@@ -3,6 +3,7 @@
 #include "Dataflow/DataflowSelection.h"
 #include "GeometryCollection/GeometryCollection.h"
 #include "GeometryCollection/ManagedArrayCollection.h"
+#include "Dataflow/DataflowConnection.h"
 
 #include UE_INLINE_GENERATED_CPP_BY_NAME(DataflowSelection)
 
@@ -64,11 +65,6 @@ TArray<int32> FDataflowSelection::AsArrayValidated(const FManagedArrayCollection
 	return SelectionArr;
 }
 
-void FDataflowSelection::SetGroupName(const FName& InGroupName)
-{
-	GroupName = InGroupName;
-}
-
 void FDataflowSelection::AsArray(TArray<int32>& SelectionArr) const
 {
 	SelectionArr.Reset(SelectionArray.Num());
@@ -258,3 +254,37 @@ FString FDataflowSelection::ToString()
 
 	return OutString;
 }
+
+FName FDataflowSelection::GetSelectionTypeFromCollection(const FDataflowConnection& InConnection)
+{
+	if (InConnection.IsType<FDataflowTransformSelection>())
+	{
+		return 	FDataflowTransformSelection::TransformGroupName;
+	}
+	else if (InConnection.IsType<FDataflowVertexSelection>())
+	{
+		return 	FDataflowVertexSelection::VerticesGroupName;
+	}
+	else if (InConnection.IsType<FDataflowFaceSelection>())
+	{
+		return 	FDataflowFaceSelection::FacesGroupName;
+	}
+	else if (InConnection.IsType<FDataflowGeometrySelection>())
+	{
+		return 	FDataflowGeometrySelection::GeometryGroupName;
+	}
+	else if (InConnection.IsType<FDataflowCurveSelection>())
+	{
+		return 	FDataflowCurveSelection::CurveGroupName;
+	}
+	else if (InConnection.IsType<FDataflowPointsSelection>())
+	{
+		return 	FDataflowPointsSelection::PointsGroupName;
+	}
+	else if (InConnection.IsType<FDataflowMaterialSelection>())
+	{
+		return 	FDataflowMaterialSelection::MaterialGroupName;
+	}
+
+	return NAME_None;
+}
diff --git a/Engine/Source/Runtime/Experimental/Dataflow/Core/Private/Dataflow/DataflowUtils.cpp b/Engine/Source/Runtime/Experimental/Dataflow/Core/Private/Dataflow/DataflowUtils.cpp
new file mode 100644
--- /dev/null
+++ b/Engine/Source/Runtime/Experimental/Dataflow/Core/Private/Dataflow/DataflowUtils.cpp
@@ -0,0 +1,83 @@
+// Copyright Epic Games, Inc. All Rights Reserved.
+
+#include "Dataflow/DataflowUtils.h"
+
+#include "GeometryCollection/GeometryCollectionAlgo.h"
+#include "Internationalization/Regex.h"
+
+namespace UE::Dataflow
+{
+	namespace Utils
+	{
+		/* e.g. "0, 2, 5-10, 12-15" */
+		EErrorCode ParseIndicesStr(const FString& InFramesString, TArray<int32>& OutIndices)
+		{
+			OutIndices.Empty();
+
+			static const FRegexPattern AllowedCharsPattern(TEXT("^[-,0-9\\s]+$"));
+
+			if (!FRegexMatcher(AllowedCharsPattern, InFramesString).FindNext())
+			{
+				// Input contains invalid characters
+				return EErrorCode::InvalidChars;
+			}
+
+			static const FRegexPattern SingleNumberPattern(TEXT("^\\s*(\\d+)\\s*$"));
+			static const FRegexPattern RangePattern(TEXT("^\\s*(\\d+)\\s*-\\s*(\\d+)\\s*$"));
+
+			TArray<FString> Segments;
+			InFramesString.ParseIntoArray(Segments, TEXT(","), true);
+			for (const FString& Segment : Segments)
+			{
+				bool bSegmentValid = false;
+
+				FRegexMatcher SingleNumberMatcher(SingleNumberPattern, Segment);
+				if (SingleNumberMatcher.FindNext())
+				{
+					const int32 SingleNumber = FCString::Atoi(*SingleNumberMatcher.GetCaptureGroup(1));
+
+					if (SingleNumber < 0 || SingleNumber > 100000)
+					{
+						return EErrorCode::NumberInRangeTooBig;
+					}
+
+					OutIndices.Add(SingleNumber);
+					bSegmentValid = true;
+				}
+				else
+				{
+					FRegexMatcher RangeMatcher(RangePattern, Segment);
+					if (RangeMatcher.FindNext())
+					{
+						const int32 RangeStart = FCString::Atoi(*RangeMatcher.GetCaptureGroup(1));
+						const int32 RangeEnd = FCString::Atoi(*RangeMatcher.GetCaptureGroup(2));
+
+						if (RangeStart > RangeEnd)
+						{
+							return EErrorCode::InvalidRangeStartEnd;
+						}
+
+						if (RangeStart < 0 || RangeStart > 100000 || RangeEnd < 0 || RangeEnd > 100000)
+						{
+							return EErrorCode::NumberInRangeTooBig;
+						}
+
+						for (int32 i = RangeStart; i <= RangeEnd; ++i)
+						{
+							OutIndices.Add(i);
+						}
+						bSegmentValid = true;
+					}
+				}
+
+				if (!bSegmentValid)
+				{
+					// Invalid format in segment
+					return EErrorCode::InvalidFormatInSegment;
+				}
+			}
+
+			return EErrorCode::None;
+		}
+	}
+}
diff --git a/Engine/Source/Runtime/Experimental/Dataflow/Core/Public/Dataflow/DataflowPlane.h b/Engine/Source/Runtime/Experimental/Dataflow/Core/Public/Dataflow/DataflowPlane.h
new file mode 100644
--- /dev/null
+++ b/Engine/Source/Runtime/Experimental/Dataflow/Core/Public/Dataflow/DataflowPlane.h
@@ -0,0 +1,47 @@
+// Copyright Epic Games, Inc. All Rights Reserved.
+
+#pragma once
+
+#include "Math/MathFwd.h"
+#include "UObject/Object.h"
+
+#include "DataflowPlane.generated.h"
+
+/** 
+* Represents a plane for dataflow
+* This uses a point and a normal to represent a plane
+*/
+USTRUCT()
+struct FDataflowPlane
+{
+	GENERATED_USTRUCT_BODY()
+
+	FDataflowPlane() :
+		Point(FVector::ZeroVector),
+		Normal(FVector::UpVector)
+	{}
+
+	FDataflowPlane(const FVector& InPoint, const FVector& InNormal) :
+		Point(InPoint),
+		Normal(InNormal.GetSafeNormal())
+	{}
+
+	DATAFLOWCORE_API FDataflowPlane(const FPlane& InPlane);
+
+	FVector GetPoint() const { return Point; }
+	FVector GetNormal() const { return Normal; }
+
+	DATAFLOWCORE_API FPlane AsPlane() const;
+	DATAFLOWCORE_API FTransform AsTransform() const;
+
+	DATAFLOWCORE_API double DistanceFromPoint(const FVector& InPoint) const;
+
+	DATAFLOWCORE_API FDataflowPlane GetTransformed(const FTransform& InTransform) const;
+
+private:
+	UPROPERTY(EditAnywhere, Category = "Plane")
+	FVector Point;
+
+	UPROPERTY(EditAnywhere, Category = "Plane")
+	FVector Normal;
+};
diff --git a/Engine/Source/Runtime/Experimental/Dataflow/Core/Public/Dataflow/DataflowPrimitiveTypes.h b/Engine/Source/Runtime/Experimental/Dataflow/Core/Public/Dataflow/DataflowPrimitiveTypes.h
new file mode 100644
--- /dev/null
+++ b/Engine/Source/Runtime/Experimental/Dataflow/Core/Public/Dataflow/DataflowPrimitiveTypes.h
@@ -0,0 +1,165 @@
+// Copyright Epic Games, Inc. All Rights Reserved.
+
+#pragma once 
+
+#include "Dataflow/DataflowAnyType.h"
+#include "StructUtils/InstancedStruct.h"
+#include "Misc/TVariant.h"
+#include "Dataflow/DataflowPlane.h"
+
+#include "DataflowPrimitiveTypes.generated.h"
+
+USTRUCT()
+struct FDataflowPrimitiveTypesStorage
+{
+	GENERATED_BODY()
+
+	FDataflowPrimitiveTypesStorage() = default;
+	FDataflowPrimitiveTypesStorage(const FDataflowPrimitiveTypesStorage& Other) = default;
+	FDataflowPrimitiveTypesStorage& operator=(const FDataflowPrimitiveTypesStorage& Other) = default;
+	FDataflowPrimitiveTypesStorage(FDataflowPrimitiveTypesStorage&& Other) = default;
+	FDataflowPrimitiveTypesStorage& operator=(FDataflowPrimitiveTypesStorage&& Other) = default;
+	FDataflowPrimitiveTypesStorage(const FBox& Primitive)
+		: Storage(TInPlaceType<FBox>(), Primitive)
+	{
+	}
+	FDataflowPrimitiveTypesStorage(const FSphere& Primitive)
+		: Storage(TInPlaceType<FSphere>(), Primitive)
+	{
+	}
+	FDataflowPrimitiveTypesStorage(const FDataflowPlane& Primitive)
+		: Storage(TInPlaceType<FDataflowPlane>(), Primitive)
+	{
+	}
+
+	template<typename T>
+	const T* TryGet() const
+	{
+		return Storage.TryGet<T>();
+	}
+	
+private:
+	// Currently we don't support serialization for this storage
+	TVariant<FBox, FSphere, FDataflowPlane> Storage;
+};
+
+template <>
+struct FDataflowConverter<FDataflowPrimitiveTypesStorage>
+{
+	template <typename TFromType>
+	static void From(const TFromType& From, FDataflowPrimitiveTypesStorage& To)
+	{
+		To = TFromType(From);
+	}
+	template <typename TToType>
+	static void To(const FDataflowPrimitiveTypesStorage& From, TToType& To)
+	{
+		if (const TToType* Ptr = From.TryGet<TToType>())
+		{
+			To = *Ptr;
+		}
+	}
+};
+
+/**
+* Primitive types 
+*/
+UE_DATAFLOW_POLICY_DECLARE_TYPENAME(FDataflowPlane)
+UE_DATAFLOW_POLICY_DECLARE_TYPENAME(FDataflowPrimitiveTypesStorage)
+
+struct FDataflowPrimitiveTypePolicy :
+	public TDataflowMultiTypePolicy<FBox, FSphere, FDataflowPlane>
+{
+};
+
+USTRUCT()
+struct FDataflowPrimitiveTypes : public FDataflowAnyType
+{
+	using FPolicyType = FDataflowPrimitiveTypePolicy;
+	using FStorageType = FDataflowPrimitiveTypesStorage;
+
+	GENERATED_BODY()
+
+	// This needs to be a UPROPERTY() because the anytype system needs to inspect it
+	UPROPERTY()
+	FDataflowPrimitiveTypesStorage Value;
+};
+
+// ------------------------------------------------------------------------------------------------------------------
+
+USTRUCT()
+struct FDataflowPrimitiveArrayStorage
+{
+	GENERATED_BODY()
+
+	FDataflowPrimitiveArrayStorage() = default;
+	FDataflowPrimitiveArrayStorage(const FDataflowPrimitiveArrayStorage& Other) = default;
+	FDataflowPrimitiveArrayStorage& operator=(const FDataflowPrimitiveArrayStorage& Other) = default;
+	FDataflowPrimitiveArrayStorage(FDataflowPrimitiveArrayStorage&& Other) = default;
+	FDataflowPrimitiveArrayStorage& operator=(FDataflowPrimitiveArrayStorage&& Other) = default;
+	FDataflowPrimitiveArrayStorage(const TArray<FBox>& PrimitiveArray)
+		: Storage(TInPlaceType<TArray<FBox>>(), PrimitiveArray)
+	{
+	}
+	FDataflowPrimitiveArrayStorage(const TArray<FSphere>& PrimitiveArray)
+		: Storage(TInPlaceType<TArray<FSphere>>(), PrimitiveArray)
+	{
+	}
+	FDataflowPrimitiveArrayStorage(const TArray<FDataflowPlane>& PrimitiveArray)
+		: Storage(TInPlaceType<TArray<FDataflowPlane>>(), PrimitiveArray)
+	{
+	}
+
+	template<typename T>
+	const T* TryGet() const
+	{
+		return Storage.TryGet<T>();
+	}
+
+private:
+	// Currently we don't support serialization for this storage
+	TVariant<TArray<FBox>, TArray<FSphere>, TArray<FDataflowPlane>> Storage;
+};
+
+template <>
+struct FDataflowConverter<FDataflowPrimitiveArrayStorage>
+{
+	template <typename TFromType>
+	static void From(const TFromType& From, FDataflowPrimitiveArrayStorage& To)
+	{
+		To = TFromType(From);
+	}
+	template <typename TToType>
+	static void To(const FDataflowPrimitiveArrayStorage& From, TToType& To)
+	{
+		if (const TToType* Ptr = From.TryGet<TToType>())
+		{
+			To = *Ptr;
+		}
+	}
+};
+
+UE_DATAFLOW_POLICY_DECLARE_TYPENAME(FDataflowPrimitiveArrayStorage)
+
+struct FDataflowPrimitiveArrayPolicy :
+	public TDataflowMultiTypePolicy<TArray<FBox>, TArray<FSphere>, TArray<FDataflowPlane>>
+{
+};
+
+USTRUCT()
+struct FDataflowPrimitiveArrayTypes : public FDataflowAnyType
+{
+	using FPolicyType = FDataflowPrimitiveArrayPolicy;
+	using FStorageType = FDataflowPrimitiveArrayStorage;
+
+	GENERATED_BODY()
+
+	// This needs to be a UPROPERTY() because the anytype system needs to inspect it
+	UPROPERTY()
+	FDataflowPrimitiveArrayStorage Value;
+};
+
+namespace UE::Dataflow
+{
+	DATAFLOWCORE_API void RegisterPrimitiveTypes();
+};
diff --git a/Engine/Source/Runtime/Experimental/Dataflow/Core/Public/Dataflow/DataflowSelection.h b/Engine/Source/Runtime/Experimental/Dataflow/Core/Public/Dataflow/DataflowSelection.h
--- a/Engine/Source/Runtime/Experimental/Dataflow/Core/Public/Dataflow/DataflowSelection.h
+++ b/Engine/Source/Runtime/Experimental/Dataflow/Core/Public/Dataflow/DataflowSelection.h
@@ -2,11 +2,12 @@
 
 #pragma once
 
-
 #include "UObject/NameTypes.h"
+
 #include "DataflowSelection.generated.h"
 
 struct FManagedArrayCollection;
+struct FDataflowConnection;
 
 USTRUCT()
 struct FDataflowSelection
@@ -45,8 +46,6 @@ struct FDataflowSelection
 	DATAFLOWCORE_API void AsArrayValidated(TArray<int32>& SelectionArr, const FManagedArrayCollection& InCollection) const;
 	DATAFLOWCORE_API TArray<int32> AsArrayValidated(const FManagedArrayCollection& InCollection) const;
 	
-	DATAFLOWCORE_API void SetGroupName(const FName& InGroupName);
-
 	/**
 	 * Sets the selection from a sparse array (it only contains the indices of the selected items)
 	 * @return true if all indices were successfully set, false if some indices were invalid and could not be set
@@ -85,6 +84,11 @@ struct FDataflowSelection
 
 	FName GetGroupName() const { return GroupName; }
 
+	/**
+	* Returns group name from input or output connection
+	*/
+	DATAFLOWCORE_API static FName GetSelectionTypeFromCollection(const FDataflowConnection& InConnection);
+
 private:
 	FName GroupName;
 	TBitArray<> SelectionArray;
diff --git a/Engine/Source/Runtime/Experimental/Dataflow/Core/Public/Dataflow/DataflowUtils.h b/Engine/Source/Runtime/Experimental/Dataflow/Core/Public/Dataflow/DataflowUtils.h
new file mode 100644
--- /dev/null
+++ b/Engine/Source/Runtime/Experimental/Dataflow/Core/Public/Dataflow/DataflowUtils.h
@@ -0,0 +1,25 @@
+// Copyright Epic Games, Inc. All Rights Reserved.
+
+#pragma once 
+
+#include "Containers/Array.h"
+
+namespace UE::Dataflow
+{
+	namespace Utils
+	{
+		enum class EErrorCode
+		{
+			None = 0,
+			InvalidChars = 1,
+			InvalidFormatInSegment = 2,
+			NumberInRangeTooBig = 3,
+			InvalidRangeStartEnd = 4
+		};
+
+		/* e.g. "0, 2, 5-10, 12-15" */
+		DATAFLOWCORE_API EErrorCode ParseIndicesStr(const FString& InFramesString, TArray<int32>& OutIndices);
+	}
+}
+
+
