add new dataflow node to simplify medial skeletons

#Synced-CL 51925774

diff --git a/Engine/Plugins/Experimental/GeometryCollectionPlugin/Source/GeometryCollectionNodes/Private/Dataflow/GeometryCollectionSamplingNodes.cpp b/Engine/Plugins/Experimental/GeometryCollectionPlugin/Source/GeometryCollectionNodes/Private/Dataflow/GeometryCollectionSamplingNodes.cpp
--- a/Engine/Plugins/Experimental/GeometryCollectionPlugin/Source/GeometryCollectionNodes/Private/Dataflow/GeometryCollectionSamplingNodes.cpp
+++ b/Engine/Plugins/Experimental/GeometryCollectionPlugin/Source/GeometryCollectionNodes/Private/Dataflow/GeometryCollectionSamplingNodes.cpp
@@ -13,6 +13,8 @@
 
 #include UE_INLINE_GENERATED_CPP_BY_NAME(GeometryCollectionSamplingNodes)
 
+#define LOCTEXT_NAMESPACE "GeometryCollectionSamplingNodes"
+
 namespace UE::Dataflow
 {
 
@@ -23,6 +25,7 @@ namespace UE::Dataflow
 		DATAFLOW_NODE_REGISTER_CREATION_FACTORY(FVertexWeightedPointSamplingDataflowNode);
 		DATAFLOW_NODE_REGISTER_CREATION_FACTORY(FFilterPointSetWithMeshDataflowNode);
 		DATAFLOW_NODE_REGISTER_CREATION_FACTORY(FMeshMedialSkeletonSamplingDataflowNode);
+		DATAFLOW_NODE_REGISTER_CREATION_FACTORY(FSimplifyMedialSkeletonDataflowNode);
 	}
 
 	static void MedialSkeletonDebugDraw(IDataflowDebugDrawInterface& DataflowRenderingInterface, const UE::Geometry::MedialAxis::FMedialSkeleton& ToDrawSkel)
@@ -402,7 +405,7 @@ void FMeshMedialSkeletonSamplingDataflowNode::Evaluate(UE::Dataflow::FContext& C
 		FDataflowMedialSkeleton OutMedialSkeleton;
 		UE::Geometry::MedialAxis::FMedialSkeleton& Skeleton = OutMedialSkeleton.Skeleton;
 
-		if (TObjectPtr<UDynamicMesh> InTargetMesh = GetValue(Context, &TargetMesh))
+		if (TObjectPtr<const UDynamicMesh> InTargetMesh = GetValue(Context, &TargetMesh))
 		{
 			const UE::Geometry::FDynamicMesh3& InDynTargetMesh = InTargetMesh->GetMeshRef();
 
@@ -421,6 +424,80 @@ void FMeshMedialSkeletonSamplingDataflowNode::Evaluate(UE::Dataflow::FContext& C
 				Skeleton = SkelSample.ComputeSkeleton(Spatial, 1e-4, &FWTree);
 			}
 		}
+		else
+		{
+			Context.Error(LOCTEXT("MedialSkeletonSamplingNullMeshError", "MeshMedialSkeletonSampling: TargetMesh input is null"), this, Out);
+		}
+
+		TArray<FSphere> Spheres = Skeleton.Spheres;
+		SetValue(Context, MoveTemp(OutMedialSkeleton), &MedialSkeleton);
+		SetValue(Context, MoveTemp(Spheres), &MedialSpheres);
+	}
+}
+
+FSimplifyMedialSkeletonDataflowNode::FSimplifyMedialSkeletonDataflowNode(const UE::Dataflow::FNodeParameters& InParam, FGuid InGuid)
+	: FDataflowNode(InParam, InGuid)
+{
+	RegisterInputConnection(&MedialSkeleton);
+	RegisterInputConnection(&TargetMesh);
+	RegisterInputConnection(&MinSpheres).SetCanHidePin(true).SetPinIsHidden(true);
+	RegisterInputConnection(&QEMErrorThresholdSqrt).SetCanHidePin(true).SetPinIsHidden(true);
+	RegisterInputConnection(&EdgeLengthThreshold).SetCanHidePin(true).SetPinIsHidden(true);
+	RegisterInputConnection(&EdgeRadiusThreshold).SetCanHidePin(true).SetPinIsHidden(true);
+	RegisterInputConnection(&bOnlySimplifySurfaces).SetCanHidePin(true).SetPinIsHidden(true);
+	RegisterInputConnection(&ClusterSkeletonDistanceWt).SetCanHidePin(true).SetPinIsHidden(true);
+	RegisterInputConnection(&bPreventEdgeSurfaceIntersections).SetCanHidePin(true).SetPinIsHidden(true);
+	RegisterInputConnection(&ClusterRegularizeWeight).SetCanHidePin(true).SetPinIsHidden(true);
+
+	RegisterOutputConnection(&MedialSkeleton);
+	RegisterOutputConnection(&MedialSpheres);
+}
+
+void FSimplifyMedialSkeletonDataflowNode::Evaluate(UE::Dataflow::FContext& Context, const FDataflowOutput* Out) const
+{
+	if (Out->IsA(&MedialSkeleton) || Out->IsA(&MedialSpheres))
+	{
+		FDataflowMedialSkeleton OutMedialSkeleton = GetValue(Context, &MedialSkeleton);
+		UE::Geometry::MedialAxis::FMedialSkeleton& Skeleton = OutMedialSkeleton.Skeleton;
+
+		if (TObjectPtr<const UDynamicMesh> InTargetMesh = GetValue(Context, &TargetMesh))
+		{
+			const UE::Geometry::FDynamicMesh3& InDynTargetMesh = InTargetMesh->GetMeshRef();
+			bool bIsMeshCompatible = Skeleton.IsCompatibleWithMesh(InDynTargetMesh);
+			bool bIsNonEmpty = Skeleton.Spheres.Num() > 0;
+			if (bIsMeshCompatible && bIsNonEmpty)
+			{
+				UE::Geometry::TMeshAABBTree3<FDynamicMesh3> Spatial(&InDynTargetMesh);
+
+				UE::Geometry::MedialAxis::FSkeletonSimplifier Simplifier;
+				Simplifier.MinSpheres = GetValue(Context, &MinSpheres);
+				double UseQEMThresholdSqrt = GetValue(Context, &QEMErrorThresholdSqrt);
+				Simplifier.QEMErrorThreshold = UseQEMThresholdSqrt * UseQEMThresholdSqrt;
+				Simplifier.EdgeLengthThreshold = GetValue(Context, &EdgeLengthThreshold);
+				Simplifier.EdgeRadiusThreshold = GetValue(Context, &EdgeRadiusThreshold);
+				Simplifier.bOnlySimplifySurfaces = GetValue(Context, &bOnlySimplifySurfaces);
+				Simplifier.ClusterSkeletonDistanceWt = GetValue(Context, &ClusterSkeletonDistanceWt);
+				Simplifier.bPreventEdgeSurfaceIntersections = GetValue(Context, &bPreventEdgeSurfaceIntersections);
+				Simplifier.ClusterRegularizeWeight = GetValue(Context, &ClusterRegularizeWeight);
+
+				Simplifier.Simplify(Skeleton, Spatial, [](int32) { return false; }, 1e-4);
+			}
+			else
+			{
+				if (!bIsMeshCompatible)
+				{
+					Context.Error(LOCTEXT("SimplifyMedialSkeletonIncompatibleMesh", "SimplifyMedialSkeleton: Incompatible mesh provided. Simplify expects a mesh matching what was used to generate the skeleton."), this, Out);
+				}
+				if (!bIsNonEmpty)
+				{
+					Context.Warning(LOCTEXT("SimplifyMedialSkeletonEmptySkeleton", "SimplifyMedialSkeleton: Empty skeleton has nothing to simplify."), this, Out);
+				}
+			}
+		}
+		else
+		{
+			Context.Error(LOCTEXT("SimplifyMedialSkeletonNullMeshError", "SimplifyMedialSkeleton: TargetMesh input is null"), this, Out);
+		}
 
 		TArray<FSphere> Spheres = Skeleton.Spheres;
 		SetValue(Context, MoveTemp(OutMedialSkeleton), &MedialSkeleton);
@@ -429,6 +506,26 @@ void FMeshMedialSkeletonSamplingDataflowNode::Evaluate(UE::Dataflow::FContext& C
 }
 
 #if WITH_EDITOR
+bool FSimplifyMedialSkeletonDataflowNode::CanDebugDrawViewMode(const FName& ViewModeName) const
+{
+	return ViewModeName == UE::Dataflow::FDataflowConstruction3DViewMode::Name;
+}
+
+void FSimplifyMedialSkeletonDataflowNode::DebugDraw(UE::Dataflow::FContext& Context, IDataflowDebugDrawInterface& DataflowRenderingInterface, const FDebugDrawParameters& DebugDrawParameters) const
+{
+	using namespace UE::Geometry;
+
+	if ((DebugDrawParameters.bNodeIsSelected || DebugDrawParameters.bNodeIsPinned))
+	{
+		if (const FDataflowOutput* MedialSkeletonOutput = FindOutput(&MedialSkeleton))
+		{
+			const UE::Geometry::MedialAxis::FMedialSkeleton& ToDrawSkeleton = MedialSkeletonOutput->GetValue(Context, MedialSkeleton).Skeleton;
+
+			UE::Dataflow::MedialSkeletonDebugDraw(DataflowRenderingInterface, ToDrawSkeleton);
+		}
+	}
+}
+
 bool FMeshMedialSkeletonSamplingDataflowNode::CanDebugDrawViewMode(const FName& ViewModeName) const
 {
 	return ViewModeName == UE::Dataflow::FDataflowConstruction3DViewMode::Name;
@@ -440,9 +537,9 @@ void FMeshMedialSkeletonSamplingDataflowNode::DebugDraw(UE::Dataflow::FContext&
 
 	if ((DebugDrawParameters.bNodeIsSelected || DebugDrawParameters.bNodeIsPinned))
 	{
-		if (const FDataflowOutput* SphereCoveringOutput = FindOutput(&MedialSkeleton))
+		if (const FDataflowOutput* MedialSkeletonOutput = FindOutput(&MedialSkeleton))
 		{
-			const UE::Geometry::MedialAxis::FMedialSkeleton& ToDrawSkeleton = SphereCoveringOutput->GetValue(Context, MedialSkeleton).Skeleton;
+			const UE::Geometry::MedialAxis::FMedialSkeleton& ToDrawSkeleton = MedialSkeletonOutput->GetValue(Context, MedialSkeleton).Skeleton;
 
 			UE::Dataflow::MedialSkeletonDebugDraw(DataflowRenderingInterface, ToDrawSkeleton);
 		}
@@ -450,3 +547,5 @@ void FMeshMedialSkeletonSamplingDataflowNode::DebugDraw(UE::Dataflow::FContext&
 }
 #endif
 
+#undef LOCTEXT_NAMESPACE
+
diff --git a/Engine/Plugins/Experimental/GeometryCollectionPlugin/Source/GeometryCollectionNodes/Public/Dataflow/GeometryCollectionSamplingNodes.h b/Engine/Plugins/Experimental/GeometryCollectionPlugin/Source/GeometryCollectionNodes/Public/Dataflow/GeometryCollectionSamplingNodes.h
--- a/Engine/Plugins/Experimental/GeometryCollectionPlugin/Source/GeometryCollectionNodes/Public/Dataflow/GeometryCollectionSamplingNodes.h
+++ b/Engine/Plugins/Experimental/GeometryCollectionPlugin/Source/GeometryCollectionNodes/Public/Dataflow/GeometryCollectionSamplingNodes.h
@@ -398,6 +398,77 @@ public:
 
 
 
+/**
+ *
+ * Simplify a medial skeleton using the quadric error metric
+ *
+ */
+USTRUCT(meta = (DataflowGeometryCollection))
+struct FSimplifyMedialSkeletonDataflowNode : public FDataflowNode
+{
+	GENERATED_USTRUCT_BODY()
+	DATAFLOW_NODE_DEFINE_INTERNAL(FSimplifyMedialSkeletonDataflowNode, "SimplifyMedialSkeleton", "MedialSkeletonSampling", "Reduce")
+
+private:
+	/** Medial skeleton to simplify */
+	UPROPERTY(meta = (DataflowInput, DataflowOutput, DataflowIntrinsic))
+	FDataflowMedialSkeleton MedialSkeleton;
+
+	/** Mesh the skeleton was sampled from, used to compute the error metric */
+	UPROPERTY(meta = (DataflowInput, DataflowIntrinsic))
+	TObjectPtr<UDynamicMesh> TargetMesh;
+
+	UPROPERTY(meta = (DataflowOutput))
+	TArray<FSphere> MedialSpheres;
+
+	// Stop simplifying at this sphere count
+	UPROPERTY(EditAnywhere, Category = Options, meta = (DataflowInput, ClampMin = 1))
+	int32 MinSpheres = 1;
+
+	// If positive, will not collapse edges if it would introduce more quadric error than the square of this threshold.
+	// (Note: Not a distance error threshold.)
+	UPROPERTY(EditAnywhere, Category = Options, meta = (DataflowInput, ClampMin = 0, DisplayName="Quadric Error Threshold (Sqrt)"))
+	float QEMErrorThresholdSqrt = 0.;
+
+	// If positive, the maximum edge length to consider for collapse
+	UPROPERTY(EditAnywhere, Category = Options, meta = (DataflowInput, ClampMin = 0, Units = "cm"))
+	float EdgeLengthThreshold = 0.;
+	
+	// If positive, the maximum cluster radius to consider for collapse.  If either cluster on an edge has smaller radii, it is allowed to collapse.
+	UPROPERTY(EditAnywhere, Category = Options, meta = (DataflowInput, ClampMin = 0, Units = "cm"))
+	float EdgeRadiusThreshold = 0.;
+
+	// If true, only simplify edges of skeleton surfaces (edges that have triangles associated)
+	UPROPERTY(EditAnywhere, Category = Options, meta = (DataflowInput))
+	bool bOnlySimplifySurfaces = true;
+
+	// How much the simplifier should attempt to remain close to the initial skeleton, vs keeping medial spheres close to the mesh surface.
+	// In the [0,1] range. At 0 the original skeleton is not considered; at 1 the mesh surface is not considered.
+	UPROPERTY(EditAnywhere, Category = Options, meta = (DataflowInput, ClampMin = 0, ClampMax = 1, UIMax = .9999))
+	float ClusterSkeletonDistanceWt = .5f;
+
+	// If true, prevent collapses that would introduce intersections between skeleton edges and the mesh
+	UPROPERTY(EditAnywhere, Category = Options, meta = (DataflowInput))
+	bool bPreventEdgeSurfaceIntersections = true;
+
+	// Regularization weight encouraging skeleton clusters to keep their original positions and radii
+	UPROPERTY(EditAnywhere, Category = Options, AdvancedDisplay, meta = (DataflowInput, ClampMin = 0))
+	float ClusterRegularizeWeight = 1e-6f;
+
+	virtual void Evaluate(UE::Dataflow::FContext& Context, const FDataflowOutput* Out) const override;
+
+#if WITH_EDITOR
+	virtual bool CanDebugDrawViewMode(const FName& ViewModeName) const override;
+	virtual void DebugDraw(UE::Dataflow::FContext& Context, IDataflowDebugDrawInterface& DataflowRenderingInterface, const FDebugDrawParameters& DebugDrawParameters) const override;
+#endif
+
+
+public:
+	FSimplifyMedialSkeletonDataflowNode(const UE::Dataflow::FNodeParameters& InParam, FGuid InGuid = FGuid::NewGuid());
+};
+
+
+
 namespace UE::Dataflow
 {
 	void GeometryCollectionSamplingNodes();
diff --git a/Engine/Plugins/Runtime/GeometryProcessing/Source/GeometryAlgorithms/Private/Skeletonization/MeshMedialAxisSampling.cpp b/Engine/Plugins/Runtime/GeometryProcessing/Source/GeometryAlgorithms/Private/Skeletonization/MeshMedialAxisSampling.cpp
--- a/Engine/Plugins/Runtime/GeometryProcessing/Source/GeometryAlgorithms/Private/Skeletonization/MeshMedialAxisSampling.cpp
+++ b/Engine/Plugins/Runtime/GeometryProcessing/Source/GeometryAlgorithms/Private/Skeletonization/MeshMedialAxisSampling.cpp
@@ -1443,6 +1443,10 @@ namespace UE::Geometry::MedialAxis
 	{
 		using namespace MeshMedialAxisSamplingLocals;
 		const MeshType& Mesh = *MeshBVH.GetMesh();
+		if (!ensure(Skeleton.IsCompatibleWithMesh(Mesh)))
+		{
+			return;
+		}
 		TArray<FVector3f> TriNormals;
 		TArray<float> TriAreas;
 		TArray<FSQEM> VtxSQEMs;
@@ -1460,6 +1464,10 @@ namespace UE::Geometry::MedialAxis
 	{
 		using namespace MeshMedialAxisSamplingLocals;
 		const MeshType& Mesh = *MeshBVH.GetMesh();
+		if (!ensure(Skeleton.IsCompatibleWithMesh(Mesh)))
+		{
+			return;
+		}
 		TArray<int32> WeldIdx = ComputeWeldIdx(Mesh, OptionalBoundaryVertexFilter, ImplicitVertexNormalWeldThreshold);
 		auto NopNormals = [](int32) -> FVector3d { return FVector3d::ZeroVector; };
 		if (WeldIdx.IsEmpty())
diff --git a/Engine/Plugins/Runtime/GeometryProcessing/Source/GeometryAlgorithms/Public/Skeletonization/MeshMedialAxisSampling.h b/Engine/Plugins/Runtime/GeometryProcessing/Source/GeometryAlgorithms/Public/Skeletonization/MeshMedialAxisSampling.h
--- a/Engine/Plugins/Runtime/GeometryProcessing/Source/GeometryAlgorithms/Public/Skeletonization/MeshMedialAxisSampling.h
+++ b/Engine/Plugins/Runtime/GeometryProcessing/Source/GeometryAlgorithms/Public/Skeletonization/MeshMedialAxisSampling.h
@@ -264,6 +264,23 @@ namespace MedialAxis
 		// Note: Relatively expensive, intended for debugging / testing
 		GEOMETRYALGORITHMS_API bool CheckValidity(EValidityCheckFailMode FailMode = EValidityCheckFailMode::Ensure) const;
 
+		template<typename TriangleMeshType>
+		bool IsCompatibleWithMesh(const TriangleMeshType& Mesh)
+		{
+			if (Mesh.MaxVertexID() != VIDtoClusterIndex.Num())
+			{
+				return false;
+			}
+			for (int32 Idx = 0; Idx < VIDtoClusterIndex.Num(); ++Idx)
+			{
+				if (VIDtoClusterIndex[Idx] != INDEX_NONE && !Mesh.IsVertex(Idx))
+				{
+					return false;
+				}
+			}
+			return true;
+		}
+
 
 		// Compute the ClusterTriangles array for the current FMedialSkeleton, given a Mesh and assuming VIDtoClusterIndex is already computed
 		template<typename TriangleMeshType>
