Dataflow:
- Set up the autoconvert between UDataflowMesh and UDynamicMesh

#Synced-CL 52238782

diff --git a/Engine/Plugins/Experimental/GeometryCollectionPlugin/Source/GeometryCollectionNodes/Private/Dataflow/GeometryCollectionMeshNodes.cpp b/Engine/Plugins/Experimental/GeometryCollectionPlugin/Source/GeometryCollectionNodes/Private/Dataflow/GeometryCollectionMeshNodes.cpp
--- a/Engine/Plugins/Experimental/GeometryCollectionPlugin/Source/GeometryCollectionNodes/Private/Dataflow/GeometryCollectionMeshNodes.cpp
+++ b/Engine/Plugins/Experimental/GeometryCollectionPlugin/Source/GeometryCollectionNodes/Private/Dataflow/GeometryCollectionMeshNodes.cpp
@@ -78,6 +78,8 @@ namespace UE::Dataflow
 		UE_DATAFLOW_REGISTER_AUTOCONVERT(TObjectPtr<UDynamicMesh>, FManagedArrayCollection, FMeshToCollectionDataflowNode);
 		UE_DATAFLOW_REGISTER_AUTOCONVERT(TObjectPtr<UDynamicMesh>, FBox, FGetMeshBoundingBoxDataflowNode);
 		UE_DATAFLOW_REGISTER_AUTOCONVERT(FBox, TObjectPtr<UDynamicMesh>, FBoxToMeshDataflowNode);
+		UE_DATAFLOW_REGISTER_AUTOCONVERT(TObjectPtr<UDynamicMesh>, TObjectPtr<UDataflowMesh>, FMakeDataflowMeshDataflowNode);
+		UE_DATAFLOW_REGISTER_AUTOCONVERT(TObjectPtr<UDataflowMesh>, TObjectPtr<UDynamicMesh>, FSplitDataflowMeshDataflowNode);
 	}
 }
 
@@ -1331,3 +1333,5 @@ void FCollectionSelectionToMeshesDataflowNode::Evaluate(UE::Dataflow::FContext&
 		SetValue(Context, NewMeshes, &Meshes);
 	}
 }
+
+// -----------------------------------------------------------------------------------------------------------
