rename medial skeleton simplifier parameter EdgeRadiusThreshold -> SphereRadiusThreshold

#Synced-CL 52243346

diff --git a/Engine/Plugins/Experimental/GeometryCollectionPlugin/Source/GeometryCollectionNodes/Private/Dataflow/GeometryCollectionSamplingNodes.cpp b/Engine/Plugins/Experimental/GeometryCollectionPlugin/Source/GeometryCollectionNodes/Private/Dataflow/GeometryCollectionSamplingNodes.cpp
--- a/Engine/Plugins/Experimental/GeometryCollectionPlugin/Source/GeometryCollectionNodes/Private/Dataflow/GeometryCollectionSamplingNodes.cpp
+++ b/Engine/Plugins/Experimental/GeometryCollectionPlugin/Source/GeometryCollectionNodes/Private/Dataflow/GeometryCollectionSamplingNodes.cpp
@@ -443,7 +443,7 @@ FSimplifyMedialSkeletonDataflowNode::FSimplifyMedialSkeletonDataflowNode(const U
 	RegisterInputConnection(&MinSpheres).SetCanHidePin(true).SetPinIsHidden(true);
 	RegisterInputConnection(&QEMErrorThresholdSqrt).SetCanHidePin(true).SetPinIsHidden(true);
 	RegisterInputConnection(&EdgeLengthThreshold).SetCanHidePin(true).SetPinIsHidden(true);
-	RegisterInputConnection(&EdgeRadiusThreshold).SetCanHidePin(true).SetPinIsHidden(true);
+	RegisterInputConnection(&SphereRadiusThreshold).SetCanHidePin(true).SetPinIsHidden(true);
 	RegisterInputConnection(&bOnlySimplifySurfaces).SetCanHidePin(true).SetPinIsHidden(true);
 	RegisterInputConnection(&ClusterSkeletonDistanceWt).SetCanHidePin(true).SetPinIsHidden(true);
 	RegisterInputConnection(&bPreventEdgeSurfaceIntersections).SetCanHidePin(true).SetPinIsHidden(true);
@@ -474,7 +474,7 @@ void FSimplifyMedialSkeletonDataflowNode::Evaluate(UE::Dataflow::FContext& Conte
 				double UseQEMThresholdSqrt = GetValue(Context, &QEMErrorThresholdSqrt);
 				Simplifier.QEMErrorThreshold = UseQEMThresholdSqrt * UseQEMThresholdSqrt;
 				Simplifier.EdgeLengthThreshold = GetValue(Context, &EdgeLengthThreshold);
-				Simplifier.EdgeRadiusThreshold = GetValue(Context, &EdgeRadiusThreshold);
+				Simplifier.SphereRadiusThreshold = GetValue(Context, &SphereRadiusThreshold);
 				Simplifier.bOnlySimplifySurfaces = GetValue(Context, &bOnlySimplifySurfaces);
 				Simplifier.ClusterSkeletonDistanceWt = GetValue(Context, &ClusterSkeletonDistanceWt);
 				Simplifier.bPreventEdgeSurfaceIntersections = GetValue(Context, &bPreventEdgeSurfaceIntersections);
diff --git a/Engine/Plugins/Experimental/GeometryCollectionPlugin/Source/GeometryCollectionNodes/Public/Dataflow/GeometryCollectionSamplingNodes.h b/Engine/Plugins/Experimental/GeometryCollectionPlugin/Source/GeometryCollectionNodes/Public/Dataflow/GeometryCollectionSamplingNodes.h
--- a/Engine/Plugins/Experimental/GeometryCollectionPlugin/Source/GeometryCollectionNodes/Public/Dataflow/GeometryCollectionSamplingNodes.h
+++ b/Engine/Plugins/Experimental/GeometryCollectionPlugin/Source/GeometryCollectionNodes/Public/Dataflow/GeometryCollectionSamplingNodes.h
@@ -434,9 +434,9 @@ private:
 	UPROPERTY(EditAnywhere, Category = Options, meta = (DataflowInput, ClampMin = 0, Units = "cm"))
 	float EdgeLengthThreshold = 0.;
 	
-	// If positive, the maximum cluster radius to consider for collapse.  If either cluster on an edge has smaller radii, it is allowed to collapse.
+	// If positive, the maximum sphere radius to consider for collapse.  If either cluster's medial sphere on an edge has smaller radii, it is allowed to collapse.
 	UPROPERTY(EditAnywhere, Category = Options, meta = (DataflowInput, ClampMin = 0, Units = "cm"))
-	float EdgeRadiusThreshold = 0.;
+	float SphereRadiusThreshold = 0.;
 
 	// If true, only simplify edges of skeleton surfaces (edges that have triangles associated)
 	UPROPERTY(EditAnywhere, Category = Options, meta = (DataflowInput))
diff --git a/Engine/Plugins/Runtime/GeometryProcessing/Source/GeometryAlgorithms/Private/Skeletonization/MeshMedialAxisSampling.cpp b/Engine/Plugins/Runtime/GeometryProcessing/Source/GeometryAlgorithms/Private/Skeletonization/MeshMedialAxisSampling.cpp
--- a/Engine/Plugins/Runtime/GeometryProcessing/Source/GeometryAlgorithms/Private/Skeletonization/MeshMedialAxisSampling.cpp
+++ b/Engine/Plugins/Runtime/GeometryProcessing/Source/GeometryAlgorithms/Private/Skeletonization/MeshMedialAxisSampling.cpp
@@ -745,8 +745,8 @@ namespace MeshMedialAxisSamplingLocals
 			{
 				return true;
 			}
-			if (Simplifier.EdgeRadiusThreshold > 0 &&
-				Skeleton.Spheres[Edge.A].W > Simplifier.EdgeRadiusThreshold && Skeleton.Spheres[Edge.B].W > Simplifier.EdgeRadiusThreshold)
+			if (Simplifier.SphereRadiusThreshold > 0 &&
+				Skeleton.Spheres[Edge.A].W > Simplifier.SphereRadiusThreshold && Skeleton.Spheres[Edge.B].W > Simplifier.SphereRadiusThreshold)
 			{
 				return true;
 			}
diff --git a/Engine/Plugins/Runtime/GeometryProcessing/Source/GeometryAlgorithms/Public/Skeletonization/MeshMedialAxisSampling.h b/Engine/Plugins/Runtime/GeometryProcessing/Source/GeometryAlgorithms/Public/Skeletonization/MeshMedialAxisSampling.h
--- a/Engine/Plugins/Runtime/GeometryProcessing/Source/GeometryAlgorithms/Public/Skeletonization/MeshMedialAxisSampling.h
+++ b/Engine/Plugins/Runtime/GeometryProcessing/Source/GeometryAlgorithms/Public/Skeletonization/MeshMedialAxisSampling.h
@@ -419,8 +419,8 @@ namespace MedialAxis
 		// If positive, the maximum edge length to consider for collapse
 		double EdgeLengthThreshold = -1;
 		
-		// If positive, the maximum cluster radius to consider for collapse (will not collapse if both cluster radii on an edge are larger than this)
-		double EdgeRadiusThreshold = -1;
+		// If positive, the maximum medial sphere radius to consider for collapse (will not collapse if both cluster's medial sphere radii on an edge are larger than this)
+		double SphereRadiusThreshold = -1;
 
 		// If true, only simplify edges that have triangles associated
 		bool bOnlySimplifySurfaces = false;
