add skinny triangle split options to dataflow medial skeleton simplifier

#Synced-CL 52315161

diff --git a/Engine/Plugins/Experimental/GeometryCollectionPlugin/Source/GeometryCollectionNodes/Private/Dataflow/GeometryCollectionSamplingNodes.cpp b/Engine/Plugins/Experimental/GeometryCollectionPlugin/Source/GeometryCollectionNodes/Private/Dataflow/GeometryCollectionSamplingNodes.cpp
--- a/Engine/Plugins/Experimental/GeometryCollectionPlugin/Source/GeometryCollectionNodes/Private/Dataflow/GeometryCollectionSamplingNodes.cpp
+++ b/Engine/Plugins/Experimental/GeometryCollectionPlugin/Source/GeometryCollectionNodes/Private/Dataflow/GeometryCollectionSamplingNodes.cpp
@@ -448,6 +448,9 @@ FSimplifyMedialSkeletonDataflowNode::FSimplifyMedialSkeletonDataflowNode(const U
 	RegisterInputConnection(&ClusterSkeletonDistanceWt).SetCanHidePin(true).SetPinIsHidden(true);
 	RegisterInputConnection(&bPreventEdgeSurfaceIntersections).SetCanHidePin(true).SetPinIsHidden(true);
 	RegisterInputConnection(&ClusterRegularizeWeight).SetCanHidePin(true).SetPinIsHidden(true);
+	RegisterInputConnection(&bSplitThinTriEdges).SetCanHidePin(true).SetPinIsHidden(true);
+	RegisterInputConnection(&SplitThinTriEdgeAngleThresholdDeg).SetCanHidePin(true).SetPinIsHidden(true);
+	RegisterInputConnection(&SplitThinTriEdgePosErrorWt).SetCanHidePin(true).SetPinIsHidden(true);
 
 	RegisterOutputConnection(&MedialSkeleton);
 	RegisterOutputConnection(&MedialSpheres);
@@ -479,6 +482,9 @@ void FSimplifyMedialSkeletonDataflowNode::Evaluate(UE::Dataflow::FContext& Conte
 				Simplifier.ClusterSkeletonDistanceWt = GetValue(Context, &ClusterSkeletonDistanceWt);
 				Simplifier.bPreventEdgeSurfaceIntersections = GetValue(Context, &bPreventEdgeSurfaceIntersections);
 				Simplifier.ClusterRegularizeWeight = GetValue(Context, &ClusterRegularizeWeight);
+				Simplifier.bSplitThinTriEdges = GetValue(Context, &bSplitThinTriEdges);
+				Simplifier.SplitThinTriEdgeAngleThresholdDeg = GetValue(Context, &SplitThinTriEdgeAngleThresholdDeg);
+				Simplifier.SplitThinTriEdgePosErrorWt = GetValue(Context, &SplitThinTriEdgePosErrorWt);
 
 				Simplifier.Simplify(Skeleton, Spatial, [](int32) { return false; }, 1e-4);
 			}
diff --git a/Engine/Plugins/Experimental/GeometryCollectionPlugin/Source/GeometryCollectionNodes/Public/Dataflow/GeometryCollectionSamplingNodes.h b/Engine/Plugins/Experimental/GeometryCollectionPlugin/Source/GeometryCollectionNodes/Public/Dataflow/GeometryCollectionSamplingNodes.h
--- a/Engine/Plugins/Experimental/GeometryCollectionPlugin/Source/GeometryCollectionNodes/Public/Dataflow/GeometryCollectionSamplingNodes.h
+++ b/Engine/Plugins/Experimental/GeometryCollectionPlugin/Source/GeometryCollectionNodes/Public/Dataflow/GeometryCollectionSamplingNodes.h
@@ -427,7 +427,7 @@ private:
 
 	// If positive, will not collapse edges if it would introduce more quadric error than the square of this threshold.
 	// (Note: Not a distance error threshold.)
-	UPROPERTY(EditAnywhere, Category = Options, meta = (DataflowInput, ClampMin = 0, DisplayName="Quadric Error Threshold (Sqrt)"))
+	UPROPERTY(EditAnywhere, Category = Options, AdvancedDisplay, meta = (DataflowInput, ClampMin = 0, DisplayName="Quadric Error Threshold (Sqrt)"))
 	float QEMErrorThresholdSqrt = 0.;
 
 	// If positive, the maximum edge length to consider for collapse
@@ -455,6 +455,18 @@ private:
 	UPROPERTY(EditAnywhere, Category = Options, AdvancedDisplay, meta = (DataflowInput, ClampMin = 0))
 	float ClusterRegularizeWeight = 1e-6f;
 
+	// Whether to split edges attached to skinny triangles before simplifying, which can add useful degrees of freedom to the simplifier
+	UPROPERTY(EditAnywhere, Category = SplitEdges, AdvancedDisplay, meta = (DataflowInput))
+	bool bSplitThinTriEdges = true;
+
+	// Minimum angle (degrees) at the opposite vertex for an edge to be a split candidate
+	UPROPERTY(EditAnywhere, Category = SplitEdges, AdvancedDisplay, meta = (EditCondition="bSplitThinTriEdges", DataflowInput, ClampMin = 90, ClampMax = 180))
+	float SplitThinTriEdgeAngleThresholdDeg = 120.f;
+
+	// Position error weight used when re-clustering vertices after thin-triangle edge splits
+	UPROPERTY(EditAnywhere, Category = SplitEdges, AdvancedDisplay, meta = (EditCondition="bSplitThinTriEdges", DataflowInput, ClampMin = 0, ClampMax = 1))
+	float SplitThinTriEdgePosErrorWt = .2f;
+
 	virtual void Evaluate(UE::Dataflow::FContext& Context, const FDataflowOutput* Out) const override;
 
 #if WITH_EDITOR
