Add automatic root selection options to dataflow medial skeleton -> animation skeleton node

#Synced-CL 52352248

diff --git a/Engine/Plugins/Experimental/GeometryCollectionPlugin/Source/GeometryCollectionNodes/Private/Dataflow/MeshMedialSkeletonConversionNodes.cpp b/Engine/Plugins/Experimental/GeometryCollectionPlugin/Source/GeometryCollectionNodes/Private/Dataflow/MeshMedialSkeletonConversionNodes.cpp
--- a/Engine/Plugins/Experimental/GeometryCollectionPlugin/Source/GeometryCollectionNodes/Private/Dataflow/MeshMedialSkeletonConversionNodes.cpp
+++ b/Engine/Plugins/Experimental/GeometryCollectionPlugin/Source/GeometryCollectionNodes/Private/Dataflow/MeshMedialSkeletonConversionNodes.cpp
@@ -35,7 +35,10 @@ namespace MeshMedialSkeletonConversionLocals
 
 	void SetSkeletonConversionOptions(UE::Geometry::MedialAxis::FMedialSkeletonToTreeSkeletonOptions& OutOptions,
 		EDataflowMedialSkeletonConversionEdgeWeightMethod EdgeWeightMethod,
-		EDataflowMedialSkeletonConversionMergeDisconnectedMethod MergeDisconnectedMethod)
+		EDataflowMedialSkeletonConversionMergeDisconnectedMethod MergeDisconnectedMethod,
+		EDataflowMedialSkeletonConversionSelectRootMethod SelectRootMethod,
+		const FVector& RootSelectionPoint,
+		const FVector& RootSelectionDirection)
 	{
 		using namespace UE::Geometry::MedialAxis;
 		switch (EdgeWeightMethod)
@@ -58,7 +61,27 @@ namespace MeshMedialSkeletonConversionLocals
 		case EDataflowMedialSkeletonConversionMergeDisconnectedMethod::ConnectClosestBones:
 			OutOptions.MergeDisconnectedMethod = FMedialSkeletonToTreeSkeletonOptions::EMergeDisconnectedMethod::ConnectClosestBones;
 			break;
-		}	
+		}
+		switch (SelectRootMethod)
+		{
+		case EDataflowMedialSkeletonConversionSelectRootMethod::ClosestToPoint:
+			OutOptions.SelectRootMethod = FMedialSkeletonToTreeSkeletonOptions::ESelectRootMethod::ClosestToPoint;
+			break;
+		case EDataflowMedialSkeletonConversionSelectRootMethod::FarthestInDirection:
+			OutOptions.SelectRootMethod = FMedialSkeletonToTreeSkeletonOptions::ESelectRootMethod::FarthestInDirection;
+			break;
+		case EDataflowMedialSkeletonConversionSelectRootMethod::ClosestToBoundsCenter:
+			OutOptions.SelectRootMethod = FMedialSkeletonToTreeSkeletonOptions::ESelectRootMethod::ClosestToBoundsCenter;
+			break;
+		case EDataflowMedialSkeletonConversionSelectRootMethod::LargestSphere:
+			OutOptions.SelectRootMethod = FMedialSkeletonToTreeSkeletonOptions::ESelectRootMethod::LargestSphere;
+			break;
+		case EDataflowMedialSkeletonConversionSelectRootMethod::ArrayOrder:
+			OutOptions.SelectRootMethod = FMedialSkeletonToTreeSkeletonOptions::ESelectRootMethod::ArrayOrder;
+			break;
+		}
+		OutOptions.RootSelectionPoint = FVector3d(RootSelectionPoint);
+		OutOptions.RootSelectionDirection = FVector3d(RootSelectionDirection);
 	}
 
 	void SetSkinBindSettings(UE::Geometry::SkinBinding::FBindSettings& OutSettings, EDataflowBindSkeletonMethod BindMethod, float Stiffness, int32 MaxInfluences, int32 VoxelRes)
@@ -132,7 +155,7 @@ void FConvertMedialSkeletonToAnimationSkeletonDataflowNode::Evaluate(UE::Dataflo
 		TArray<FTransform> Poses;
 
 		UE::Geometry::MedialAxis::FMedialSkeletonToTreeSkeletonOptions Options;
-		MeshMedialSkeletonConversionLocals::SetSkeletonConversionOptions(Options, EdgeWeightMethod, MergeDisconnectedMethod);
+		MeshMedialSkeletonConversionLocals::SetSkeletonConversionOptions(Options, EdgeWeightMethod, MergeDisconnectedMethod, SelectRootMethod, RootSelectionPoint, RootSelectionDirection);
 		const int32 UseRootIndex = GetValue(Context, &RootClusterIndex);
 		Options.ToHierarchy(InMedialSkeleton.Skeleton, AvoidIntersectingMeshBVH, Parents, Poses, Options, UseRootIndex);
 
@@ -237,7 +260,7 @@ void FSkinMeshViaMedialSkeleton::Evaluate(UE::Dataflow::FContext& Context, const
 			FDynamicMesh3& DynMesh = OutMesh->GetMeshRef();
 
 			UE::Geometry::MedialAxis::FMedialSkeletonToTreeSkeletonOptions Options;
-			MeshMedialSkeletonConversionLocals::SetSkeletonConversionOptions(Options, EdgeWeightMethod, MergeDisconnectedMethod);
+			MeshMedialSkeletonConversionLocals::SetSkeletonConversionOptions(Options, EdgeWeightMethod, MergeDisconnectedMethod, SelectRootMethod, RootSelectionPoint, RootSelectionDirection);
 			const int32 UseRootIndex = GetValue(Context, &RootClusterIndex);
 			Options.ToDynamicMeshSkeleton(InMedialSkeleton.Skeleton, DynMesh, FSkeletalMeshAttributes::DefaultSkinWeightProfileName, Options, UseRootIndex);
 	
diff --git a/Engine/Plugins/Experimental/GeometryCollectionPlugin/Source/GeometryCollectionNodes/Private/Dataflow/MeshMedialSkeletonConversionNodes.h b/Engine/Plugins/Experimental/GeometryCollectionPlugin/Source/GeometryCollectionNodes/Private/Dataflow/MeshMedialSkeletonConversionNodes.h
--- a/Engine/Plugins/Experimental/GeometryCollectionPlugin/Source/GeometryCollectionNodes/Private/Dataflow/MeshMedialSkeletonConversionNodes.h
+++ b/Engine/Plugins/Experimental/GeometryCollectionPlugin/Source/GeometryCollectionNodes/Private/Dataflow/MeshMedialSkeletonConversionNodes.h
@@ -32,6 +32,21 @@ enum class EDataflowMedialSkeletonConversionMergeDisconnectedMethod : uint8
 	AddTopLevelRoot
 };
 
+UENUM()
+enum class EDataflowMedialSkeletonConversionSelectRootMethod : uint8
+{
+	// Select the cluster closest to the RootSelectionPoint
+	ClosestToPoint,
+	// Select the cluster farthest in the RootSelectionDirection
+	FarthestInDirection,
+	// Select the cluster closest to the bounding box center of the medial skeleton
+	ClosestToBoundsCenter,
+	// Select the cluster with the largest medial sphere radius
+	LargestSphere,
+	// Select the first cluster (index 0)
+	ArrayOrder
+};
+
 UENUM()
 enum class EDataflowBindSkeletonMethod : uint8
 {
@@ -70,6 +85,18 @@ private:
 	UPROPERTY(EditAnywhere, Category = "Options")
 	EDataflowMedialSkeletonConversionMergeDisconnectedMethod MergeDisconnectedMethod = EDataflowMedialSkeletonConversionMergeDisconnectedMethod::ConnectClosestBones;
 
+	/** Method to automatically select the root cluster when RootClusterIndex is not specified */
+	UPROPERTY(EditAnywhere, Category = "Options")
+	EDataflowMedialSkeletonConversionSelectRootMethod SelectRootMethod = EDataflowMedialSkeletonConversionSelectRootMethod::ArrayOrder;
+
+	/** When SelectRootMethod is ClosestToPoint, find the closest cluster to this point */
+	UPROPERTY(EditAnywhere, Category = "Options", meta = (EditCondition = "SelectRootMethod == EDataflowMedialSkeletonConversionSelectRootMethod::ClosestToPoint"))
+	FVector RootSelectionPoint = FVector::ZeroVector;
+
+	/** When SelectRootMethod is FarthestInDirection, the direction to search along */
+	UPROPERTY(EditAnywhere, Category = "Options", meta = (EditCondition = "SelectRootMethod == EDataflowMedialSkeletonConversionSelectRootMethod::FarthestInDirection"))
+	FVector RootSelectionDirection = FVector(0, 0, -1);
+
 	/** If non-negative, the index of the medial axis cluster (sphere) to use as the root of the Skeleton */
 	UPROPERTY(EditAnywhere, Category = "Options", meta = (DataflowInput))
 	int32 RootClusterIndex = INDEX_NONE;
@@ -141,6 +168,18 @@ private:
 	UPROPERTY(EditAnywhere, Category = "Options")
 	EDataflowMedialSkeletonConversionMergeDisconnectedMethod MergeDisconnectedMethod = EDataflowMedialSkeletonConversionMergeDisconnectedMethod::ConnectClosestBones;
 
+	/** Method to automatically select the root cluster when RootClusterIndex is not specified */
+	UPROPERTY(EditAnywhere, Category = "Options")
+	EDataflowMedialSkeletonConversionSelectRootMethod SelectRootMethod = EDataflowMedialSkeletonConversionSelectRootMethod::ArrayOrder;
+
+	/** When SelectRootMethod is ClosestToPoint, find the closest cluster to this point */
+	UPROPERTY(EditAnywhere, Category = "Options", meta = (EditCondition = "SelectRootMethod == EDataflowMedialSkeletonConversionSelectRootMethod::ClosestToPoint"))
+	FVector RootSelectionPoint = FVector::ZeroVector;
+
+	/** When SelectRootMethod is FarthestInDirection, the direction to search along */
+	UPROPERTY(EditAnywhere, Category = "Options", meta = (EditCondition = "SelectRootMethod == EDataflowMedialSkeletonConversionSelectRootMethod::FarthestInDirection"))
+	FVector RootSelectionDirection = FVector(0, 0, -1);
+
 	/** If non-negative, the index of the medial axis cluster (sphere) to use as the root of the Skeleton */
 	UPROPERTY(EditAnywhere, Category = "Options", meta = (DataflowInput))
 	int32 RootClusterIndex = INDEX_NONE;
