Geometry collection : use the new dataflow template picker when creating a new Geometry collection from the factory ( from the asset browser )
- remove the older flow that was asking for dataflow if none was set on the existing geometry collection
- added a template folder in the plugin as well

#virtualized
#Synced-CL 52234830

diff --git a/Engine/Plugins/Experimental/GeometryCollectionPlugin/Source/GeometryCollectionEditor/Private/GeometryCollection/AssetDefinition_GeometryCollection.cpp b/Engine/Plugins/Experimental/GeometryCollectionPlugin/Source/GeometryCollectionEditor/Private/GeometryCollection/AssetDefinition_GeometryCollection.cpp
--- a/Engine/Plugins/Experimental/GeometryCollectionPlugin/Source/GeometryCollectionEditor/Private/GeometryCollection/AssetDefinition_GeometryCollection.cpp
+++ b/Engine/Plugins/Experimental/GeometryCollectionPlugin/Source/GeometryCollectionEditor/Private/GeometryCollection/AssetDefinition_GeometryCollection.cpp
@@ -27,9 +27,6 @@
 bool bCanEditGeometryCollection = true;
 FAutoConsoleVariableRef CVarGeometryCollectionIsEditable(TEXT("p.Chaos.GC.IsEditable"), bCanEditGeometryCollection, TEXT("Whether to allow edits of the geometry collection"));
 
-bool bGeometryCollectionOpenEditorAskForDataflow = true;
-FAutoConsoleVariableRef CVarGeometryCollectionOpenEditorAskForDataflow(TEXT("p.Chaos.GC.OpenEditorAskForDataflow"), bGeometryCollectionOpenEditorAskForDataflow, TEXT("Whether to ask to create add a Dataflow graph when opening a Geometry Collection asset editor"));
-
 namespace UE::GeometryCollection
 {
 	struct FColorScheme
@@ -87,17 +84,6 @@ EAssetCommandResult UAssetDefinition_GeometryCollection::OpenAssets(const FAsset
 		// Validate the asset
 		if (UGeometryCollection* const GeometryCollection = Cast<UGeometryCollection>(GeometryCollectionObjects[0]))
 		{
-			if (bGeometryCollectionOpenEditorAskForDataflow)
-			{
-				if (!FDataflowEditorToolkit::HasDataflowAsset(GeometryCollection))
-				{
-					if (UDataflow* const NewDataflowAsset = Cast<UDataflow>(UE::DataflowAssetDefinitionHelpers::NewOrOpenDataflowAsset(GeometryCollection)))
-					{
-						GeometryCollection->DataflowInstance.SetDataflowAsset(NewDataflowAsset);
-					}
-				}
-			}
-
 			if (FDataflowEditorToolkit::HasDataflowAsset(GeometryCollection))
 			{
 				UAssetEditorSubsystem* const AssetEditorSubsystem = GEditor->GetEditorSubsystem<UAssetEditorSubsystem>();
diff --git a/Engine/Plugins/Experimental/GeometryCollectionPlugin/Source/GeometryCollectionEditor/Private/GeometryCollection/GeometryCollectionEditorPlugin.cpp b/Engine/Plugins/Experimental/GeometryCollectionPlugin/Source/GeometryCollectionEditor/Private/GeometryCollection/GeometryCollectionEditorPlugin.cpp
--- a/Engine/Plugins/Experimental/GeometryCollectionPlugin/Source/GeometryCollectionEditor/Private/GeometryCollection/GeometryCollectionEditorPlugin.cpp
+++ b/Engine/Plugins/Experimental/GeometryCollectionPlugin/Source/GeometryCollectionEditor/Private/GeometryCollection/GeometryCollectionEditorPlugin.cpp
@@ -8,6 +8,7 @@
 #include "LevelEditor.h"
 #include "Modules/ModuleManager.h"
 #include "PropertyEditorModule.h"
+#include "Dataflow/DataflowTemplateRegistry.h"
 #include "GeometryCollection/GeometryCollectionComponent.h"
 #include "GeometryCollection/GeometryCollectionConversion.h"
 #include "GeometryCollection/GeometryCollectionCommands.h"
@@ -15,7 +16,6 @@
 #include "GeometryCollection/AssetTypeActions_GeometryCollectionCache.h"
 #include "GeometryCollection/GeometryCollectionAssetBroker.h"
 #include "GeometryCollection/GeometryCollectionEditorStyle.h"
-#include "GeometryCollection/GeometryCollectionConversion.h"
 #include "GeometryCollection/GeometryCollectionThumbnailRenderer.h"
 #include "GeometryCollection/GeometryCollectionSelectRigidBodyEdMode.h"
 #include "GeometryCollection/GeometryCollectionSelectionCommands.h"
@@ -261,6 +261,9 @@ void IGeometryCollectionEditorPlugin::StartupModule()
 	// Register tool menus - delayed until ToolMenus module is loaded
 	UToolMenus::RegisterStartupCallback(
 		FSimpleMulticastDelegate::FDelegate::CreateRaw(this, &IGeometryCollectionEditorPlugin::RegisterMenus));
+
+	FDataflowTemplateRegistry::Get().RegisterTemplateFolder(UGeometryCollection::StaticClass(),
+		{ TEXT("/GeometryCollectionPlugin/DataflowTemplates"), FSlateIcon("GeometryCollectionEditorStyle", "ClassThumbnail.GeometryCollection") });
 }
 
 
diff --git a/Engine/Plugins/Experimental/GeometryCollectionPlugin/Source/GeometryCollectionEditor/Private/GeometryCollection/GeometryCollectionFactory.cpp b/Engine/Plugins/Experimental/GeometryCollectionPlugin/Source/GeometryCollectionEditor/Private/GeometryCollection/GeometryCollectionFactory.cpp
--- a/Engine/Plugins/Experimental/GeometryCollectionPlugin/Source/GeometryCollectionEditor/Private/GeometryCollection/GeometryCollectionFactory.cpp
+++ b/Engine/Plugins/Experimental/GeometryCollectionPlugin/Source/GeometryCollectionEditor/Private/GeometryCollection/GeometryCollectionFactory.cpp
@@ -2,6 +2,7 @@
 
 #include "GeometryCollection/GeometryCollectionFactory.h"
 
+#include "Dataflow/AssetDefinition_DataflowAsset.h"
 #include "GeometryCollection/GeometryCollectionAlgo.h"
 #include "GeometryCollection/GeometryCollectionConversion.h"
 #include "GeometryCollection/GeometryCollectionObject.h"
@@ -188,14 +189,16 @@ UObject* UGeometryCollectionFactory::FactoryCreateNew(UClass* Class, UObject* In
 		}
 	}
 
+	UE::DataflowAssetDefinitionHelpers::SetDataflowFromTemplatePicker(NewGeometryCollection);
+
 	// Assign a default dataflow template if available 
-	if (const UDataflow* DataflowTemplateObject = LoadObject<UDataflow>(nullptr, TEXT("/GeometryCollectionPlugin/DF_GC_Template_StaticMesh.DF_GC_Template_StaticMesh")))
+	/*if (const UDataflow* DataflowTemplateObject = LoadObject<UDataflow>(nullptr, TEXT("/GeometryCollectionPlugin/DF_GC_Template_StaticMesh.DF_GC_Template_StaticMesh")))
 	{
 		if (UDataflow* DataflowEmbeddedAsset = DuplicateObject<UDataflow>(DataflowTemplateObject, NewGeometryCollection, TEXT("EmbeddedDataflow")))
 		{
 			NewGeometryCollection->GetDataflowInstance().SetDataflowAsset(DataflowEmbeddedAsset);
 		}
-	}
+	}*/
 
 	NewGeometryCollection->Modify();
 	return NewGeometryCollection;
