Make the FSkinMeshViaMedialSkeleton dataflow node use skin-weight binding algorithm + options, rather than rigid binding to clusters

#Synced-CL 52392718

diff --git a/Engine/Plugins/Experimental/GeometryCollectionPlugin/Source/GeometryCollectionNodes/Private/Dataflow/MeshMedialSkeletonConversionNodes.cpp b/Engine/Plugins/Experimental/GeometryCollectionPlugin/Source/GeometryCollectionNodes/Private/Dataflow/MeshMedialSkeletonConversionNodes.cpp
--- a/Engine/Plugins/Experimental/GeometryCollectionPlugin/Source/GeometryCollectionNodes/Private/Dataflow/MeshMedialSkeletonConversionNodes.cpp
+++ b/Engine/Plugins/Experimental/GeometryCollectionPlugin/Source/GeometryCollectionNodes/Private/Dataflow/MeshMedialSkeletonConversionNodes.cpp
@@ -5,6 +5,7 @@
 #include "DynamicMesh/DynamicBoneAttribute.h"
 #include "DynamicMesh/DynamicMeshAttributeSet.h"
 #include "DynamicMesh/DynamicVertexSkinWeightsAttribute.h"
+#include "Operations/MedialSkeletonSkinBinding.h"
 #include "Operations/SkinWeightBinding.h"
 #include "UDynamicMesh.h"
 
@@ -239,6 +240,10 @@ FSkinMeshViaMedialSkeleton::FSkinMeshViaMedialSkeleton(const UE::Dataflow::FNode
 	RegisterInputConnection(&Mesh);
 	RegisterInputConnection(&MedialSkeleton);
 	RegisterInputConnection(&RootClusterIndex).SetCanHidePin(true).SetPinIsHidden(true);
+	RegisterInputConnection(&Stiffness).SetCanHidePin(true).SetPinIsHidden(true);
+	RegisterInputConnection(&MaxInfluences).SetCanHidePin(true).SetPinIsHidden(true);
+	RegisterInputConnection(&VoxelResolution).SetCanHidePin(true).SetPinIsHidden(true);
+	RegisterInputConnection(&ClusterNeighborSearchRange).SetCanHidePin(true).SetPinIsHidden(true);
 
 	RegisterOutputConnection(&Mesh, &Mesh);
 }
@@ -246,6 +251,7 @@ FSkinMeshViaMedialSkeleton::FSkinMeshViaMedialSkeleton(const UE::Dataflow::FNode
 void FSkinMeshViaMedialSkeleton::Evaluate(UE::Dataflow::FContext& Context, const FDataflowOutput* Out) const
 {
 	using namespace UE::Geometry;
+	using namespace MeshMedialSkeletonConversionLocals;
 
 	if (Out->IsA(&Mesh))
 	{
@@ -259,11 +265,46 @@ void FSkinMeshViaMedialSkeleton::Evaluate(UE::Dataflow::FContext& Context, const
 			OutMesh->SetMesh(InMesh->GetMeshRef());
 			FDynamicMesh3& DynMesh = OutMesh->GetMeshRef();
 
-			UE::Geometry::MedialAxis::FMedialSkeletonToTreeSkeletonOptions Options;
-			MeshMedialSkeletonConversionLocals::SetSkeletonConversionOptions(Options, EdgeWeightMethod, MergeDisconnectedMethod, SelectRootMethod, RootSelectionPoint, RootSelectionDirection);
-			const int32 UseRootIndex = GetValue(Context, &RootClusterIndex);
-			Options.ToDynamicMeshSkeleton(InMedialSkeleton.Skeleton, DynMesh, FSkeletalMeshAttributes::DefaultSkinWeightProfileName, Options, UseRootIndex);
-	
+			if (InMedialSkeleton.Skeleton.Spheres.IsEmpty())
+			{
+				Context.Error(LOCTEXT("SkinMeshViaMedialSkeletonEmptyError", "SkinMeshViaMedialSkeleton: Input medial skeleton is empty."), this);
+			}
+			else
+			{
+				SkinBinding::FBindMedialSkeletonSettings Settings;
+				Settings.ClusterNeighborSearchRange = GetValue(Context, &ClusterNeighborSearchRange);
+				SetSkeletonConversionOptions(Settings.AnimationSkeletonOptions, EdgeWeightMethod, MergeDisconnectedMethod, SelectRootMethod, RootSelectionPoint, RootSelectionDirection);
+				const int32 UseRootIndex = GetValue(Context, &RootClusterIndex);
+
+				SetSkinBindSettings(Settings.BindSettings, BindMethod,
+					GetValue(Context, &Stiffness),
+					GetValue(Context, &MaxInfluences),
+					GetValue(Context, &VoxelResolution));
+
+				bool bMeshWasCompatible = false;
+				bool bSuccess = SkinBinding::CreateSkinWeightsFromMedialSkeleton(
+					InMedialSkeleton.Skeleton,
+					DynMesh,
+					bMeshWasCompatible,
+					FSkeletalMeshAttributes::DefaultSkinWeightProfileName,
+					Settings,
+					UseRootIndex);
+
+				if (!bSuccess)
+				{
+					Context.Error(
+						LOCTEXT("SkinMeshViaMedialSkeletonFailed",
+							"SkinMeshViaMedialSkeleton: Failed to bind a skeleton."),
+						this);
+				}
+				else if (!bMeshWasCompatible)
+				{
+					Context.Warning(
+						LOCTEXT("SkinMeshViaMedialSkeletonIncompat",
+							"SkinMeshViaMedialSkeleton: Medial skeleton VIDtoClusterIndex is not compatible with the target mesh. Falling back to unconstrained skin binding."),
+						this);
+				}
+			}
 		}
 		else
 		{
diff --git a/Engine/Plugins/Experimental/GeometryCollectionPlugin/Source/GeometryCollectionNodes/Private/Dataflow/MeshMedialSkeletonConversionNodes.h b/Engine/Plugins/Experimental/GeometryCollectionPlugin/Source/GeometryCollectionNodes/Private/Dataflow/MeshMedialSkeletonConversionNodes.h
--- a/Engine/Plugins/Experimental/GeometryCollectionPlugin/Source/GeometryCollectionNodes/Private/Dataflow/MeshMedialSkeletonConversionNodes.h
+++ b/Engine/Plugins/Experimental/GeometryCollectionPlugin/Source/GeometryCollectionNodes/Private/Dataflow/MeshMedialSkeletonConversionNodes.h
@@ -184,6 +184,30 @@ private:
 	UPROPERTY(EditAnywhere, Category = "Options", meta = (DataflowInput))
 	int32 RootClusterIndex = INDEX_NONE;
 
+	/** Binding method to use for skin weights */
+	UPROPERTY(EditAnywhere, Category = "Binding")
+	EDataflowBindSkeletonMethod BindMethod = EDataflowBindSkeletonMethod::DirectDistance;
+
+	/** Stiffness of binding. Lower values allow more distant bones to contribute more */
+	UPROPERTY(EditAnywhere, Category = "Binding", meta = (DataflowInput, ClampMin=0, ClampMax=1))
+	float Stiffness = 0.2f;
+
+	/** Maximum bones that will influence each vertex */
+	UPROPERTY(EditAnywhere, Category = "Binding", meta = (DataflowInput, ClampMin = 1, UIMax = 10))
+	int32 MaxInfluences = 5;
+
+	/** The resolution of the voxel grid if doing geodesic voxel binding */
+	UPROPERTY(EditAnywhere, Category = "Binding", meta = (DataflowInput, ClampMin = 1, UIMin = 8, ClampMax = 1024))
+	int32 VoxelResolution = 256;
+
+	/** 
+	 * How far to search through medial skeleton neighbors when building bone groups for constrained skin binding.
+	 * 1 = immediate neighbors only, 2 = neighbors of neighbors, etc.
+	 * Higher values allow smoother blending across more bones but reduce locality. 
+	 */
+	UPROPERTY(EditAnywhere, Category = "Binding", meta = (DataflowInput, ClampMin = 1, UIMax = 5, ClampMax = 100))
+	int32 ClusterNeighborSearchRange = 1;
+
 	virtual void Evaluate(UE::Dataflow::FContext& Context, const FDataflowOutput* Out) const override;
 
 
