add sphere overlap threshold to simplify skeleton dataflow node

#Synced-CL 52426395

diff --git a/Engine/Plugins/Experimental/GeometryCollectionPlugin/Source/GeometryCollectionNodes/Private/Dataflow/GeometryCollectionSamplingNodes.cpp b/Engine/Plugins/Experimental/GeometryCollectionPlugin/Source/GeometryCollectionNodes/Private/Dataflow/GeometryCollectionSamplingNodes.cpp
--- a/Engine/Plugins/Experimental/GeometryCollectionPlugin/Source/GeometryCollectionNodes/Private/Dataflow/GeometryCollectionSamplingNodes.cpp
+++ b/Engine/Plugins/Experimental/GeometryCollectionPlugin/Source/GeometryCollectionNodes/Private/Dataflow/GeometryCollectionSamplingNodes.cpp
@@ -444,6 +444,7 @@ FSimplifyMedialSkeletonDataflowNode::FSimplifyMedialSkeletonDataflowNode(const U
 	RegisterInputConnection(&QEMErrorThresholdSqrt).SetCanHidePin(true).SetPinIsHidden(true);
 	RegisterInputConnection(&EdgeLengthThreshold).SetCanHidePin(true).SetPinIsHidden(true);
 	RegisterInputConnection(&SphereRadiusThreshold).SetCanHidePin(true).SetPinIsHidden(true);
+	RegisterInputConnection(&SphereOverlapThreshold).SetCanHidePin(true).SetPinIsHidden(true);
 	RegisterInputConnection(&bOnlySimplifySurfaces).SetCanHidePin(true).SetPinIsHidden(true);
 	RegisterInputConnection(&ClusterSkeletonDistanceWt).SetCanHidePin(true).SetPinIsHidden(true);
 	RegisterInputConnection(&bPreventEdgeSurfaceIntersections).SetCanHidePin(true).SetPinIsHidden(true);
@@ -478,6 +479,7 @@ void FSimplifyMedialSkeletonDataflowNode::Evaluate(UE::Dataflow::FContext& Conte
 				Simplifier.QEMErrorThreshold = UseQEMThresholdSqrt * UseQEMThresholdSqrt;
 				Simplifier.EdgeLengthThreshold = GetValue(Context, &EdgeLengthThreshold);
 				Simplifier.SphereRadiusThreshold = GetValue(Context, &SphereRadiusThreshold);
+				Simplifier.SphereOverlapThreshold = GetValue(Context, &SphereOverlapThreshold);
 				Simplifier.bOnlySimplifySurfaces = GetValue(Context, &bOnlySimplifySurfaces);
 				Simplifier.ClusterSkeletonDistanceWt = GetValue(Context, &ClusterSkeletonDistanceWt);
 				Simplifier.bPreventEdgeSurfaceIntersections = GetValue(Context, &bPreventEdgeSurfaceIntersections);
diff --git a/Engine/Plugins/Experimental/GeometryCollectionPlugin/Source/GeometryCollectionNodes/Public/Dataflow/GeometryCollectionSamplingNodes.h b/Engine/Plugins/Experimental/GeometryCollectionPlugin/Source/GeometryCollectionNodes/Public/Dataflow/GeometryCollectionSamplingNodes.h
--- a/Engine/Plugins/Experimental/GeometryCollectionPlugin/Source/GeometryCollectionNodes/Public/Dataflow/GeometryCollectionSamplingNodes.h
+++ b/Engine/Plugins/Experimental/GeometryCollectionPlugin/Source/GeometryCollectionNodes/Public/Dataflow/GeometryCollectionSamplingNodes.h
@@ -438,6 +438,11 @@ private:
 	UPROPERTY(EditAnywhere, Category = Options, meta = (DataflowInput, ClampMin = 0, Units = "cm"))
 	float SphereRadiusThreshold = 0.;
 
+	// If positive, the minimum medial sphere overlap fraction to consider for collapse -- so e.g. at 1.0 we will only collapse an edge if its smaller medial sphere is contained in the larger.
+	// (A degenerate zero-radius sphere is considered to fully overlap any sphere it touches.)
+	UPROPERTY(EditAnywhere, Category = Options, meta = (DataflowInput, ClampMin = 0, ClampMax = 1))
+	float SphereOverlapThreshold = 0.;
+
 	// If true, only simplify edges of skeleton surfaces (edges that have triangles associated)
 	UPROPERTY(EditAnywhere, Category = Options, meta = (DataflowInput))
 	bool bOnlySimplifySurfaces = true;
