Extract/Report groom card attributes to dataflow + make the attributeToSelection more generic

#virtualized
#Synced-CL 52466270

diff --git a/Engine/Plugins/Experimental/GeometryCollectionPlugin/Source/GeometryCollectionNodes/Private/Dataflow/GeometryCollectionSelectionNodes.cpp b/Engine/Plugins/Experimental/GeometryCollectionPlugin/Source/GeometryCollectionNodes/Private/Dataflow/GeometryCollectionSelectionNodes.cpp
--- a/Engine/Plugins/Experimental/GeometryCollectionPlugin/Source/GeometryCollectionNodes/Private/Dataflow/GeometryCollectionSelectionNodes.cpp
+++ b/Engine/Plugins/Experimental/GeometryCollectionPlugin/Source/GeometryCollectionNodes/Private/Dataflow/GeometryCollectionSelectionNodes.cpp
@@ -3066,6 +3066,111 @@ void FCollectionSelectionByMeshDataflowNode::Evaluate(UE::Dataflow::FContext& Co
 
 /* ----------------------------------------------------------------------------------------------------------------------- */
 
+namespace UE::Dataflow::Private
+{
+	template<typename AttributeType, typename = void>
+	struct THasLexFromString : std::false_type {};
+
+	template<typename AttributeType>
+	struct THasLexFromString<AttributeType, std::void_t<decltype(LexFromString(DeclVal<AttributeType&>(), DeclVal<const TCHAR*>()))>>
+		: std::true_type {};
+
+	template<typename AttributeType>
+	void SetAttributeFromString(FManagedArrayCollection& Collection, const FName& AttributeName, const FName& GroupName, 
+		const FString& SelectionString, const FString& DefaultString, const TArray<int32>& SelectedElements)
+	{
+		if constexpr (THasLexFromString<AttributeType>::value)
+		{
+			AttributeType SelectionValue{};
+			LexFromString(SelectionValue, *SelectionString);
+		
+			AttributeType DefaultValue{};
+			LexFromString(DefaultValue, *DefaultString);
+		
+			const bool bHasAttribute = Collection.HasAttribute(AttributeName, GroupName);
+			if (TManagedArray<AttributeType>* AttributeArray = Collection.FindOrAddAttributeTyped<AttributeType>(AttributeName, GroupName))
+			{
+				if (!bHasAttribute)
+				{
+					AttributeArray->Fill(DefaultValue);
+				}
+				if (!SelectedElements.IsEmpty())
+				{
+					for (int32 ElementIndex : SelectedElements)
+					{
+						if (AttributeArray->IsValidIndex(ElementIndex))
+						{
+							(*AttributeArray)[ElementIndex] = SelectionValue;
+						}
+					}
+				}
+			}
+		}
+	}
+	
+	FManagedArrayCollection::EArrayType GetAttributeArrayType(const ESelectionToAttributeTypeFlags& TypeFlag)
+	{
+		FManagedArrayCollection::EArrayType ArrayType = FManagedArrayCollection::EArrayType::FFloatType;
+		switch (TypeFlag)
+		{
+			case ESelectionToAttributeTypeFlags::Int:
+				ArrayType = FManagedArrayCollection::EArrayType::FInt32Type;
+				break;
+			case ESelectionToAttributeTypeFlags::Bool:
+				ArrayType = FManagedArrayCollection::EArrayType::FBoolType;
+				break;
+			case ESelectionToAttributeTypeFlags::Float:
+				ArrayType = FManagedArrayCollection::EArrayType::FFloatType;
+				break;
+			case ESelectionToAttributeTypeFlags::Double:
+				ArrayType = FManagedArrayCollection::EArrayType::FDoubleType;
+				break;
+			case ESelectionToAttributeTypeFlags::Vec2f:
+				ArrayType = FManagedArrayCollection::EArrayType::FVector2DType;
+				break;
+			case ESelectionToAttributeTypeFlags::Vec3f:
+				ArrayType = FManagedArrayCollection::EArrayType::FVectorType;
+				break;
+			case ESelectionToAttributeTypeFlags::Vec4f:
+				ArrayType = FManagedArrayCollection::EArrayType::FVector4fType;
+				break;
+			case ESelectionToAttributeTypeFlags::Vec2i:
+				ArrayType = FManagedArrayCollection::EArrayType::FIntVector2Type;
+				break;
+			case ESelectionToAttributeTypeFlags::Vec3i:
+				ArrayType = FManagedArrayCollection::EArrayType::FIntVectorType;
+				break;
+			case ESelectionToAttributeTypeFlags::Vec4i:
+				ArrayType = FManagedArrayCollection::EArrayType::FIntVector4Type;
+				break;
+			case ESelectionToAttributeTypeFlags::Color:
+				ArrayType = FManagedArrayCollection::EArrayType::FLinearColorType;
+				break;
+			case ESelectionToAttributeTypeFlags::Quat:
+				ArrayType = FManagedArrayCollection::EArrayType::FQuatType;
+				break;
+			case ESelectionToAttributeTypeFlags::String:
+				ArrayType = FManagedArrayCollection::EArrayType::FStringType;
+				break;
+			case ESelectionToAttributeTypeFlags::Name:
+				ArrayType = FManagedArrayCollection::EArrayType::FNameType;
+				break;
+			case ESelectionToAttributeTypeFlags::Transform:
+				ArrayType = FManagedArrayCollection::EArrayType::FTransformType;
+				break;
+			case ESelectionToAttributeTypeFlags::Matrix:
+				ArrayType = FManagedArrayCollection::EArrayType::FPMatrix33dType;
+				break;
+			case ESelectionToAttributeTypeFlags::Box:
+				ArrayType = FManagedArrayCollection::EArrayType::FBoxType;
+				break;
+			default:
+				break;
+		}
+		return ArrayType;
+	}
+}
+
 FCollectionSelectionToAttributeDataflowNode::FCollectionSelectionToAttributeDataflowNode(const UE::Dataflow::FNodeParameters& InParam, FGuid InGuid)
 	: FDataflowNode(InParam, InGuid)
 {
@@ -3076,7 +3181,6 @@ FCollectionSelectionToAttributeDataflowNode::FCollectionSelectionToAttributeData
 	RegisterInputConnection(&DefaultValue).SetCanHidePin(true).SetPinIsHidden(true);
 	RegisterInputConnection(&Value).SetCanHidePin(true).SetPinIsHidden(true);
 
-
 	RegisterOutputConnection(&Collection, &Collection);
 	RegisterOutputConnection(&Attribute, &Attribute);
 }
@@ -3093,105 +3197,28 @@ void FCollectionSelectionToAttributeDataflowNode::Evaluate(UE::Dataflow::FContex
 			const FDataflowInput* SelectionInput = FindInput(&Selection);
 			if (SelectionInput)
 			{
-				const FName InGroupName = FDataflowSelection::GetSelectionTypeFromCollection(*SelectionInput);
-
-				if (InGroupName != NAME_None && InCollection.HasGroup(InGroupName))
+				const FName GroupName = FDataflowSelection::GetSelectionTypeFromCollection(*SelectionInput);
+				if (GroupName != NAME_None && InCollection.HasGroup(GroupName))
 				{
-					if (InCollection.NumElements(InGroupName) == InSelection.Num())
+					if (InCollection.NumElements(GroupName) == InSelection.Num())
 					{
-						const FName AttributeName = FName(*Attribute);
-
-						if (InCollection.HasAttribute(AttributeName, InGroupName))
-						{
-							const FManagedArrayCollection::EArrayType ArrayType = InCollection.GetAttributeType(AttributeName, InGroupName);
-
-							if (ArrayType == FManagedArrayCollection::EArrayType::FInt32Type)
-							{
-								TManagedArray<int32>* Attr = InCollection.FindAttribute<int32>(AttributeName, InGroupName);
-
-								for (int32 Idx = 0; Idx < InSelection.Num(); ++Idx)
-								{
-									if (InSelection.IsSelected(Idx))
-									{
-										(*Attr)[Idx] = FCString::Atoi(*Value);
-									}
-								}
-							}
-							else if (ArrayType == FManagedArrayCollection::EArrayType::FFloatType)
-							{
-								TManagedArray<float>* Attr = InCollection.FindAttribute<float>(AttributeName, InGroupName);
-
-								for (int32 Idx = 0; Idx < InSelection.Num(); ++Idx)
-								{
-									if (InSelection.IsSelected(Idx))
-									{
-										(*Attr)[Idx] = FCString::Atof(*Value);
-									}
-								}
-							}
-							else if (ArrayType == FManagedArrayCollection::EArrayType::FBoolType)
-							{
-								TManagedArray<bool>* Attr = InCollection.FindAttribute<bool>(AttributeName, InGroupName);
-
-								for (int32 Idx = 0; Idx < InSelection.Num(); ++Idx)
-								{
-									if (InSelection.IsSelected(Idx))
-									{
-										(*Attr)[Idx] = Value.ToBool();
-									}
-								}
-							}
-						}
-						else
+						const FName AttributeName = FName(GetValue(Context, &Attribute));
+						
+						const FString SelectionString = GetValue(Context, &Value);
+						const FString DefaultString = GetValue(Context, &DefaultValue);
+						
+						const FManagedArrayCollection::EArrayType ArrayType = InCollection.HasAttribute(AttributeName, GroupName) ? 
+							InCollection.GetAttributeType(AttributeName, GroupName) : UE::Dataflow::Private::GetAttributeArrayType(AttributeType);
+						
+						switch (ArrayType)
 						{
-							if (AttributeType == ESelectionToAttributeTypeFlags::Int)
-							{
-								TManagedArray<int32>& Attr = InCollection.AddAttribute<int32>(AttributeName, InGroupName);
-
-								for (int32 Idx = 0; Idx < InSelection.Num(); ++Idx)
-								{
-									if (InSelection.IsSelected(Idx))
-									{
-										Attr[Idx] = FCString::Atoi(*Value);
-									}
-									else
-									{
-										Attr[Idx] = FCString::Atoi(*DefaultValue);
-									}
-								}
-							}
-							else if (AttributeType == ESelectionToAttributeTypeFlags::Float)
-							{
-								TManagedArray<float>& Attr = InCollection.AddAttribute<float>(AttributeName, InGroupName);
-
-								for (int32 Idx = 0; Idx < InSelection.Num(); ++Idx)
-								{
-									if (InSelection.IsSelected(Idx))
-									{
-										Attr[Idx] = FCString::Atof(*Value);
-									}
-									else
-									{
-										Attr[Idx] = FCString::Atof(*DefaultValue);
-									}
-								}
-							}
-							else if (AttributeType == ESelectionToAttributeTypeFlags::Bool)
-							{
-								TManagedArray<bool>& Attr = InCollection.AddAttribute<bool>(AttributeName, InGroupName);
-
-								for (int32 Idx = 0; Idx < InSelection.Num(); ++Idx)
-								{
-									if (InSelection.IsSelected(Idx))
-									{
-										Attr[Idx] = Value.ToBool();
-									}
-									else
-									{
-										Attr[Idx] = DefaultValue.ToBool();
-									}
-								}
-							}
+#define MANAGED_ARRAY_TYPE(CppType, EnumName) \
+							case EManagedArrayType::F##EnumName##Type: \
+								UE::Dataflow::Private::SetAttributeFromString<CppType>(InCollection, AttributeName, GroupName, SelectionString, DefaultString, InSelection.AsArray()); \
+							break;
+#include "GeometryCollection/ManagedArrayTypeValues.inl"
+#undef MANAGED_ARRAY_TYPE
+							default: break;
 						}
 			
 						SetValue(Context, MoveTemp(InCollection), &Collection);
diff --git a/Engine/Plugins/Experimental/GeometryCollectionPlugin/Source/GeometryCollectionNodes/Public/Dataflow/GeometryCollectionSelectionNodes.h b/Engine/Plugins/Experimental/GeometryCollectionPlugin/Source/GeometryCollectionNodes/Public/Dataflow/GeometryCollectionSelectionNodes.h
--- a/Engine/Plugins/Experimental/GeometryCollectionPlugin/Source/GeometryCollectionNodes/Public/Dataflow/GeometryCollectionSelectionNodes.h
+++ b/Engine/Plugins/Experimental/GeometryCollectionPlugin/Source/GeometryCollectionNodes/Public/Dataflow/GeometryCollectionSelectionNodes.h
@@ -2273,7 +2273,21 @@ enum class ESelectionToAttributeTypeFlags : uint8
 {
 	Int UMETA(DisplayName = "Int"),
 	Float UMETA(DisplayName = "Float"),
-	Bool UMETA(DisplayName = "Bool")
+	Bool UMETA(DisplayName = "Bool"),
+	Double UMETA(DisplayName = "Double"),
+	String UMETA(DisplayName = "String"),
+	Name UMETA(DisplayName = "Name"),
+	Vec2f UMETA(DisplayName = "Vec2f"),
+	Vec3f UMETA(DisplayName = "Vec3f"),
+	Vec4f UMETA(DisplayName = "Vec4f"),
+	Vec2i UMETA(DisplayName = "Vec2i"),
+	Vec3i UMETA(DisplayName = "Vec3i"),
+	Vec4i UMETA(DisplayName = "Vec4i"),
+	Color UMETA(DisplayName = "Color"),
+	Transform UMETA(DisplayName = "Transform"),
+	Matrix UMETA(DisplayName = "Matrix"),
+	Quat UMETA(DisplayName = "Quat"),
+	Box UMETA(DisplayName = "Box"),
 };
 
 /**
@@ -2281,7 +2295,6 @@ enum class ESelectionToAttributeTypeFlags : uint8
  * Creates specified attribute on the incoming selection's group and sets specified value to selected elements,
  * the nonselected elements will be set to default valu. 
  * If the attribute already exists then sets specified value to selected elements.
- * Currently supported attribute types: float, int32, bool
  *
  */
 USTRUCT(meta = (DataflowGeometryCollection))
@@ -2299,12 +2312,12 @@ public:
 	UPROPERTY(meta = (DataflowInput, DataflowIntrinsic))
 	FDataflowSelectionTypes Selection;
 
-	/** Attribute type for storing the selection */
+	/** Attribute type for storing the selection if not already in the collection */
 	UPROPERTY(EditAnywhere, Category = "Selection")
 	ESelectionToAttributeTypeFlags AttributeType = ESelectionToAttributeTypeFlags::Int;
 
 	/** Attribute for storing the selection */
-	UPROPERTY(EditAnywhere, Category = "Selection", meta = (DataflowInput, DataflowOutput, DataflowPassthrough = "Attribute"))
+	UPROPERTY(EditAnywhere, Category = "Selection", meta = (DataflowInput, DataflowOutput, DataflowPassthrough = "Attribute", DisplayName = "Attribute Name"))
 	FString Attribute = FString("Dummy");
 
 	/** Attribute default value */
@@ -2312,7 +2325,7 @@ public:
 	FString DefaultValue = FString("0");
 
 	/** Attribute value */
-	UPROPERTY(EditAnywhere, Category = "Selection", meta = (DataflowInput))
+	UPROPERTY(EditAnywhere, Category = "Selection", meta = (DataflowInput, DisplayName = "Selection Value"))
 	FString Value = FString("1");
 
 	FCollectionSelectionToAttributeDataflowNode(const UE::Dataflow::FNodeParameters& InParam, FGuid InGuid = FGuid::NewGuid());
diff --git a/Engine/Plugins/Experimental/HairCardGenerator/Source/HairCardGeneratorDataflow/HairCardGeneratorDataflow.Build.cs b/Engine/Plugins/Experimental/HairCardGenerator/Source/HairCardGeneratorDataflow/HairCardGeneratorDataflow.Build.cs
--- a/Engine/Plugins/Experimental/HairCardGenerator/Source/HairCardGeneratorDataflow/HairCardGeneratorDataflow.Build.cs
+++ b/Engine/Plugins/Experimental/HairCardGenerator/Source/HairCardGeneratorDataflow/HairCardGeneratorDataflow.Build.cs
@@ -21,6 +21,7 @@ namespace UnrealBuildTool.Rules
 					"DataflowEngine",
 					"DataflowEditor",
 					"DataflowNodes",
+					"DataflowEnginePlugin",
 					"HairStrandsCore",
 					"HairStrandsDataflow",
 					"HairStrandsEditor",
diff --git a/Engine/Plugins/Experimental/HairCardGenerator/Source/HairCardGeneratorDataflow/Private/BuildCardsSettingsNode.cpp b/Engine/Plugins/Experimental/HairCardGenerator/Source/HairCardGeneratorDataflow/Private/BuildCardsSettingsNode.cpp
--- a/Engine/Plugins/Experimental/HairCardGenerator/Source/HairCardGeneratorDataflow/Private/BuildCardsSettingsNode.cpp
+++ b/Engine/Plugins/Experimental/HairCardGenerator/Source/HairCardGeneratorDataflow/Private/BuildCardsSettingsNode.cpp
@@ -1,6 +1,8 @@
 // Copyright Epic Games, Inc. All Rights Reserved.
 
 #include "BuildCardsSettingsNode.h"
+
+#include "GroomBuilder.h"
 #include "HairCardGeneratorEditorModule.h"
 #include "Dataflow/DataflowDebugDrawInterface.h"
 #include "GeometryCollection/Facades/CollectionConstraintOverrideFacade.h"
@@ -10,23 +12,37 @@
 
 namespace UE::CardGen::Private
 {
-	static void BuildStrandsPositions(const FManagedArrayCollection& GroomCollection, TArray<TArray<FVector>>& StrandsPositions)
+	static void BuildStrandsPositions(const FManagedArrayCollection& GroomCollection, TArray<TArray<FVector>>& StrandsPositions, const FHairDescription& HairDescription)
 	{
 		const GeometryCollection::Facades::FCollectionCurveGeometryFacade CurvesFacade(GroomCollection);
 		if (CurvesFacade.IsValid())
 		{
-			const int32 NumCurves = CurvesFacade.GetNumCurves();
-			StrandsPositions.SetNum(NumCurves);
+			FHairDescriptionGroups HairDescriptionGroups;
+			FGroomBuilder::BuildHairDescriptionGroups(HairDescription, HairDescriptionGroups, false);
+			
+			StrandsPositions.SetNum(HairDescription.GetNumStrands());
 
 			int32 PointOffset = 0;
+			const int32 NumCurves = CurvesFacade.GetNumCurves();
 			for(int32 CurveIndex = 0; CurveIndex < NumCurves; ++CurveIndex)
 			{
-				const int32 NumPoints = CurvesFacade.GetCurvePointOffsets()[CurveIndex]-PointOffset;
-				StrandsPositions[CurveIndex].SetNum(NumPoints);
+				const int32 GroupIndex = CurvesFacade.GetCurveGeometryIndices()[CurveIndex];
+				const int32 SourceCurve = CurvesFacade.GetCurveSourceIndices()[CurveIndex] >> 1;
 				
-				for(int32 PointIndex = 0; PointIndex < NumPoints; ++PointIndex)
+				if (HairDescriptionGroups.HairGroups.IsValidIndex(GroupIndex) && HairDescriptionGroups.HairGroups[GroupIndex].Strands.StrandsCurves.CurvesMapping.IsValidIndex(SourceCurve))
 				{
-					StrandsPositions[CurveIndex][PointIndex] = FVector(CurvesFacade.GetPointRestPositions()[PointIndex + PointOffset]);
+					const int32 StrandID = HairDescriptionGroups.HairGroups[GroupIndex].Strands.StrandsCurves.CurvesMapping[SourceCurve];
+				
+					const int32 NumPoints = CurvesFacade.GetCurvePointOffsets()[CurveIndex]-PointOffset;
+					if (StrandsPositions.IsValidIndex(StrandID))
+					{
+						StrandsPositions[StrandID].SetNum(NumPoints);
+				
+						for(int32 PointIndex = 0; PointIndex < NumPoints; ++PointIndex)
+						{
+							StrandsPositions[StrandID][PointIndex] = FVector(CurvesFacade.GetPointRestPositions()[PointIndex + PointOffset]);
+						}
+					}
 				}
 				PointOffset = CurvesFacade.GetCurvePointOffsets()[CurveIndex];
 			}
@@ -173,9 +189,9 @@ void FBuildCardsSettingsNode::Evaluate(UE::Dataflow::FContext& Context, const FD
 								}
 							}
 						}
-						FHairCardGeneratorUtils::LoadGroomStrands(LocalAsset, [&GeometryCollection](TArray<TArray<FVector>>& CurvesPositions)
+						FHairCardGeneratorUtils::LoadGroomStrands(LocalAsset, [&GeometryCollection,LocalAsset](TArray<TArray<FVector>>& CurvesPositions)
 						{
-							UE::CardGen::Private::BuildStrandsPositions(GeometryCollection, CurvesPositions);
+							UE::CardGen::Private::BuildStrandsPositions(GeometryCollection, CurvesPositions, LocalAsset->GetHairDescription());
 						});
 						OutputSettings.Last().GenerationSettings->UpdateStrandFilterAssignment();
 					}
diff --git a/Engine/Plugins/Experimental/HairCardGenerator/Source/HairCardGeneratorDataflow/Private/HairCardGeneratorDataflowModule.cpp b/Engine/Plugins/Experimental/HairCardGenerator/Source/HairCardGeneratorDataflow/Private/HairCardGeneratorDataflowModule.cpp
--- a/Engine/Plugins/Experimental/HairCardGenerator/Source/HairCardGeneratorDataflow/Private/HairCardGeneratorDataflowModule.cpp
+++ b/Engine/Plugins/Experimental/HairCardGenerator/Source/HairCardGeneratorDataflow/Private/HairCardGeneratorDataflowModule.cpp
@@ -6,6 +6,7 @@
 #include "GenerateCardsGeometryNode.h"
 #include "GenerateCardsTexturesNode.h"
 #include "CardsAssetTerminalNode.h"
+#include "ModifyCardsAttributesNode.h"
 #include "HairCardDataflowRendering.h"
 #include "Modules/ModuleManager.h"
 #include "HairStrandsFactory.h"
@@ -19,6 +20,8 @@ void FHairCardGeneratorDataflowModule::StartupModule()
 	DATAFLOW_NODE_REGISTER_CREATION_FACTORY(FGenerateCardsTexturesNode);
 	DATAFLOW_NODE_REGISTER_CREATION_FACTORY(FBuildCardsSettingsNode);
 	DATAFLOW_NODE_REGISTER_CREATION_FACTORY(FCardsAssetTerminalNode);
+	DATAFLOW_NODE_REGISTER_CREATION_FACTORY(FExtractCardsAttributesNode);
+	DATAFLOW_NODE_REGISTER_CREATION_FACTORY(FReportCardsAttributesNode);
 	
 	UE::CardGen::Private::RegisterCollectionRenderableTypes();
 	UE::Groom::RegisterGroomDataflowTemplate({ "SelectCardsDataflowTemplate", "Cards", "Add a Dataflow with a cards generation.", "/HairCardGenerator/Dataflow/Templates/DF_GroomCardsTemplate.DF_GroomCardsTemplate", false });
diff --git a/Engine/Plugins/Experimental/HairCardGenerator/Source/HairCardGeneratorDataflow/Private/ModifyCardsAttributesNode.cpp b/Engine/Plugins/Experimental/HairCardGenerator/Source/HairCardGeneratorDataflow/Private/ModifyCardsAttributesNode.cpp
new file mode 100644
--- /dev/null
+++ b/Engine/Plugins/Experimental/HairCardGenerator/Source/HairCardGeneratorDataflow/Private/ModifyCardsAttributesNode.cpp
@@ -0,0 +1,147 @@
+// Copyright Epic Games, Inc. All Rights Reserved.
+
+#include "ModifyCardsAttributesNode.h"
+
+#include "GroomBuilder.h"
+#include "GroomInstance.h"
+#include "Dataflow/DataflowInputOutput.h"
+#include "Dataflow/DataflowObjectInterface.h"
+#include "Dataflow/DataflowTypePolicy.h"
+#include "GeometryCollection/Facades/CollectionCurveFacade.h"
+#include "GeometryCollection/Facades/CollectionTransformSelectionFacade.h"
+
+#include UE_INLINE_GENERATED_CPP_BY_NAME(ModifyCardsAttributesNode)
+
+namespace UE::Groom::Private
+{
+	static void ExtractGroomAttributes(FManagedArrayCollection& GroomCollection, const FHairDescription& HairDescription)
+	{
+		const GeometryCollection::Facades::FCollectionCurveGeometryFacade GeometryFacade(GroomCollection);
+		GeometryCollection::Facades::FCollectionCurveAttributesFacade AttributesFacade(GroomCollection);
+		
+		const TStrandAttributesConstRef<FName> CardGroupsAttribute = HairDescription.StrandAttributes().GetAttributesRef<FName>(HairAttribute::Strand::GroupCardsName);
+		
+		if (AttributesFacade.IsValid() && GeometryFacade.IsValid() && CardGroupsAttribute.IsValid())
+		{
+			FHairDescriptionGroups HairDescriptionGroups;
+			FGroomBuilder::BuildHairDescriptionGroups(HairDescription, HairDescriptionGroups, false);
+		
+			const int32 NumCurves = GeometryFacade.GetNumCurves();
+			
+			TArray<FName> CardGroups;
+			CardGroups.SetNum(NumCurves);
+		
+			for(int32 CurveIndex = 0; CurveIndex < NumCurves; ++CurveIndex)
+			{
+				const int32 GroupIndex = GeometryFacade.GetCurveGeometryIndices()[CurveIndex];
+				const int32 SourceCurve = GeometryFacade.GetCurveSourceIndices()[CurveIndex] >> 1;
+			
+				if (HairDescriptionGroups.HairGroups.IsValidIndex(GroupIndex) && CardGroups.IsValidIndex(CurveIndex) &&
+					HairDescriptionGroups.HairGroups[GroupIndex].Strands.StrandsCurves.CurvesMapping.IsValidIndex(SourceCurve))
+				{
+					const int32 StrandID = HairDescriptionGroups.HairGroups[GroupIndex].Strands.StrandsCurves.CurvesMapping[SourceCurve];
+					if ((StrandID < CardGroupsAttribute.GetNumElements()) && (StrandID >= 0))
+					{
+						CardGroups[CurveIndex] = CardGroupsAttribute[FStrandID(StrandID)];
+					}
+				}
+			}
+			AttributesFacade.SetCurveCardGroups(CardGroups);
+		}
+	}
+	
+	static void ReportGroomAttributes(const FManagedArrayCollection& GroomCollection, FHairDescription& HairDescription)
+	{
+		const GeometryCollection::Facades::FCollectionCurveGeometryFacade GeometryFacade(GroomCollection);
+		const GeometryCollection::Facades::FCollectionCurveAttributesFacade AttributesFacade(GroomCollection);
+		
+		if (AttributesFacade.IsValid() && GeometryFacade.IsValid() && GeometryFacade.GetNumCurves() == AttributesFacade.GetNumCurves())
+		{
+			TStrandAttributesRef<FName> CardGroupsAttribute = HairDescription.StrandAttributes().GetAttributesRef<FName>(HairAttribute::Strand::GroupCardsName);
+			if (!CardGroupsAttribute.IsValid())
+			{
+				HairDescription.StrandAttributes().RegisterAttribute<FName>(HairAttribute::Strand::GroupCardsName);
+				CardGroupsAttribute = HairDescription.StrandAttributes().GetAttributesRef<FName>(HairAttribute::Strand::GroupCardsName);
+			}
+			
+			FHairDescriptionGroups HairDescriptionGroups;
+			FGroomBuilder::BuildHairDescriptionGroups(HairDescription, HairDescriptionGroups, false);
+		
+			const int32 NumCurves = GeometryFacade.GetNumCurves();
+			const TArray<FName>& CardGroups = AttributesFacade.GetCurveCardGroups();
+		
+			for(int32 CurveIndex = 0; CurveIndex < NumCurves; ++CurveIndex)
+			{
+				const int32 GroupIndex = GeometryFacade.GetCurveGeometryIndices()[CurveIndex];
+				const int32 SourceCurve = GeometryFacade.GetCurveSourceIndices()[CurveIndex] >> 1;
+			
+				if (HairDescriptionGroups.HairGroups.IsValidIndex(GroupIndex) && CardGroups.IsValidIndex(CurveIndex) &&
+					HairDescriptionGroups.HairGroups[GroupIndex].Strands.StrandsCurves.CurvesMapping.IsValidIndex(SourceCurve))
+				{
+					const int32 StrandID = HairDescriptionGroups.HairGroups[GroupIndex].Strands.StrandsCurves.CurvesMapping[SourceCurve];
+					if ((StrandID < CardGroupsAttribute.GetNumElements()) && (StrandID >= 0))
+					{
+						CardGroupsAttribute[FStrandID(StrandID)] = CardGroups[CurveIndex];
+					}
+				}
+			}
+		}
+	}
+}
+
+FExtractCardsAttributesNode::FExtractCardsAttributesNode(const UE::Dataflow::FNodeParameters& InParam, FGuid InGuid)
+		: FDataflowNode(InParam, InGuid)
+{
+	RegisterInputConnection(&Collection);
+	RegisterInputConnection(&GroomAsset);
+	RegisterOutputConnection(&Collection, &Collection);
+	RegisterOutputConnection(&CardGroups);
+}
+
+void FExtractCardsAttributesNode::Evaluate(UE::Dataflow::FContext& Context, const FDataflowOutput* Out) const
+{
+	if (Out->IsA<FManagedArrayCollection>(&Collection))
+	{
+		FManagedArrayCollection GroomCollection= GetValue<FManagedArrayCollection>(Context, &Collection);
+		if(TObjectPtr<const UGroomAsset> InputAsset = GetValue<TObjectPtr<const UGroomAsset>>(Context, &GroomAsset))
+		{
+			const FHairDescription HairDesc = InputAsset->LoadHairDescription(EHairDescriptionType::Source);
+			UE::Groom::Private::ExtractGroomAttributes(GroomCollection, HairDesc);
+		}
+		SetValue(Context, MoveTemp(GroomCollection), &Collection);
+	}
+	else if (Out->IsA<FCollectionAttributeKey>(&CardGroups))
+	{
+		FCollectionAttributeKey AttributeKey;
+		AttributeKey.Group = GeometryCollection::Facades::FCollectionCurveGeometryFacade::CurvesGroup.ToString();
+		AttributeKey.Attribute = GeometryCollection::Facades::FCollectionCurveAttributesFacade::CurveCardGroupsAttribute.ToString();
+		
+		SetValue(Context, MoveTemp(AttributeKey), &CardGroups);
+	}
+}
+
+FReportCardsAttributesNode::FReportCardsAttributesNode(const UE::Dataflow::FNodeParameters& InParam, FGuid InGuid)
+		: FDataflowNode(InParam, InGuid)
+{
+	RegisterInputConnection(&Collection);
+	RegisterInputConnection(&GroomAsset);
+	RegisterOutputConnection(&Collection, &Collection);
+}
+
+void FReportCardsAttributesNode::Evaluate(UE::Dataflow::FContext& Context, const FDataflowOutput* Out) const
+{
+	if (Out->IsA<FManagedArrayCollection>(&Collection))
+	{
+		const FManagedArrayCollection& GroomCollection = GetValue<FManagedArrayCollection>(Context, &Collection);
+		if(TObjectPtr<UGroomAsset> InputAsset = GetValue<TObjectPtr<UGroomAsset>>(Context, &GroomAsset))
+		{
+			FHairDescription HairDesc = InputAsset->LoadHairDescription(EHairDescriptionType::Source);
+			UE::Groom::Private::ReportGroomAttributes(GroomCollection, HairDesc);
+			
+			InputAsset->CommitHairDescription(MoveTemp(HairDesc), EHairDescriptionType::Edit);
+		}
+		SetValue(Context, GroomCollection, &Collection);
+	}
+}
+
+
diff --git a/Engine/Plugins/Experimental/HairCardGenerator/Source/HairCardGeneratorDataflow/Public/ModifyCardsAttributesNode.h b/Engine/Plugins/Experimental/HairCardGenerator/Source/HairCardGeneratorDataflow/Public/ModifyCardsAttributesNode.h
new file mode 100644
--- /dev/null
+++ b/Engine/Plugins/Experimental/HairCardGenerator/Source/HairCardGeneratorDataflow/Public/ModifyCardsAttributesNode.h
@@ -0,0 +1,74 @@
+// Copyright Epic Games, Inc. All Rights Reserved.
+
+#pragma once 
+
+#include "CoreMinimal.h"
+#include "GroomAsset.h"
+#include "ChaosLog.h"
+#include "BuildCardsSettingsNode.h"
+#include "GetGroomAssetNode.h"
+#include "Dataflow/DataflowConnectionTypes.h"
+#include "Dataflow/DataflowCore.h"
+#include "Dataflow/DataflowEngine.h"
+
+#include "ModifyCardsAttributesNode.generated.h"
+
+/** Extract the cards attributes from the description onto the collection */
+USTRUCT(meta = (Experimental, DataflowGroom))
+struct FExtractCardsAttributesNode : public FDataflowNode
+{
+	GENERATED_BODY()
+	
+	DATAFLOW_NODE_DEFINE_INTERNAL(FExtractCardsAttributesNode, "ExtractCardsAttributes", "Groom", "")
+
+public:
+	
+	FExtractCardsAttributesNode(const UE::Dataflow::FNodeParameters& InParam, FGuid InGuid = FGuid::NewGuid());
+
+private:
+
+	//~ Begin FDataflowNode interface
+	virtual void Evaluate(UE::Dataflow::FContext& Context, const FDataflowOutput* Out) const override;
+	//~ End FDataflowNode interface
+	
+	/** Managed array collection to be used to store datas */
+	UPROPERTY(meta = (DataflowInput, DataflowOutput, DisplayName = "Collection", DataflowPassthrough  = "Collection", DataflowRenderGroups = "CardsClumps"))
+	FManagedArrayCollection Collection;
+	
+	/** Input asset to read the attributes from */
+	UPROPERTY(EditAnywhere, Category = "Groom", meta = (DataflowInput, DisplayName = "Groom Asset"))
+	TObjectPtr<const UGroomAsset> GroomAsset = nullptr;
+	
+	/** Attribute key for the extracted card group  */
+	UPROPERTY(EditAnywhere, Category = "Groom", meta = (DataflowOutput, DisplayName = "Card Groups"))
+	FCollectionAttributeKey CardGroups;
+};
+
+/** Reports the cards attributes from the collection onto the description */
+USTRUCT(meta = (Experimental, DataflowGroom))
+struct FReportCardsAttributesNode : public FDataflowNode
+{
+	GENERATED_BODY()
+	
+	DATAFLOW_NODE_DEFINE_INTERNAL(FReportCardsAttributesNode, "ReportCardsAttributes", "Groom", "")
+
+public:
+	
+	FReportCardsAttributesNode(const UE::Dataflow::FNodeParameters& InParam, FGuid InGuid = FGuid::NewGuid());
+
+private:
+
+	//~ Begin FDataflowNode interface
+	virtual void Evaluate(UE::Dataflow::FContext& Context, const FDataflowOutput* Out) const override;
+	//~ End FDataflowNode interface
+	
+	/** Managed array collection to be used to store data */
+	UPROPERTY(meta = (DataflowInput, DataflowOutput, DisplayName = "Collection", DataflowPassthrough  = "Collection", DataflowRenderGroups = "CardsClumps"))
+	FManagedArrayCollection Collection;
+	
+	/** Input asset to write the attributes onto */
+	UPROPERTY(EditAnywhere, Category = "Groom", meta = (DataflowInput, DisplayName = "Groom Asset"))
+	TObjectPtr<UGroomAsset> GroomAsset = nullptr;
+};
+
+
diff --git a/Engine/Source/Runtime/Experimental/Chaos/Private/GeometryCollection/Facades/CollectionCurveFacade.cpp b/Engine/Source/Runtime/Experimental/Chaos/Private/GeometryCollection/Facades/CollectionCurveFacade.cpp
--- a/Engine/Source/Runtime/Experimental/Chaos/Private/GeometryCollection/Facades/CollectionCurveFacade.cpp
+++ b/Engine/Source/Runtime/Experimental/Chaos/Private/GeometryCollection/Facades/CollectionCurveFacade.cpp
@@ -24,6 +24,9 @@ namespace GeometryCollection::Facades
 	const FName FCollectionCurveHierarchyFacade::CurveParentIndicesAttribute("ParentIndices");
 	const FName FCollectionCurveHierarchyFacade::CurveLodIndicesAttribute("LodIndices");
 	
+	// Curves attributes facade 
+	const FName FCollectionCurveAttributesFacade::CurveCardGroupsAttribute("CardGroups");
+	
 	FCollectionCurveGeometryFacade::FCollectionCurveGeometryFacade(FManagedArrayCollection& InCollection) :
 		ConstCollection(InCollection), Collection(&InCollection),
 		PointRestOrientations(InCollection, PointRestOrientationsAttribute, PointsGroup),
@@ -417,5 +420,29 @@ namespace GeometryCollection::Facades
 		CurveParentIndices.Add();
 		CurveLodIndices.Add();
 	}
+	
+	FCollectionCurveAttributesFacade::FCollectionCurveAttributesFacade(FManagedArrayCollection& InCollection) :
+		ConstCollection(InCollection), Collection(&InCollection),
+		CurveCardGroups(InCollection, CurveCardGroupsAttribute, FCollectionCurveGeometryFacade::CurvesGroup)
+	{
+		DefineSchema();
+	}
+	
+	FCollectionCurveAttributesFacade::FCollectionCurveAttributesFacade(const FManagedArrayCollection& InCollection) :
+		ConstCollection(InCollection), Collection(nullptr),
+		CurveCardGroups(InCollection, CurveCardGroupsAttribute, FCollectionCurveGeometryFacade::CurvesGroup)
+	{
+	}
+	
+	bool FCollectionCurveAttributesFacade::IsValid() const
+	{
+		return CurveCardGroups.IsValid();
+	}
+	
+	void FCollectionCurveAttributesFacade::DefineSchema()
+	{
+		check(!IsConst());
+		CurveCardGroups.Add();
+	}
 }
 
diff --git a/Engine/Source/Runtime/Experimental/Chaos/Public/GeometryCollection/Facades/CollectionCurveFacade.h b/Engine/Source/Runtime/Experimental/Chaos/Public/GeometryCollection/Facades/CollectionCurveFacade.h
--- a/Engine/Source/Runtime/Experimental/Chaos/Public/GeometryCollection/Facades/CollectionCurveFacade.h
+++ b/Engine/Source/Runtime/Experimental/Chaos/Public/GeometryCollection/Facades/CollectionCurveFacade.h
@@ -220,5 +220,48 @@ namespace GeometryCollection::Facades
 		/** Curve lod indices */
 		TManagedArrayAccessor<int32> CurveLodIndices;
 	};
+	
+	class FCollectionCurveAttributesFacade
+	{
+	public:
+
+		/** Collection attribute names */
+		static CHAOS_API const FName CurveCardGroupsAttribute;
+
+		CHAOS_API FCollectionCurveAttributesFacade(FManagedArrayCollection& InCollection);
+		CHAOS_API FCollectionCurveAttributesFacade(const FManagedArrayCollection& InCollection);
+
+		/** Create the facade attributes. */
+		CHAOS_API void DefineSchema();
+
+		/** Is the facade defined constant. */
+		bool IsConst() const { return Collection == nullptr; }
+
+		/** Is the Facade defined on the collection? */
+		CHAOS_API bool IsValid() const;
+		
+		/** Get the number of curves */
+		int32 GetNumCurves() const { return CurveCardGroups.Num(); }
+		
+		/** Get the curve card groups */
+		const TArray<FName>& GetCurveCardGroups() const { return CurveCardGroups.Get().GetConstArray(); }
+
+		/** Set the curve card groups */
+		void SetCurveCardGroups(const TArray<FName>& InCurveCardGroups) { CurveCardGroups.Modify() = InCurveCardGroups; }
+
+		/** Get the managed array collection */
+		const FManagedArrayCollection& GetManagedArrayCollection() const {return ConstCollection;}
+		
+	protected :
+		
+		/** Const collection the facade is linked to */
+		const FManagedArrayCollection& ConstCollection;
+
+		/** Non-const collection the facade is linked to */
+		FManagedArrayCollection* Collection = nullptr;
+
+		/** Curve card groups */
+		TManagedArrayAccessor<FName> CurveCardGroups;
+	};
 }
 
diff --git a/Engine/Source/Runtime/Experimental/Dataflow/Core/Private/Dataflow/DataflowConvertNodes.cpp b/Engine/Source/Runtime/Experimental/Dataflow/Core/Private/Dataflow/DataflowConvertNodes.cpp
--- a/Engine/Source/Runtime/Experimental/Dataflow/Core/Private/Dataflow/DataflowConvertNodes.cpp
+++ b/Engine/Source/Runtime/Experimental/Dataflow/Core/Private/Dataflow/DataflowConvertNodes.cpp
@@ -417,7 +417,8 @@ namespace UE::Private::DataflowConvertNodesLocal
 		return SetMatchingSelectionHelper<FDataflowTransformSelection>(SelectionArr, InCollection, Context, Output, OutMemberPtr, ThisNode)
 			|| SetMatchingSelectionHelper<FDataflowGeometrySelection>(SelectionArr, InCollection, Context, Output, OutMemberPtr, ThisNode)
 			|| SetMatchingSelectionHelper<FDataflowFaceSelection>(SelectionArr, InCollection, Context, Output, OutMemberPtr, ThisNode)
-			|| SetMatchingSelectionHelper<FDataflowVertexSelection>(SelectionArr, InCollection, Context, Output, OutMemberPtr, ThisNode);
+			|| SetMatchingSelectionHelper<FDataflowVertexSelection>(SelectionArr, InCollection, Context, Output, OutMemberPtr, ThisNode) 
+			|| SetMatchingSelectionHelper<FDataflowCurveSelection>(SelectionArr, InCollection, Context, Output, OutMemberPtr, ThisNode);
 	}
 }
 
