Clean up handling of skeletal mesh data in mesh->skeletal mesh dataflow node

#Synced-CL 52440158

diff --git a/Engine/Plugins/Experimental/GeometryCollectionPlugin/Source/GeometryCollectionNodes/Private/Dataflow/MeshToSkeletalMeshTerminalNode.cpp b/Engine/Plugins/Experimental/GeometryCollectionPlugin/Source/GeometryCollectionNodes/Private/Dataflow/MeshToSkeletalMeshTerminalNode.cpp
--- a/Engine/Plugins/Experimental/GeometryCollectionPlugin/Source/GeometryCollectionNodes/Private/Dataflow/MeshToSkeletalMeshTerminalNode.cpp
+++ b/Engine/Plugins/Experimental/GeometryCollectionPlugin/Source/GeometryCollectionNodes/Private/Dataflow/MeshToSkeletalMeshTerminalNode.cpp
@@ -15,6 +15,7 @@
 #include "Engine/SkeletalMesh.h"
 #include "Animation/Skeleton.h"
 #include "Engine/SkinnedAssetCommon.h"
+#include "SkeletalMeshTypes.h"
 
 #include "Dataflow/DataflowDebugDrawInterface.h"
 #include "Dataflow/DataflowDebugDraw.h"
@@ -46,13 +47,6 @@ namespace MeshToSkeletalMeshAssetLocals
 	// editor-only due to FStaticToSkeletalMeshConverter usage
 	void MeshDescriptionToSkeletalMesh(const FMeshDescription& MeshDescription, FSkeletalMeshAttributes& Attributes, const TArray<TObjectPtr<UMaterialInterface>>& InMaterials, USkeletalMesh* OutSkeletalMesh, USkeleton* OutSkeleton)
 	{
-		// Reset the skeletal mesh
-		OutSkeletalMesh->SetNumSourceModels(0);
-		OutSkeletalMesh->GetImportedModel()->LODModels.Reset();
-		OutSkeletalMesh->ReleaseResources();
-		OutSkeletalMesh->ReleaseSkinWeightProfileResources();
-
-
 		// Copy bones to skeleton
 		if (Attributes.GetNumBones() &&
 			// Note FReferenceSkeletonModifier does not allow us to remove a root bone, so we can't clear an existing skeleton here; need to be passed a fresh/empty skeleton
@@ -284,18 +278,6 @@ void FMeshToSkeletalMeshTerminalNode::SetAssetValue(TObjectPtr<UObject> Asset, U
 		return;
 	}
 
-	// When re-using a skeletal mesh asset, we manually clear existing data so we don't e.g. have anything referencing a previous skeleton
-	if (bUsingExistingSkeletalMesh)
-	{
-		OutSkeletalMesh->PreEditChange(nullptr);
-
-		OutSkeletalMesh->Clear();
-
-		// Make sure that we create non-transient assets
-		OutSkeletalMesh->ClearFlags(RF_Transient);
-	}
-
-
 	const FDynamicMesh3& DynMesh = InMesh->GetMeshRef();
 	const TArray<TObjectPtr<UMaterialInterface>>& InMaterials = GetValue(Context, &Materials);
 	// convert the collection to a mesh description including skeletal mesh attributes
@@ -308,13 +290,18 @@ void FMeshToSkeletalMeshTerminalNode::SetAssetValue(TObjectPtr<UObject> Asset, U
 	FDynamicMeshToMeshDescription Converter(ConversionOptions);
 	Converter.Convert(&DynMesh, MeshDescription);
 
-	MeshToSkeletalMeshAssetLocals::MeshDescriptionToSkeletalMesh(MeshDescription, Attributes, InMaterials, OutSkeletalMesh, OutSkeleton);
-
+	TOptional<FScopedSkeletalMeshPostEditChange> ScopedPostEditChange;
 	if (bUsingExistingSkeletalMesh)
 	{
-		OutSkeletalMesh->PostEditChange();
+		ScopedPostEditChange.Emplace(OutSkeletalMesh);
+
+		OutSkeletalMesh->Clear();
+		// Make sure that we create non-transient assets
+		OutSkeletalMesh->ClearFlags(RF_Transient);
 	}
 
+	MeshToSkeletalMeshAssetLocals::MeshDescriptionToSkeletalMesh(MeshDescription, Attributes, InMaterials, OutSkeletalMesh, OutSkeleton);
+
 	// finally set the skeletal mesh to the output
 	SetValue(Context, OutSkeletalMesh, &SkeletalMeshAsset);
 	SetValue(Context, OutSkeleton, &SkeletonAsset);
