Add skeleton subdivide dataflow node
+ make simplify and subidivide nodes consistently auto-weld boundaries

#Synced-CL 52686699

diff --git a/Engine/Plugins/Experimental/GeometryCollectionPlugin/Source/GeometryCollectionNodes/Private/Dataflow/GeometryCollectionSamplingNodes.cpp b/Engine/Plugins/Experimental/GeometryCollectionPlugin/Source/GeometryCollectionNodes/Private/Dataflow/GeometryCollectionSamplingNodes.cpp
--- a/Engine/Plugins/Experimental/GeometryCollectionPlugin/Source/GeometryCollectionNodes/Private/Dataflow/GeometryCollectionSamplingNodes.cpp
+++ b/Engine/Plugins/Experimental/GeometryCollectionPlugin/Source/GeometryCollectionNodes/Private/Dataflow/GeometryCollectionSamplingNodes.cpp
@@ -26,6 +26,7 @@ namespace UE::Dataflow
 		DATAFLOW_NODE_REGISTER_CREATION_FACTORY(FFilterPointSetWithMeshDataflowNode);
 		DATAFLOW_NODE_REGISTER_CREATION_FACTORY(FMeshMedialSkeletonSamplingDataflowNode);
 		DATAFLOW_NODE_REGISTER_CREATION_FACTORY(FSimplifyMedialSkeletonDataflowNode);
+		DATAFLOW_NODE_REGISTER_CREATION_FACTORY(FSubdivideMedialSkeletonDataflowNode);
 	}
 
 	static void MedialSkeletonDebugDraw(IDataflowDebugDrawInterface& DataflowRenderingInterface, const UE::Geometry::MedialAxis::FMedialSkeleton& ToDrawSkel)
@@ -488,7 +489,7 @@ void FSimplifyMedialSkeletonDataflowNode::Evaluate(UE::Dataflow::FContext& Conte
 				Simplifier.SplitThinTriEdgeAngleThresholdDeg = GetValue(Context, &SplitThinTriEdgeAngleThresholdDeg);
 				Simplifier.SplitThinTriEdgePosErrorWt = GetValue(Context, &SplitThinTriEdgePosErrorWt);
 
-				Simplifier.Simplify(Skeleton, Spatial, [](int32) { return false; }, 1e-4);
+				Simplifier.Simplify(Skeleton, Spatial);
 			}
 			else
 			{
@@ -555,5 +556,130 @@ void FMeshMedialSkeletonSamplingDataflowNode::DebugDraw(UE::Dataflow::FContext&
 }
 #endif
 
+FSubdivideMedialSkeletonDataflowNode::FSubdivideMedialSkeletonDataflowNode(const UE::Dataflow::FNodeParameters& InParam, FGuid InGuid)
+	: FDataflowNode(InParam, InGuid)
+{
+	RegisterInputConnection(&MedialSkeleton);
+	RegisterInputConnection(&TargetMesh);
+	RegisterInputConnection(&TargetEdgeLengthRadiusFraction).SetCanHidePin(true).SetPinIsHidden(true);
+	RegisterInputConnection(&RadiusReference).SetCanHidePin(true).SetPinIsHidden(true);
+	RegisterInputConnection(&MinTargetEdgeLength).SetCanHidePin(true).SetPinIsHidden(true);
+	RegisterInputConnection(&MaxTargetEdgeLength).SetCanHidePin(true).SetPinIsHidden(true);
+	RegisterInputConnection(&bSubdivideOnSurfaces).SetCanHidePin(true).SetPinIsHidden(true);
+	RegisterInputConnection(&ReassignClusterPosErrorWt).SetCanHidePin(true).SetPinIsHidden(true);
+
+	RegisterOutputConnection(&MedialSkeleton);
+	RegisterOutputConnection(&MedialSpheres);
+}
+
+void FSubdivideMedialSkeletonDataflowNode::Evaluate(UE::Dataflow::FContext& Context, const FDataflowOutput* Out) const
+{
+	if (Out->IsA(&MedialSkeleton) || Out->IsA(&MedialSpheres))
+	{
+		FDataflowMedialSkeleton OutMedialSkeleton = GetValue(Context, &MedialSkeleton);
+		UE::Geometry::MedialAxis::FMedialSkeleton& Skeleton = OutMedialSkeleton.Skeleton;
+
+		if (TObjectPtr<const UDynamicMesh> InTargetMesh = GetValue(Context, &TargetMesh))
+		{
+			const UE::Geometry::FDynamicMesh3& InDynTargetMesh = InTargetMesh->GetMeshRef();
+			bool bIsMeshCompatible = Skeleton.IsCompatibleWithMesh(InDynTargetMesh);
+			bool bIsNonEmpty = Skeleton.Spheres.Num() > 0;
+			if (bIsMeshCompatible && bIsNonEmpty)
+			{
+				const double RadiusFraction = GetValue(Context, &TargetEdgeLengthRadiusFraction);
+				const ESubdivideEdgeRadiusReference RadRef = GetValue(Context, &RadiusReference);
+				const double MinLen = GetValue(Context, &MinTargetEdgeLength);
+				const double MaxLen = GetValue(Context, &MaxTargetEdgeLength);
+
+				if (MinLen > 0 && MaxLen > 0 && MinLen > MaxLen)
+				{
+					Context.Warning(LOCTEXT("SubdivideMinMaxInverted", "SubdivideMedialSkeleton: MinTargetEdgeLength exceeds MaxTargetEdgeLength; Max will take precedence."), this, Out);
+				}
+
+				// If all parameters are zero, skip subdivision
+				if (RadiusFraction > 0 || MinLen > 0 || MaxLen > 0)
+				{
+					UE::Geometry::TMeshAABBTree3<FDynamicMesh3> Spatial(&InDynTargetMesh);
+
+					UE::Geometry::MedialAxis::FSkeletonSubdivider Subdivider;
+					Subdivider.bSubdivideOnSurfaces = GetValue(Context, &bSubdivideOnSurfaces);
+					Subdivider.ReassignClusterPosErrorWt = GetValue(Context, &ReassignClusterPosErrorWt);
+					Subdivider.GetTargetEdgeLength = [RadiusFraction, RadRef, MinLen, MaxLen](const FSphere& A, const FSphere& B) -> double
+					{
+						double Radius;
+						switch (RadRef)
+						{
+						default:
+						case ESubdivideEdgeRadiusReference::Larger:
+							Radius = FMath::Max(A.W, B.W);
+							break;
+						case ESubdivideEdgeRadiusReference::Smaller:
+							Radius = FMath::Min(A.W, B.W);
+							break;
+						case ESubdivideEdgeRadiusReference::Average:
+							Radius = (A.W + B.W) * 0.5;
+							break;
+						}
+						double TargetLen = RadiusFraction * Radius;
+						if (MinLen > 0)
+						{
+							TargetLen = FMath::Max(TargetLen, MinLen);
+						}
+						if (MaxLen > 0)
+						{
+							// Use MaxLen as the target if nothing else has set one
+							TargetLen = (TargetLen > 0) ? FMath::Min(TargetLen, MaxLen) : MaxLen;
+						}
+						return TargetLen;
+					};
+
+					Subdivider.Subdivide(Skeleton, Spatial);
+				}
+			}
+			else
+			{
+				if (!bIsMeshCompatible)
+				{
+					Context.Error(LOCTEXT("SubdivideMedialSkeletonIncompatibleMesh", "SubdivideMedialSkeleton: Incompatible mesh provided. Subdivide expects a mesh matching what was used to generate the skeleton."), this, Out);
+				}
+				if (!bIsNonEmpty)
+				{
+					Context.Warning(LOCTEXT("SubdivideMedialSkeletonEmptySkeleton", "SubdivideMedialSkeleton: Empty skeleton has nothing to subdivide."), this, Out);
+				}
+			}
+		}
+		else
+		{
+			Context.Error(LOCTEXT("SubdivideMedialSkeletonNullMeshError", "SubdivideMedialSkeleton: TargetMesh input is null"), this, Out);
+		}
+
+		TArray<FSphere> Spheres = Skeleton.Spheres;
+		SetValue(Context, MoveTemp(OutMedialSkeleton), &MedialSkeleton);
+		SetValue(Context, MoveTemp(Spheres), &MedialSpheres);
+	}
+}
+
+#if WITH_EDITOR
+bool FSubdivideMedialSkeletonDataflowNode::CanDebugDrawViewMode(const FName& ViewModeName) const
+{
+	return ViewModeName == UE::Dataflow::FDataflowConstruction3DViewMode::Name;
+}
+
+void FSubdivideMedialSkeletonDataflowNode::DebugDraw(UE::Dataflow::FContext& Context, IDataflowDebugDrawInterface& DataflowRenderingInterface, const FDebugDrawParameters& DebugDrawParameters) const
+{
+	using namespace UE::Geometry;
+
+	if ((DebugDrawParameters.bNodeIsSelected || DebugDrawParameters.bNodeIsPinned))
+	{
+		if (const FDataflowOutput* MedialSkeletonOutput = FindOutput(&MedialSkeleton))
+		{
+			const UE::Geometry::MedialAxis::FMedialSkeleton& ToDrawSkeleton = MedialSkeletonOutput->GetValue(Context, MedialSkeleton).Skeleton;
+
+			UE::Dataflow::MedialSkeletonDebugDraw(DataflowRenderingInterface, ToDrawSkeleton);
+		}
+	}
+}
+#endif
+
 #undef LOCTEXT_NAMESPACE
 
diff --git a/Engine/Plugins/Experimental/GeometryCollectionPlugin/Source/GeometryCollectionNodes/Public/Dataflow/GeometryCollectionSamplingNodes.h b/Engine/Plugins/Experimental/GeometryCollectionPlugin/Source/GeometryCollectionNodes/Public/Dataflow/GeometryCollectionSamplingNodes.h
--- a/Engine/Plugins/Experimental/GeometryCollectionPlugin/Source/GeometryCollectionNodes/Public/Dataflow/GeometryCollectionSamplingNodes.h
+++ b/Engine/Plugins/Experimental/GeometryCollectionPlugin/Source/GeometryCollectionNodes/Public/Dataflow/GeometryCollectionSamplingNodes.h
@@ -485,6 +485,76 @@ public:
 };
 
 
+/** Which sphere radius to use when computing radius-based target edge lengths */
+UENUM()
+enum class ESubdivideEdgeRadiusReference : uint8
+{
+	// Use the larger of the two sphere radii on the edge
+	Larger,
+	// Use the smaller of the two sphere radii on the edge
+	Smaller,
+	// Use the average of the two sphere radii on the edge
+	Average
+};
+
+/**
+ *
+ * Subdivide a medial skeleton by splitting edges that exceed a target length
+ *
+ */
+USTRUCT(meta = (DataflowGeometryCollection))
+struct FSubdivideMedialSkeletonDataflowNode : public FDataflowNode
+{
+	GENERATED_USTRUCT_BODY()
+	DATAFLOW_NODE_DEFINE_INTERNAL(FSubdivideMedialSkeletonDataflowNode, "SubdivideMedialSkeleton", "MedialSkeletonSampling", "Subdivide Refine Split")
+
+private:
+	/** Medial skeleton to subdivide */
+	UPROPERTY(meta = (DataflowInput, DataflowOutput, DataflowIntrinsic))
+	FDataflowMedialSkeleton MedialSkeleton;
+
+	/** Mesh the skeleton was sampled from, used to re-cluster vertices after subdivision */
+	UPROPERTY(meta = (DataflowInput, DataflowIntrinsic))
+	TObjectPtr<UDynamicMesh> TargetMesh;
+
+	UPROPERTY(meta = (DataflowOutput))
+	TArray<FSphere> MedialSpheres;
+
+	// If positive, target edge length as a fraction of a medial sphere radius on the edge (see Radius Reference)
+	UPROPERTY(EditAnywhere, Category = Options, meta = (DataflowInput, ClampMin = 0))
+	float TargetEdgeLengthRadiusFraction = 0.f;
+
+	// Which sphere radius to reference when computing radius-based target edge lengths
+	UPROPERTY(EditAnywhere, Category = Options, meta = (DataflowInput, EditCondition = "TargetEdgeLengthRadiusFraction > 0"))
+	ESubdivideEdgeRadiusReference RadiusReference = ESubdivideEdgeRadiusReference::Average;
+
+	// If positive, minimum target edge length
+	UPROPERTY(EditAnywhere, Category = Options, meta = (DataflowInput, ClampMin = 0, Units = "cm"))
+	float MinTargetEdgeLength = 0.f;
+
+	// If positive, maximum target edge length. Also used as the target directly when Radius Fraction and Min are both 0
+	UPROPERTY(EditAnywhere, Category = Options, meta = (DataflowInput, ClampMin = 0, Units = "cm"))
+	float MaxTargetEdgeLength = 0.f;
+
+	// If true, also subdivide edges that have associated triangles (splitting those triangles)
+	UPROPERTY(EditAnywhere, Category = Options, meta = (DataflowInput))
+	bool bSubdivideOnSurfaces = false;
+
+	// Position error weight for re-clustering vertices after edge splits
+	UPROPERTY(EditAnywhere, Category = Options, AdvancedDisplay, meta = (DataflowInput, ClampMin = 0, ClampMax = 1))
+	float ReassignClusterPosErrorWt = .2f;
+
+	virtual void Evaluate(UE::Dataflow::FContext& Context, const FDataflowOutput* Out) const override;
+
+#if WITH_EDITOR
+	virtual bool CanDebugDrawViewMode(const FName& ViewModeName) const override;
+	virtual void DebugDraw(UE::Dataflow::FContext& Context, IDataflowDebugDrawInterface& DataflowRenderingInterface, const FDebugDrawParameters& DebugDrawParameters) const override;
+#endif
+
+public:
+	FSubdivideMedialSkeletonDataflowNode(const UE::Dataflow::FNodeParameters& InParam, FGuid InGuid = FGuid::NewGuid());
+};
+
 
 namespace UE::Dataflow
 {
