clean up skeleton simplify node options -- move split edge settings into options so they're easier to find, and don't expose regularize weights

#Synced-CL 52703645

diff --git a/Engine/Plugins/Experimental/GeometryCollectionPlugin/Source/GeometryCollectionNodes/Private/Dataflow/GeometryCollectionSamplingNodes.cpp b/Engine/Plugins/Experimental/GeometryCollectionPlugin/Source/GeometryCollectionNodes/Private/Dataflow/GeometryCollectionSamplingNodes.cpp
--- a/Engine/Plugins/Experimental/GeometryCollectionPlugin/Source/GeometryCollectionNodes/Private/Dataflow/GeometryCollectionSamplingNodes.cpp
+++ b/Engine/Plugins/Experimental/GeometryCollectionPlugin/Source/GeometryCollectionNodes/Private/Dataflow/GeometryCollectionSamplingNodes.cpp
@@ -449,7 +449,6 @@ FSimplifyMedialSkeletonDataflowNode::FSimplifyMedialSkeletonDataflowNode(const U
 	RegisterInputConnection(&bOnlySimplifySurfaces).SetCanHidePin(true).SetPinIsHidden(true);
 	RegisterInputConnection(&ClusterSkeletonDistanceWt).SetCanHidePin(true).SetPinIsHidden(true);
 	RegisterInputConnection(&bPreventEdgeSurfaceIntersections).SetCanHidePin(true).SetPinIsHidden(true);
-	RegisterInputConnection(&ClusterRegularizeWeight).SetCanHidePin(true).SetPinIsHidden(true);
 	RegisterInputConnection(&bSplitThinTriEdges).SetCanHidePin(true).SetPinIsHidden(true);
 	RegisterInputConnection(&SplitThinTriEdgeAngleThresholdDeg).SetCanHidePin(true).SetPinIsHidden(true);
 	RegisterInputConnection(&SplitThinTriEdgePosErrorWt).SetCanHidePin(true).SetPinIsHidden(true);
@@ -484,7 +483,7 @@ void FSimplifyMedialSkeletonDataflowNode::Evaluate(UE::Dataflow::FContext& Conte
 				Simplifier.bOnlySimplifySurfaces = GetValue(Context, &bOnlySimplifySurfaces);
 				Simplifier.ClusterSkeletonDistanceWt = GetValue(Context, &ClusterSkeletonDistanceWt);
 				Simplifier.bPreventEdgeSurfaceIntersections = GetValue(Context, &bPreventEdgeSurfaceIntersections);
-				Simplifier.ClusterRegularizeWeight = GetValue(Context, &ClusterRegularizeWeight);
+				Simplifier.ClusterRegularizeWeight = ClusterRegularizeWeight;
 				Simplifier.bSplitThinTriEdges = GetValue(Context, &bSplitThinTriEdges);
 				Simplifier.SplitThinTriEdgeAngleThresholdDeg = GetValue(Context, &SplitThinTriEdgeAngleThresholdDeg);
 				Simplifier.SplitThinTriEdgePosErrorWt = GetValue(Context, &SplitThinTriEdgePosErrorWt);
diff --git a/Engine/Plugins/Experimental/GeometryCollectionPlugin/Source/GeometryCollectionNodes/Public/Dataflow/GeometryCollectionSamplingNodes.h b/Engine/Plugins/Experimental/GeometryCollectionPlugin/Source/GeometryCollectionNodes/Public/Dataflow/GeometryCollectionSamplingNodes.h
--- a/Engine/Plugins/Experimental/GeometryCollectionPlugin/Source/GeometryCollectionNodes/Public/Dataflow/GeometryCollectionSamplingNodes.h
+++ b/Engine/Plugins/Experimental/GeometryCollectionPlugin/Source/GeometryCollectionNodes/Public/Dataflow/GeometryCollectionSamplingNodes.h
@@ -457,19 +457,19 @@ private:
 	bool bPreventEdgeSurfaceIntersections = true;
 
 	// Regularization weight encouraging skeleton clusters to keep their original positions and radii
-	UPROPERTY(EditAnywhere, Category = Options, AdvancedDisplay, meta = (DataflowInput, ClampMin = 0))
+	//~ Not exposed to users currently
 	float ClusterRegularizeWeight = 1e-6f;
 
 	// Whether to split edges attached to skinny triangles before simplifying, which can add useful degrees of freedom to the simplifier
-	UPROPERTY(EditAnywhere, Category = SplitEdges, AdvancedDisplay, meta = (DataflowInput))
+	UPROPERTY(EditAnywhere, Category = Options, AdvancedDisplay, meta = (DataflowInput))
 	bool bSplitThinTriEdges = true;
 
 	// Minimum angle (degrees) at the opposite vertex for an edge to be a split candidate
-	UPROPERTY(EditAnywhere, Category = SplitEdges, AdvancedDisplay, meta = (EditCondition="bSplitThinTriEdges", DataflowInput, ClampMin = 90, ClampMax = 180))
+	UPROPERTY(EditAnywhere, Category = Options, AdvancedDisplay, meta = (EditCondition="bSplitThinTriEdges", DataflowInput, ClampMin = 90, ClampMax = 180))
 	float SplitThinTriEdgeAngleThresholdDeg = 120.f;
 
 	// Position error weight used when re-clustering vertices after thin-triangle edge splits
-	UPROPERTY(EditAnywhere, Category = SplitEdges, AdvancedDisplay, meta = (EditCondition="bSplitThinTriEdges", DataflowInput, ClampMin = 0, ClampMax = 1))
+	UPROPERTY(EditAnywhere, Category = Options, AdvancedDisplay, meta = (EditCondition="bSplitThinTriEdges", DataflowInput, ClampMin = 0, ClampMax = 1))
 	float SplitThinTriEdgePosErrorWt = .2f;
 
 	virtual void Evaluate(UE::Dataflow::FContext& Context, const FDataflowOutput* Out) const override;
