Add a split islands dataflow node + associated function in fracture engine utility
minor cleanup of select-all handling for similar utility nodes

#Synced-CL 52703645

diff --git a/Engine/Plugins/Experimental/Fracture/Source/FractureEngine/Private/FractureEngineUtility.cpp b/Engine/Plugins/Experimental/Fracture/Source/FractureEngine/Private/FractureEngineUtility.cpp
--- a/Engine/Plugins/Experimental/Fracture/Source/FractureEngine/Private/FractureEngineUtility.cpp
+++ b/Engine/Plugins/Experimental/Fracture/Source/FractureEngine/Private/FractureEngineUtility.cpp
@@ -433,5 +433,27 @@ void FFractureEngineUtility::ValidateGeometryCollection(FManagedArrayCollection&
 	}
 }
 
+void FFractureEngineUtility::SplitIslands(FManagedArrayCollection& InOutCollection,
+	FDataflowTransformSelection InTransformSelection,
+	const float InCloseVertexDistance,
+	const float InVertexToSurfaceBridgeDistance)
+{
+	if (TUniquePtr<FGeometryCollection> GeomCollection = TUniquePtr<FGeometryCollection>(InOutCollection.NewCopy<FGeometryCollection>()))
+	{
+		TArray<int32> SelectionArray = InTransformSelection.AsArray();
+		GeometryCollection::Facades::FCollectionTransformSelectionFacade SelectionFacade(*GeomCollection);
+		SelectionFacade.ConvertSelectionToRigidNodes(SelectionArray);
+
+		FIslandSplitSettings Settings(true, (double)InCloseVertexDistance, (double)InVertexToSurfaceBridgeDistance);
 
+		::SplitIslands(
+			*GeomCollection,
+			SelectionArray,
+			0.0,
+			nullptr,
+			Settings);
+
+		InOutCollection = (const FManagedArrayCollection&)(*GeomCollection);
+	}
+}
 
diff --git a/Engine/Plugins/Experimental/Fracture/Source/FractureEngine/Public/FractureEngineUtility.h b/Engine/Plugins/Experimental/Fracture/Source/FractureEngine/Public/FractureEngineUtility.h
--- a/Engine/Plugins/Experimental/Fracture/Source/FractureEngine/Public/FractureEngineUtility.h
+++ b/Engine/Plugins/Experimental/Fracture/Source/FractureEngine/Public/FractureEngineUtility.h
@@ -96,6 +96,11 @@ public:
 		const bool InRemoveClustersOfOne,
 		const bool InRemoveDanglingClusters);
 
+	static UE_API void SplitIslands(FManagedArrayCollection& InOutCollection,
+		FDataflowTransformSelection InTransformSelection,
+		const float InCloseVertexDistance,
+		const float InVertexToSurfaceBridgeDistance);
+
 };
 
 #undef UE_API
diff --git a/Engine/Plugins/Experimental/GeometryCollectionPlugin/Source/GeometryCollectionNodes/Private/Dataflow/GeometryCollectionUtilityNodes.cpp b/Engine/Plugins/Experimental/GeometryCollectionPlugin/Source/GeometryCollectionNodes/Private/Dataflow/GeometryCollectionUtilityNodes.cpp
--- a/Engine/Plugins/Experimental/GeometryCollectionPlugin/Source/GeometryCollectionNodes/Private/Dataflow/GeometryCollectionUtilityNodes.cpp
+++ b/Engine/Plugins/Experimental/GeometryCollectionPlugin/Source/GeometryCollectionNodes/Private/Dataflow/GeometryCollectionUtilityNodes.cpp
@@ -47,6 +47,7 @@ namespace UE::Dataflow
 		DATAFLOW_NODE_REGISTER_CREATION_FACTORY(FUpdateVolumeAttributesDataflowNode);
 		DATAFLOW_NODE_REGISTER_CREATION_FACTORY(FGetConvexHullVolumeDataflowNode);
 		DATAFLOW_NODE_REGISTER_CREATION_FACTORY(FFixTinyGeoDataflowNode);
+		DATAFLOW_NODE_REGISTER_CREATION_FACTORY(FSplitIslandsDataflowNode);
 		DATAFLOW_NODE_REGISTER_CREATION_FACTORY(FRecomputeNormalsInGeometryCollectionDataflowNode);
 		DATAFLOW_NODE_REGISTER_CREATION_FACTORY(FResampleGeometryCollectionDataflowNode);
 		DATAFLOW_NODE_REGISTER_CREATION_FACTORY(FValidateGeometryCollectionDataflowNode);
@@ -1390,16 +1391,7 @@ void FFixTinyGeoDataflowNode::Evaluate(UE::Dataflow::FContext& Context, const FD
 		//
 		if (!IsConnected(&TransformSelection))
 		{
-			const FManagedArrayCollection& InCollection = GetValue(Context, &Collection);
-
-			GeometryCollection::Facades::FCollectionTransformSelectionFacade TransformSelectionFacade(InCollection);
-			const TArray<int32>& SelectionArr = TransformSelectionFacade.SelectAll();
-
-			FDataflowTransformSelection NewTransformSelection;
-			NewTransformSelection.InitializeFromCollection(InCollection, false);
-			NewTransformSelection.SetFromArray(SelectionArr);
-
-			InTransformSelection = NewTransformSelection;
+			InTransformSelection.InitializeFromCollection(GetValue(Context, &Collection), true);
 		}
 
 		if (InTransformSelection.AnySelected())
@@ -1429,26 +1421,54 @@ void FFixTinyGeoDataflowNode::Evaluate(UE::Dataflow::FContext& Context, const FD
 	}
 }
 
-void FRecomputeNormalsInGeometryCollectionDataflowNode::Evaluate(UE::Dataflow::FContext& Context, const FDataflowOutput* Out) const
+FSplitIslandsDataflowNode::FSplitIslandsDataflowNode(const UE::Dataflow::FNodeParameters& InParam, FGuid InGuid)
+	: FDataflowNode(InParam, InGuid)
+{
+	RegisterInputConnection(&Collection);
+	RegisterInputConnection(&TransformSelection);
+	RegisterOutputConnection(&Collection, &Collection);
+}
+
+void FSplitIslandsDataflowNode::Evaluate(UE::Dataflow::FContext& Context, const FDataflowOutput* Out) const
 {
 	if (Out->IsA(&Collection))
 	{
 		FDataflowTransformSelection InTransformSelection = GetValue(Context, &TransformSelection);
-		//
 		// If not connected select everything by default
-		//
 		if (!IsConnected(&TransformSelection))
 		{
-			const FManagedArrayCollection& InCollection = GetValue(Context, &Collection);
+			InTransformSelection.InitializeFromCollection(GetValue(Context, &Collection), true);
+		}
 
-			GeometryCollection::Facades::FCollectionTransformSelectionFacade TransformSelectionFacade(InCollection);
-			const TArray<int32>& SelectionArr = TransformSelectionFacade.SelectAll();
+		if (InTransformSelection.AnySelected())
+		{
+			FManagedArrayCollection InCollection = GetValue(Context, &Collection);
+
+			FFractureEngineUtility::SplitIslands(InCollection,
+				InTransformSelection,
+				CloseVertexDistance,
+				VertexToSurfaceBridgeDistance);
 
-			FDataflowTransformSelection NewTransformSelection;
-			NewTransformSelection.InitializeFromCollection(InCollection, false);
-			NewTransformSelection.SetFromArray(SelectionArr);
+			SetValue(Context, MoveTemp(InCollection), &Collection);
 
-			InTransformSelection = NewTransformSelection;
+			return;
+		}
+
+		SafeForwardInput(Context, &Collection, &Collection);
+	}
+}
+
+void FRecomputeNormalsInGeometryCollectionDataflowNode::Evaluate(UE::Dataflow::FContext& Context, const FDataflowOutput* Out) const
+{
+	if (Out->IsA(&Collection))
+	{
+		FDataflowTransformSelection InTransformSelection = GetValue(Context, &TransformSelection);
+		//
+		// If not connected select everything by default
+		//
+		if (!IsConnected(&TransformSelection))
+		{
+			InTransformSelection.InitializeFromCollection(GetValue(Context, &Collection), true);
 		}
 
 		if (InTransformSelection.AnySelected())
@@ -1481,16 +1501,7 @@ void FResampleGeometryCollectionDataflowNode::Evaluate(UE::Dataflow::FContext& C
 		//
 		if (!IsConnected(&TransformSelection))
 		{
-			const FManagedArrayCollection& InCollection = GetValue(Context, &Collection);
-
-			GeometryCollection::Facades::FCollectionTransformSelectionFacade TransformSelectionFacade(InCollection);
-			const TArray<int32>& SelectionArr = TransformSelectionFacade.SelectAll();
-
-			FDataflowTransformSelection NewTransformSelection;
-			NewTransformSelection.InitializeFromCollection(InCollection, false);
-			NewTransformSelection.SetFromArray(SelectionArr);
-
-			InTransformSelection = NewTransformSelection;
+			InTransformSelection.InitializeFromCollection(GetValue(Context, &Collection), true);
 		}
 
 		if (InTransformSelection.AnySelected())
diff --git a/Engine/Plugins/Experimental/GeometryCollectionPlugin/Source/GeometryCollectionNodes/Public/Dataflow/GeometryCollectionUtilityNodes.h b/Engine/Plugins/Experimental/GeometryCollectionPlugin/Source/GeometryCollectionNodes/Public/Dataflow/GeometryCollectionUtilityNodes.h
--- a/Engine/Plugins/Experimental/GeometryCollectionPlugin/Source/GeometryCollectionNodes/Public/Dataflow/GeometryCollectionUtilityNodes.h
+++ b/Engine/Plugins/Experimental/GeometryCollectionPlugin/Source/GeometryCollectionNodes/Public/Dataflow/GeometryCollectionUtilityNodes.h
@@ -1058,6 +1058,41 @@ public:
 	virtual void Evaluate(UE::Dataflow::FContext& Context, const FDataflowOutput* Out) const override;
 };
 
+/**
+ *
+ * Editor Fracture Mode / Utilities / Split Islands tool
+ * Split selected geometry into islands based on mesh connectivity and distances.
+ *
+ */
+USTRUCT(meta = (DataflowGeometryCollection))
+struct FSplitIslandsDataflowNode : public FDataflowNode
+{
+	GENERATED_USTRUCT_BODY()
+	DATAFLOW_NODE_DEFINE_INTERNAL(FSplitIslandsDataflowNode, "SplitIslands", "GeometryCollection|Fracture|Utilities", "Disconnect Connected Components")
+
+private:
+	/** Collection to use */
+	UPROPERTY(meta = (DataflowInput, DataflowOutput, DataflowPassthrough = "Collection", DataflowIntrinsic))
+	FManagedArrayCollection Collection;
+
+	/** The selected pieces to use */
+	UPROPERTY(meta = (DataflowInput, DisplayName = "TransformSelection"))
+	FDataflowTransformSelection TransformSelection;
+
+	/** Distance threshold for connecting spatially-close vertices before computing islands */
+	UPROPERTY(EditAnywhere, Category = "Split Islands", meta = (Units = "cm", ClampMin = "0.0"))
+	float CloseVertexDistance = 0.001f;
+
+	/** If > 0, bridge separate islands whose surfaces are within this vertex-to-triangle distance. 0 = disabled. */
+	UPROPERTY(EditAnywhere, Category = "Split Islands", meta = (Units = "cm", ClampMin = "0.0"))
+	float VertexToSurfaceBridgeDistance = 0.f;
+	
+	virtual void Evaluate(UE::Dataflow::FContext& Context, const FDataflowOutput* Out) const override;
+	
+public:
+	FSplitIslandsDataflowNode(const UE::Dataflow::FNodeParameters& InParam, FGuid InGuid = FGuid::NewGuid());
+};
+
 /**
  *
  * Editor Fracture Mode / Utilities / Normals tool
