minor cleanup of options for medial->anim skeleton conversion; make sure root controls are dataflow inputs and allow control/positioning of the custom root

#Synced-CL 52703645

diff --git a/Engine/Plugins/Experimental/GeometryCollectionPlugin/Source/GeometryCollectionNodes/Private/Dataflow/MeshMedialSkeletonConversionNodes.cpp b/Engine/Plugins/Experimental/GeometryCollectionPlugin/Source/GeometryCollectionNodes/Private/Dataflow/MeshMedialSkeletonConversionNodes.cpp
--- a/Engine/Plugins/Experimental/GeometryCollectionPlugin/Source/GeometryCollectionNodes/Private/Dataflow/MeshMedialSkeletonConversionNodes.cpp
+++ b/Engine/Plugins/Experimental/GeometryCollectionPlugin/Source/GeometryCollectionNodes/Private/Dataflow/MeshMedialSkeletonConversionNodes.cpp
@@ -128,7 +128,11 @@ FConvertMedialSkeletonToAnimationSkeletonDataflowNode::FConvertMedialSkeletonToA
 {
 	RegisterInputConnection(&MedialSkeleton);
 	RegisterInputConnection(&AvoidIntersectingMesh);
+	RegisterInputConnection(&RootSelectionPoint).SetCanHidePin(true).SetPinIsHidden(true);
+	RegisterInputConnection(&RootSelectionDirection).SetCanHidePin(true).SetPinIsHidden(true);
 	RegisterInputConnection(&RootClusterIndex).SetCanHidePin(true).SetPinIsHidden(true);
+	RegisterInputConnection(&bAddCustomAnimationRoot).SetCanHidePin(true).SetPinIsHidden(true);
+	RegisterInputConnection(&CustomAnimationRootPosition).SetCanHidePin(true).SetPinIsHidden(true);
 
 	RegisterOutputConnection(&Skeleton);
 }
@@ -156,7 +160,12 @@ void FConvertMedialSkeletonToAnimationSkeletonDataflowNode::Evaluate(UE::Dataflo
 		TArray<FTransform> Poses;
 
 		UE::Geometry::MedialAxis::FMedialSkeletonToTreeSkeletonOptions Options;
-		MeshMedialSkeletonConversionLocals::SetSkeletonConversionOptions(Options, EdgeWeightMethod, MergeDisconnectedMethod, SelectRootMethod, RootSelectionPoint, RootSelectionDirection);
+		MeshMedialSkeletonConversionLocals::SetSkeletonConversionOptions(Options, EdgeWeightMethod, MergeDisconnectedMethod, SelectRootMethod,
+			GetValue(Context, &RootSelectionPoint), GetValue(Context, &RootSelectionDirection));
+		if (GetValue(Context, &bAddCustomAnimationRoot))
+		{
+			Options.CustomRootPosition = FVector3d(GetValue(Context, &CustomAnimationRootPosition));
+		}
 		const int32 UseRootIndex = GetValue(Context, &RootClusterIndex);
 		Options.ToHierarchy(InMedialSkeleton.Skeleton, AvoidIntersectingMeshBVH, Parents, Poses, Options, UseRootIndex);
 
@@ -239,7 +248,11 @@ FSkinMeshViaMedialSkeleton::FSkinMeshViaMedialSkeleton(const UE::Dataflow::FNode
 {
 	RegisterInputConnection(&Mesh);
 	RegisterInputConnection(&MedialSkeleton);
+	RegisterInputConnection(&RootSelectionPoint).SetCanHidePin(true).SetPinIsHidden(true);
+	RegisterInputConnection(&RootSelectionDirection).SetCanHidePin(true).SetPinIsHidden(true);
 	RegisterInputConnection(&RootClusterIndex).SetCanHidePin(true).SetPinIsHidden(true);
+	RegisterInputConnection(&bAddCustomAnimationRoot).SetCanHidePin(true).SetPinIsHidden(true);
+	RegisterInputConnection(&CustomAnimationRootPosition).SetCanHidePin(true).SetPinIsHidden(true);
 	RegisterInputConnection(&Stiffness).SetCanHidePin(true).SetPinIsHidden(true);
 	RegisterInputConnection(&MaxInfluences).SetCanHidePin(true).SetPinIsHidden(true);
 	RegisterInputConnection(&VoxelResolution).SetCanHidePin(true).SetPinIsHidden(true);
@@ -273,7 +286,12 @@ void FSkinMeshViaMedialSkeleton::Evaluate(UE::Dataflow::FContext& Context, const
 			{
 				SkinBinding::FBindMedialSkeletonSettings Settings;
 				Settings.ClusterNeighborSearchRange = GetValue(Context, &ClusterNeighborSearchRange);
-				SetSkeletonConversionOptions(Settings.AnimationSkeletonOptions, EdgeWeightMethod, MergeDisconnectedMethod, SelectRootMethod, RootSelectionPoint, RootSelectionDirection);
+				SetSkeletonConversionOptions(Settings.AnimationSkeletonOptions, EdgeWeightMethod, MergeDisconnectedMethod, SelectRootMethod,
+					GetValue(Context, &RootSelectionPoint), GetValue(Context, &RootSelectionDirection));
+				if (GetValue(Context, &bAddCustomAnimationRoot))
+				{
+					Settings.AnimationSkeletonOptions.CustomRootPosition = FVector3d(GetValue(Context, &CustomAnimationRootPosition));
+				}
 				const int32 UseRootIndex = GetValue(Context, &RootClusterIndex);
 
 				SetSkinBindSettings(Settings.BindSettings, BindMethod,
diff --git a/Engine/Plugins/Experimental/GeometryCollectionPlugin/Source/GeometryCollectionNodes/Private/Dataflow/MeshMedialSkeletonConversionNodes.h b/Engine/Plugins/Experimental/GeometryCollectionPlugin/Source/GeometryCollectionNodes/Private/Dataflow/MeshMedialSkeletonConversionNodes.h
--- a/Engine/Plugins/Experimental/GeometryCollectionPlugin/Source/GeometryCollectionNodes/Private/Dataflow/MeshMedialSkeletonConversionNodes.h
+++ b/Engine/Plugins/Experimental/GeometryCollectionPlugin/Source/GeometryCollectionNodes/Private/Dataflow/MeshMedialSkeletonConversionNodes.h
@@ -90,17 +90,25 @@ private:
 	EDataflowMedialSkeletonConversionSelectRootMethod SelectRootMethod = EDataflowMedialSkeletonConversionSelectRootMethod::ArrayOrder;
 
 	/** When SelectRootMethod is ClosestToPoint, find the closest cluster to this point */
-	UPROPERTY(EditAnywhere, Category = "Options", meta = (EditCondition = "SelectRootMethod == EDataflowMedialSkeletonConversionSelectRootMethod::ClosestToPoint"))
+	UPROPERTY(EditAnywhere, Category = "Options", meta = (DataflowInput, EditCondition = "SelectRootMethod == EDataflowMedialSkeletonConversionSelectRootMethod::ClosestToPoint"))
 	FVector RootSelectionPoint = FVector::ZeroVector;
 
 	/** When SelectRootMethod is FarthestInDirection, the direction to search along */
-	UPROPERTY(EditAnywhere, Category = "Options", meta = (EditCondition = "SelectRootMethod == EDataflowMedialSkeletonConversionSelectRootMethod::FarthestInDirection"))
+	UPROPERTY(EditAnywhere, Category = "Options", meta = (DataflowInput, EditCondition = "SelectRootMethod == EDataflowMedialSkeletonConversionSelectRootMethod::FarthestInDirection"))
 	FVector RootSelectionDirection = FVector(0, 0, -1);
 
 	/** If non-negative, the index of the medial axis cluster (sphere) to use as the root of the Skeleton */
 	UPROPERTY(EditAnywhere, Category = "Options", meta = (DataflowInput))
 	int32 RootClusterIndex = INDEX_NONE;
 
+	/** If true, add an extra root bone at CustomAnimationRootPosition, unrelated to any medial skeleton cluster */
+	UPROPERTY(EditAnywhere, Category = "Options", meta = (DataflowInput))
+	bool bAddCustomAnimationRoot = false;
+
+	/** The world-space position for the custom root bone. Only used when bAddCustomAnimationRoot is true. */
+	UPROPERTY(EditAnywhere, Category = "Options", meta = (DataflowInput, EditCondition = "bAddCustomAnimationRoot"))
+	FVector CustomAnimationRootPosition = FVector::ZeroVector;
+
 	virtual void Evaluate(UE::Dataflow::FContext& Context, const FDataflowOutput* Out) const override;
 
 
@@ -173,17 +181,25 @@ private:
 	EDataflowMedialSkeletonConversionSelectRootMethod SelectRootMethod = EDataflowMedialSkeletonConversionSelectRootMethod::ArrayOrder;
 
 	/** When SelectRootMethod is ClosestToPoint, find the closest cluster to this point */
-	UPROPERTY(EditAnywhere, Category = "Options", meta = (EditCondition = "SelectRootMethod == EDataflowMedialSkeletonConversionSelectRootMethod::ClosestToPoint"))
+	UPROPERTY(EditAnywhere, Category = "Options", meta = (DataflowInput, EditCondition = "SelectRootMethod == EDataflowMedialSkeletonConversionSelectRootMethod::ClosestToPoint"))
 	FVector RootSelectionPoint = FVector::ZeroVector;
 
 	/** When SelectRootMethod is FarthestInDirection, the direction to search along */
-	UPROPERTY(EditAnywhere, Category = "Options", meta = (EditCondition = "SelectRootMethod == EDataflowMedialSkeletonConversionSelectRootMethod::FarthestInDirection"))
+	UPROPERTY(EditAnywhere, Category = "Options", meta = (DataflowInput, EditCondition = "SelectRootMethod == EDataflowMedialSkeletonConversionSelectRootMethod::FarthestInDirection"))
 	FVector RootSelectionDirection = FVector(0, 0, -1);
 
 	/** If non-negative, the index of the medial axis cluster (sphere) to use as the root of the Skeleton */
 	UPROPERTY(EditAnywhere, Category = "Options", meta = (DataflowInput))
 	int32 RootClusterIndex = INDEX_NONE;
 
+	/** If true, add an extra root bone at CustomAnimationRootPosition. In constrained binding, this bone will not receive skin weights from clustered vertices. */
+	UPROPERTY(EditAnywhere, Category = "Options", meta = (DataflowInput))
+	bool bAddCustomAnimationRoot = false;
+
+	/** The world-space position for the custom root bone. Only used when bAddCustomAnimationRoot is true. */
+	UPROPERTY(EditAnywhere, Category = "Options", meta = (DataflowInput, EditCondition = "bAddCustomAnimationRoot"))
+	FVector CustomAnimationRootPosition = FVector::ZeroVector;
+
 	/** Binding method to use for skin weights */
 	UPROPERTY(EditAnywhere, Category = "Binding")
 	EDataflowBindSkeletonMethod BindMethod = EDataflowBindSkeletonMethod::DirectDistance;
diff --git a/Engine/Plugins/Runtime/GeometryProcessing/Source/GeometryAlgorithms/Private/Skeletonization/MeshMedialAxisSampling.cpp b/Engine/Plugins/Runtime/GeometryProcessing/Source/GeometryAlgorithms/Private/Skeletonization/MeshMedialAxisSampling.cpp
--- a/Engine/Plugins/Runtime/GeometryProcessing/Source/GeometryAlgorithms/Private/Skeletonization/MeshMedialAxisSampling.cpp
+++ b/Engine/Plugins/Runtime/GeometryProcessing/Source/GeometryAlgorithms/Private/Skeletonization/MeshMedialAxisSampling.cpp
@@ -2288,17 +2288,25 @@ namespace UE::Geometry::MedialAxis
 		TArray<bool> Visited;
 		Visited.SetNumZeroed(NumClusters);
 
+		// World position of the virtual/custom root bone, if needed
+		bool bHasCustomRoot = Options.CustomRootPosition.IsSet();
+		FVector3d VirtualRootPosition = bHasCustomRoot ? *Options.CustomRootPosition : FVector3d::ZeroVector;
+
 		TArray<FIndex2i> MedialQueue;
 
-		if (MedialForest.GetSize(RootIdx) != NumClusters)
+		const bool bStillDisconnected = (MedialForest.GetSize(RootIdx) != NumClusters);
+
+		if (bStillDisconnected || bHasCustomRoot)
 		{
-			ensure(Options.MergeDisconnectedMethod == EMergeDisconnectedMethod::AddTopLevelRoot); // should only reach here w/ this merge method; otherwise, expect we have already merged to a single root
+			ensure(bHasCustomRoot || Options.MergeDisconnectedMethod == EMergeDisconnectedMethod::AddTopLevelRoot);
 
-			// add the virtual root at origin
+			// add the virtual root
 			OutParent.Add(INDEX_NONE);
-			OutTransform.Add(FTransform::Identity);
+			OutTransform.Add(FTransform(VirtualRootPosition));
 			UseIndexToMedial->Add(INDEX_NONE); // virtual root doesn't correspond to a medial skeleton cluster
-
+		}
+		if (bStillDisconnected)
+		{
 			TSet<int32> SeenTrees;
 			for (int32 Idx = 0; Idx < NumClusters; ++Idx)
 			{
@@ -2315,7 +2323,7 @@ namespace UE::Geometry::MedialAxis
 		// If above multi-root handling hasn't already added to the queue, add the root
 		if (MedialQueue.IsEmpty())
 		{
-			MedialQueue.Emplace(RootIdx, INDEX_NONE);
+			MedialQueue.Emplace(RootIdx, bHasCustomRoot ? 0 : INDEX_NONE);
 		}
 
 		while (!MedialQueue.IsEmpty())
@@ -2329,7 +2337,7 @@ namespace UE::Geometry::MedialAxis
 			Visited[SrcIdxOutParent.A] = true;
 			int32 OutIdx = UseIndexToMedial->Add(SrcIdxOutParent.A);
 			OutParent.Add(SrcIdxOutParent.B);
-			FVector3d ParentPos(0.);
+			FVector3d ParentPos = VirtualRootPosition;
 			if (SrcIdxOutParent.B != INDEX_NONE && (*UseIndexToMedial)[SrcIdxOutParent.B] != INDEX_NONE)
 			{
 				ParentPos = InMedialSkeleton.Spheres[(*UseIndexToMedial)[SrcIdxOutParent.B]].Center;
diff --git a/Engine/Plugins/Runtime/GeometryProcessing/Source/GeometryAlgorithms/Public/Skeletonization/MeshMedialAxisSampling.h b/Engine/Plugins/Runtime/GeometryProcessing/Source/GeometryAlgorithms/Public/Skeletonization/MeshMedialAxisSampling.h
--- a/Engine/Plugins/Runtime/GeometryProcessing/Source/GeometryAlgorithms/Public/Skeletonization/MeshMedialAxisSampling.h
+++ b/Engine/Plugins/Runtime/GeometryProcessing/Source/GeometryAlgorithms/Public/Skeletonization/MeshMedialAxisSampling.h
@@ -385,6 +385,10 @@ namespace MedialAxis
 		// When SelectRootMethod is FarthestInDirection, the direction to use
 		FVector3d RootSelectionDirection = -FVector3d::UnitZ();
 
+		// If set, always add a top-level root bone at this position (not corresponding to a medial skeleton cluster)
+		// Note: When MergeDisconnectedMethod==AddTopLevelRoot, this will also be used as that top level root bone.
+		TOptional<FVector3d> CustomRootPosition;
+
 		/**
 		 * Remap a medial skeleton to a hierarchical skeleton suitable for using in a URefSkeleton for animation
 		 * Note the indices of the hierarchical skeleton bones will be remapped from the medial skeleton, because URefSkeleton requires that parent bone indices be smaller than child bone indices
