Dataflow : enable gizmo on a several nodes
- Make nodes ( make box ... )
- Static mesh to collection transform
- Selection by primitive node

#Synced-CL 52671010

#Synced-CL 52671010

diff --git a/Engine/Plugins/Experimental/Dataflow/Source/DataflowNodes/Private/Dataflow/SamplerNodes/DataflowGradientSamplerNode.h b/Engine/Plugins/Experimental/Dataflow/Source/DataflowNodes/Private/Dataflow/SamplerNodes/DataflowGradientSamplerNode.h
--- a/Engine/Plugins/Experimental/Dataflow/Source/DataflowNodes/Private/Dataflow/SamplerNodes/DataflowGradientSamplerNode.h
+++ b/Engine/Plugins/Experimental/Dataflow/Source/DataflowNodes/Private/Dataflow/SamplerNodes/DataflowGradientSamplerNode.h
@@ -20,7 +20,7 @@ struct FDataflowLinearGradientFloatSampler : public FDataflowFloatSamplerBase
 
 public:
 	// Point in space where the start value is defined
-	UPROPERTY(EditAnywhere, Category = Gradient, meta = (DataflowInput))
+	UPROPERTY(EditAnywhere, Category = Gradient, meta = (DataflowInput, GizmoType = "Translate"))
 	FVector3f StartPoint = { 0, 0, 100 };
 
 	// Value defined at start point
@@ -28,7 +28,7 @@ public:
 	float StartValue = 1.0;
 
 	// Point in space where end value is defined
-	UPROPERTY(EditAnywhere, Category = Gradient, meta = (DataflowInput))
+	UPROPERTY(EditAnywhere, Category = Gradient, meta = (DataflowInput, GizmoType = "Translate"))
 	FVector3f EndPoint = { 0, 0, 0 };
 
 	// Value defined at end point
diff --git a/Engine/Plugins/Experimental/GeometryCollectionPlugin/Source/GeometryCollectionNodes/Public/Dataflow/GeometryCollectionMakeNodes.h b/Engine/Plugins/Experimental/GeometryCollectionPlugin/Source/GeometryCollectionNodes/Public/Dataflow/GeometryCollectionMakeNodes.h
--- a/Engine/Plugins/Experimental/GeometryCollectionPlugin/Source/GeometryCollectionNodes/Public/Dataflow/GeometryCollectionMakeNodes.h
+++ b/Engine/Plugins/Experimental/GeometryCollectionPlugin/Source/GeometryCollectionNodes/Public/Dataflow/GeometryCollectionMakeNodes.h
@@ -124,16 +124,16 @@ public:
 	UPROPERTY(EditAnywhere, Category = "Box", meta = (DisplayName = "Input Data Type"));
 	EMakeBoxDataTypeEnum DataType = EMakeBoxDataTypeEnum::Dataflow_MakeBox_DataType_MinMax;
 
-	UPROPERTY(EditAnywhere, Category = "Box", meta = (DataflowInput, EditCondition = "DataType == EMakeBoxDataTypeEnum::Dataflow_MakeBox_DataType_MinMax", EditConditionHides));
+	UPROPERTY(EditAnywhere, Category = "Box", meta = (DataflowInput, EditCondition = "DataType == EMakeBoxDataTypeEnum::Dataflow_MakeBox_DataType_MinMax", EditConditionHides, GizmoType = "Translate"))
 	FVector Min = FVector(0.0);
 
-	UPROPERTY(EditAnywhere, Category = "Box", meta = (DataflowInput, EditCondition = "DataType == EMakeBoxDataTypeEnum::Dataflow_MakeBox_DataType_MinMax", EditConditionHides));
+	UPROPERTY(EditAnywhere, Category = "Box", meta = (DataflowInput, EditCondition = "DataType == EMakeBoxDataTypeEnum::Dataflow_MakeBox_DataType_MinMax", EditConditionHides, GizmoType = "Translate"))
 	FVector Max = FVector(10.0);
 
-	UPROPERTY(EditAnywhere, Category = "Box", meta = (DataflowInput, EditCondition = "DataType == EMakeBoxDataTypeEnum::Dataflow_MakeBox_DataType_CenterSize", EditConditionHides));
+	UPROPERTY(EditAnywhere, Category = "Box", meta = (DataflowInput, EditCondition = "DataType == EMakeBoxDataTypeEnum::Dataflow_MakeBox_DataType_CenterSize"))
 	FVector Center = FVector(0.0);
 
-	UPROPERTY(EditAnywhere, Category = "Box", meta = (DataflowInput, EditCondition = "DataType == EMakeBoxDataTypeEnum::Dataflow_MakeBox_DataType_CenterSize", EditConditionHides));
+	UPROPERTY(EditAnywhere, Category = "Box", meta = (DataflowInput, EditCondition = "DataType == EMakeBoxDataTypeEnum::Dataflow_MakeBox_DataType_CenterSize"))
 	FVector Size = FVector(10.0);
 
 	UPROPERTY(meta = (DataflowOutput));
@@ -165,7 +165,7 @@ struct FMakeSphereDataflowNode : public FDataflowNode
 	DATAFLOW_NODE_DEFINE_INTERNAL(FMakeSphereDataflowNode, "MakeSphere", "Generators|Sphere", "")
 
 public:
-	UPROPERTY(EditAnywhere, Category = "Sphere", meta = (DataflowInput));
+	UPROPERTY(EditAnywhere, Category = "Sphere", meta = (DataflowInput, GizmoType = "Translate"))
 	FVector Center = FVector(0.f);
 
 	UPROPERTY(EditAnywhere, Category = "Sphere", meta = (DataflowInput));
@@ -632,7 +632,7 @@ struct FBreakTransformDataflowNode : public FDataflowNode
 
 private:
 	/** Transform to break into components */
-	UPROPERTY(EditAnywhere, Category = "Transform", meta = (DataflowInput));
+	UPROPERTY(EditAnywhere, Category = "Transform", meta = (DataflowInput, GizmoType = "Transform"));
 	FTransform Transform;
 
 	/** Translation */
@@ -796,7 +796,7 @@ struct FMakeBoxMeshDataflowNode : public FDataflowNode
 
 private:
 	/**  */
-	UPROPERTY(EditAnywhere, Category = "Box");
+	UPROPERTY(EditAnywhere, Category = "Box", meta = (GizmoType = "Translate"));
 	FVector Center = FVector(0.0);
 
 	/**  */
@@ -837,7 +837,7 @@ struct FMakePlaneDataflowNode : public FDataflowNode
 
 private:
 	/** Base point */
-	UPROPERTY(EditAnywhere, Category = "Plane", meta = (DataflowInput));
+	UPROPERTY(EditAnywhere, Category = "Plane", meta = (DataflowInput, GizmoType = "Translate"));
 	FVector BasePoint = FVector(0.0);
 
 	/** Normal vector */
@@ -983,7 +983,7 @@ struct FMakeRectangleMeshDataflowNode : public FDataflowNode
 
 private:
 	/** Rectangle will be translated so that center is at this point */
-	UPROPERTY(EditAnywhere, Category = "Rectangle", meta = (DataflowInput));
+	UPROPERTY(EditAnywhere, Category = "Rectangle", meta = (DataflowInput, GizmoType = "Translate"));
 	FVector Origin = FVector(0.0);
 
 	/** Normal vector of all vertices will be set to this value. Default is +Z axis. */
@@ -1028,7 +1028,7 @@ struct FMakeTorusMeshDataflowNode : public FDataflowNode
 
 private:
 	/** Torus will be translated so that center is at this point */
-	UPROPERTY(EditAnywhere, Category = "Torus", meta = (DataflowInput));
+	UPROPERTY(EditAnywhere, Category = "Torus", meta = (DataflowInput, GizmoType = "Translate"));
 	FVector Origin = FVector(0.0);
 
 	/** Radius of the profile */
diff --git a/Engine/Plugins/Experimental/GeometryCollectionPlugin/Source/GeometryCollectionNodes/Public/Dataflow/GeometryCollectionSelectionNodes.h b/Engine/Plugins/Experimental/GeometryCollectionPlugin/Source/GeometryCollectionNodes/Public/Dataflow/GeometryCollectionSelectionNodes.h
--- a/Engine/Plugins/Experimental/GeometryCollectionPlugin/Source/GeometryCollectionNodes/Public/Dataflow/GeometryCollectionSelectionNodes.h
+++ b/Engine/Plugins/Experimental/GeometryCollectionPlugin/Source/GeometryCollectionNodes/Public/Dataflow/GeometryCollectionSelectionNodes.h
@@ -2158,7 +2158,7 @@ private:
 	FDataflowPrimitiveTypes Primitive;
 
 	/** Transform for the primitive */
-	UPROPERTY(EditAnywhere, Category = "Selection", meta = (DataflowInput))
+	UPROPERTY(EditAnywhere, Category = "Selection", meta = (DataflowInput, GizmoType = "Transform"))
 	FTransform Transform;
 
 	/** If true then all elements must contained by the primitive (all vertices of a face, all vertices of a curve, etc.) */
diff --git a/Engine/Plugins/Experimental/GeometryCollectionPlugin/Source/GeometryCollectionNodes/Public/Dataflow/GeometryCollectionStaticMeshToCollectionNode.h b/Engine/Plugins/Experimental/GeometryCollectionPlugin/Source/GeometryCollectionNodes/Public/Dataflow/GeometryCollectionStaticMeshToCollectionNode.h
--- a/Engine/Plugins/Experimental/GeometryCollectionPlugin/Source/GeometryCollectionNodes/Public/Dataflow/GeometryCollectionStaticMeshToCollectionNode.h
+++ b/Engine/Plugins/Experimental/GeometryCollectionPlugin/Source/GeometryCollectionNodes/Public/Dataflow/GeometryCollectionStaticMeshToCollectionNode.h
@@ -30,7 +30,7 @@ private:
 	TObjectPtr<UStaticMesh> StaticMesh;
 
 	/** Transform to apply to the mesh before converting it to a collection */
-	UPROPERTY(EditAnywhere, Category = "Asset", meta = (DataflowInput));
+	UPROPERTY(EditAnywhere, Category = "Asset", meta = (DataflowInput, GizmoType = "Transform"));
 	FTransform MeshTransform = FTransform::Identity;
 
 	/** Geometry collection newly created */
